Add Physics Control Component plugin.

diff --git a/Engine/Plugins/Experimental/PhysicsControl/PhysicsControl.uplugin b/Engine/Plugins/Experimental/PhysicsControl/PhysicsControl.uplugin
new file mode 100644
--- /dev/null
+++ b/Engine/Plugins/Experimental/PhysicsControl/PhysicsControl.uplugin
@@ -0,0 +1,30 @@
+{
+	"FileVersion": 3,
+	"Version": 1,
+	"VersionName": "1.0",
+	"FriendlyName": "PhysicsControl",
+	"Description": "",
+	"Category": "Other",
+	"CreatedBy": "Epic Games",
+	"CreatedByURL": "",
+	"DocsURL": "",
+	"MarketplaceURL": "",
+	"SupportURL": "",
+	"CanContainContent": true,
+	"CanContainVerse": false,
+	"IsBetaVersion": false,
+	"IsExperimentalVersion": true,
+	"Installed": false,
+	"Modules": [
+		{
+			"Name": "PhysicsControl",
+			"Type": "Runtime",
+			"LoadingPhase": "Default"
+		},
+		{
+			"Name": "PhysicsControlVisualizer",
+			"Type": "Editor",
+			"LoadingPhase": "PostEngineInit"
+		}
+	]
+}
\ No newline at end of file
diff --git a/Engine/Plugins/Experimental/PhysicsControl/Source/PhysicsControl/PhysicsControl.Build.cs b/Engine/Plugins/Experimental/PhysicsControl/Source/PhysicsControl/PhysicsControl.Build.cs
new file mode 100644
--- /dev/null
+++ b/Engine/Plugins/Experimental/PhysicsControl/Source/PhysicsControl/PhysicsControl.Build.cs
@@ -0,0 +1,29 @@
+// Copyright Epic Games, Inc. All Rights Reserved.
+
+using UnrealBuildTool;
+
+public class PhysicsControl : ModuleRules
+{
+	public PhysicsControl(ReadOnlyTargetRules Target) : base(Target)
+	{
+		PCHUsage = ModuleRules.PCHUsageMode.UseExplicitOrSharedPCHs;
+		
+		PublicDependencyModuleNames.AddRange(
+			new string[]
+			{
+				"Core"
+			}
+		);
+			
+		
+		PrivateDependencyModuleNames.AddRange(
+			new string[]
+			{
+				"CoreUObject",
+				"PhysicsCore",
+				"Engine"
+			}
+		);
+		
+	}
+}
diff --git a/Engine/Plugins/Experimental/PhysicsControl/Source/PhysicsControl/Private/PhysicsControlComponent.cpp b/Engine/Plugins/Experimental/PhysicsControl/Source/PhysicsControl/Private/PhysicsControlComponent.cpp
new file mode 100644
--- /dev/null
+++ b/Engine/Plugins/Experimental/PhysicsControl/Source/PhysicsControl/Private/PhysicsControlComponent.cpp
@@ -0,0 +1,1096 @@
+// Copyright Epic Games, Inc. All Rights Reserved.
+
+#include "PhysicsControlComponent.h"
+#include "PhysicsControlComponentLog.h"
+#include "PhysicsControlRecord.h"
+#include "PhysicsControlComponentHelpers.h"
+#include "PhysicsControlComponentImpl.h"
+
+#include "PhysicsEngine/ConstraintInstance.h"
+#include "PhysicsEngine/BodyInstance.h"
+#include "Physics/PhysicsInterfaceCore.h"
+#include "Components/StaticMeshComponent.h"
+#include "Components/SkeletalMeshComponent.h"
+
+//======================================================================================================================
+UPhysicsControlComponent::UPhysicsControlComponent(const FObjectInitializer& ObjectInitializer)
+	: Super(ObjectInitializer)
+{
+	Implementation = MakePimpl<FPhysicsControlComponentImpl>(this);
+
+	TeleportDistanceThreshold = 300.0f;
+	TeleportRotationThreshold = 0.0f;
+
+	bShowDebugVisualization = true;
+	VisualizationSizeScale = 5.0f;
+	VelocityPredictionTime = 0.2f;
+	MaxNumControlsOrModifiersPerName = 256;
+
+	// ActorComponent setup
+	bWantsInitializeComponent = true;
+	PrimaryComponentTick.bCanEverTick = true;
+	PrimaryComponentTick.bTickEvenWhenPaused = true;
+	PrimaryComponentTick.TickGroup = TG_PrePhysics;
+}
+
+//======================================================================================================================
+void UPhysicsControlComponent::InitializeComponent()
+{
+	Super::InitializeComponent();
+	Implementation->ResetControls(false);
+}
+
+//======================================================================================================================
+void UPhysicsControlComponent::BeginDestroy()
+{
+	for (TPair<FName, FPhysicsControlRecord>& PhysicsControlRecordPair : Implementation->PhysicsControlRecords)
+	{
+		FPhysicsControlRecord& Record = PhysicsControlRecordPair.Value;
+		Record.PhysicsControlState.Reset();
+	}
+	Implementation->PhysicsControlRecords.Empty();
+	Super::BeginDestroy();
+}
+
+//======================================================================================================================
+void UPhysicsControlComponent::TickComponent(
+	float DeltaTime, enum ELevelTick TickType, FActorComponentTickFunction* ThisTickFunction)
+{
+	// Handle control removal
+	bool bRemovedOneControl = false;
+	for (auto It = Implementation->PhysicsControlRecords.CreateIterator(); It; ++It)
+	{
+		FPhysicsControlRecord& Record = It.Value();
+
+		if (Record.PhysicsControlState.bPendingDestroy)
+		{
+			Record.PhysicsControlState.Reset();
+			It.RemoveCurrent();
+			bRemovedOneControl = true;
+		}
+	}
+	Implementation->PhysicsControlRecords.Compact();
+
+	// Handle body modifier removal
+	bool bRemovedOneBodyModifier = false;
+	for (auto It = Implementation->PhysicsBodyModifiers.CreateIterator(); It; ++It)
+	{
+		FPhysicsBodyModifier& BodyModifier = It.Value();
+		if (BodyModifier.bPendingDestroy)
+		{
+			It.RemoveCurrent();
+			bRemovedOneBodyModifier = true;
+		}
+	}
+	Implementation->PhysicsBodyModifiers.Compact();
+
+	// We only want to continue the update if this is a "real" tick that corresponds to updating the
+	// world. We certainly don't want to tick during a pause, because part of the processing involves 
+	// (optionally) calculating target velocities based on target positions in previous ticks etc.
+	if (TickType != LEVELTICK_All)
+	{
+		return;
+	}
+
+	// Update the skeletal mesh caches
+	Implementation->UpdateCachedSkeletalBoneData(DeltaTime);
+
+	for (TPair<FName, FPhysicsControlRecord>& RecordPair : Implementation->PhysicsControlRecords)
+	{
+		// New constraint requested when one doesn't exist
+		FPhysicsControlRecord& Record = RecordPair.Value;
+		FConstraintInstance* ConstraintInstance = Record.PhysicsControlState.ConstraintInstance.Get();
+		if (!ConstraintInstance)
+		{
+			ConstraintInstance = Record.CreateConstraint(this);
+		}
+
+		if (ConstraintInstance)
+		{
+			// Constraint is not enabled
+			if (!Record.PhysicsControlState.bEnabled)
+			{
+				ConstraintInstance->SetAngularDriveParams(0.0f, 0.0f, 0.0f);
+				ConstraintInstance->SetLinearDriveParams(0.0f, 0.0f, 0.0f);
+			}
+			else
+			{
+				Implementation->ApplyControl(Record);
+			}
+		}
+	}
+
+	// Handle body modifiers
+	for (TPair<FName, FPhysicsBodyModifier>& BodyModifierPair : Implementation->PhysicsBodyModifiers)
+	{
+		FPhysicsBodyModifier& BodyModifier = BodyModifierPair.Value;
+		FBodyInstance* BodyInstance = GetBodyInstance(BodyModifier.MeshComponent, BodyModifier.BoneName);
+		// Note that there is a Default type, which means we don't change anything
+		if (BodyModifier.MovementType == EPhysicsMovementType::Simulated)
+		{
+			BodyInstance->SetInstanceSimulatePhysics(true);
+		}
+		else if (BodyModifier.MovementType == EPhysicsMovementType::Kinematic)
+		{
+			BodyInstance->SetInstanceSimulatePhysics(false, true);
+			Implementation->ApplyKinematicTarget(BodyModifier);
+		}
+
+		if (BodyInstance->IsInstanceSimulatingPhysics())
+		{
+			float GravityZ = BodyInstance->GetPhysicsScene()->GetOwningWorld()->GetGravityZ();
+			float AppliedGravityZ = BodyInstance->bEnableGravity ? GravityZ : 0.0f;
+			float DesiredGravityZ = GravityZ * BodyModifier.GravityMultiplier;
+			float GravityZToApply = DesiredGravityZ - AppliedGravityZ;
+			BodyInstance->AddForce(FVector(0, 0, GravityZToApply), true, true);
+		}
+	}
+
+	// Go through and de-activate any records if they're set to auto disable. 
+	for (TPair<FName, FPhysicsControlRecord>& RecordPair : Implementation->PhysicsControlRecords)
+	{
+		FPhysicsControlRecord& Record = RecordPair.Value;
+		if (Record.PhysicsControl.ControlSettings.bAutoDisable)
+		{
+			Record.PhysicsControlState.bEnabled = false;
+		}
+	}
+}
+
+//======================================================================================================================
+TMap<FName, FPhysicsControlLimbBones> UPhysicsControlComponent::GetLimbBonesFromSkeletalMesh(
+	USkeletalMeshComponent*                     SkeletalMeshComponent,
+	const TArray<FPhysicsControlLimbSetupData>& LimbSetupDatas) const
+{
+	TMap<FName, FPhysicsControlLimbBones> Result;
+
+	UPhysicsAsset* PhysicsAsset = SkeletalMeshComponent ? SkeletalMeshComponent->GetPhysicsAsset() : nullptr;
+	if (!PhysicsAsset)
+	{
+		UE_LOG(LogPhysicsControlComponent, Warning, TEXT("No physics asset in skeletal mesh"));
+		return Result;
+	}
+
+	TSet<FName> AllBones;
+
+	// Now walk through each limb, picking up bones, ignoring any that we have already encountered.
+	// This requires the setup data to have been ordered properly.
+	for (const FPhysicsControlLimbSetupData& LimbSetupData : LimbSetupDatas)
+	{
+		FPhysicsControlLimbBones& LimbBones = Result.Add(LimbSetupData.LimbName);
+		LimbBones.SkeletalMeshComponent = SkeletalMeshComponent;
+
+		if (LimbSetupData.bIncludeParentBone)
+		{
+			FName ParentBoneName;
+			ParentBoneName = GetPhysicalParentBone(SkeletalMeshComponent, LimbSetupData.StartBone);
+			if (!ParentBoneName.IsNone())
+			{
+				LimbBones.BoneNames.Add(ParentBoneName);
+				AllBones.Add(ParentBoneName);
+			}
+		}
+
+		SkeletalMeshComponent->ForEachBodyBelow(
+			LimbSetupData.StartBone, true, /*bSkipCustomType=*/false,
+			[PhysicsAsset, &AllBones, &LimbBones](const FBodyInstance* BI)
+			{
+				if (USkeletalBodySetup* BodySetup = Cast<USkeletalBodySetup>(BI->BodySetup.Get()))
+				{
+					const FName BoneName = PhysicsAsset->SkeletalBodySetups[BI->InstanceBodyIndex]->BoneName;
+					if (!AllBones.Find(BoneName))
+					{
+						LimbBones.BoneNames.Add(BoneName);
+						AllBones.Add(BoneName);
+					}
+				}
+			});
+	}
+	return Result;
+}
+
+//======================================================================================================================
+// Note - some params passed by value to allow them to be unconnected and enable inline BP editing
+FName UPhysicsControlComponent::MakeControl(
+	UMeshComponent*         ParentMeshComponent,
+	const FName             ParentBoneName,
+	UMeshComponent*         ChildMeshComponent,
+	const FName             ChildBoneName,
+	FPhysicsControlData     ControlData, 
+	FPhysicsControlTarget   ControlTarget, 
+	FPhysicsControlSettings ControlSettings,
+	bool                    bEnabled)
+{
+	FName Name = Implementation->GetUniqueControlName(ParentBoneName, ChildBoneName);
+	if (MakeNamedControl(
+		Name, ParentMeshComponent, ParentBoneName, ChildMeshComponent, ChildBoneName, 
+		ControlData, ControlTarget, ControlSettings, bEnabled))
+	{
+		return Name;
+	}
+	return FName();
+}
+
+//======================================================================================================================
+// Note - some params passed by value to allow them to be unconnected and enable inline BP editing
+bool UPhysicsControlComponent::MakeNamedControl(
+	const FName             Name, 
+	UMeshComponent*         ParentMeshComponent,
+	const FName             ParentBoneName,
+	UMeshComponent*         ChildMeshComponent,
+	const FName             ChildBoneName,
+	FPhysicsControlData     ControlData, 
+	FPhysicsControlTarget   ControlTarget, 
+	FPhysicsControlSettings ControlSettings,
+	bool                    bEnabled)
+{
+	FPhysicsControlRecord* PhysicsControlRecord = Implementation->FindControlRecord(Name);
+	if (PhysicsControlRecord)
+	{
+		return false;
+	}
+
+	if (!ChildMeshComponent)
+	{
+		UE_LOG(LogPhysicsControlComponent, Warning,
+			TEXT("Unable to make a Control as the child mesh component has not been set"));
+		return false;
+	}
+
+	if (ParentMeshComponent && ParentMeshComponent->IsA<USkeletalMeshComponent>())
+	{
+		Implementation->AddSkeletalMeshReference(Cast<USkeletalMeshComponent>(ParentMeshComponent));
+	}
+	if (ChildMeshComponent && ChildMeshComponent->IsA<USkeletalMeshComponent>())
+	{
+		Implementation->AddSkeletalMeshReference(Cast<USkeletalMeshComponent>(ChildMeshComponent));
+	}
+
+	FPhysicsControlRecord& NewRecord = Implementation->PhysicsControlRecords.Add(
+		Name, FPhysicsControl(
+			ParentMeshComponent, ParentBoneName, ChildMeshComponent, ChildBoneName,
+			ControlData, ControlTarget, ControlSettings));
+	NewRecord.PhysicsControlState.bEnabled = bEnabled;
+	NewRecord.ResetControlPoint();
+
+	return true;
+}
+
+//======================================================================================================================
+// Note - some params passed by value to allow them to be unconnected and enable inline BP editing
+TArray<FName> UPhysicsControlComponent::MakeControlsFromSkeletalMeshBelow(
+	USkeletalMeshComponent* SkeletalMeshComponent,
+	FName                   BoneName,
+	bool                    bIncludeSelf,
+	EPhysicsControlType     ControlType,
+	FPhysicsControlData     ControlData,
+	FPhysicsControlSettings ControlSettings,
+	bool                    bEnabled)
+{
+	TArray<FName> Result;
+	UPhysicsAsset* PhysicsAsset = SkeletalMeshComponent ? SkeletalMeshComponent->GetPhysicsAsset() : nullptr;
+	if (!PhysicsAsset)
+	{
+		UE_LOG(LogPhysicsControlComponent, Warning, TEXT("No physics asset in skeletal mesh"));
+		return Result;
+	}
+
+	int32 NumBones = SkeletalMeshComponent->GetNumBones();
+	UMeshComponent* ParentMeshComponent = 
+		(ControlType == EPhysicsControlType::ParentSpace) ? SkeletalMeshComponent : nullptr;
+
+	SkeletalMeshComponent->ForEachBodyBelow(
+		BoneName, bIncludeSelf, /*bSkipCustomType=*/false, 
+		[
+			this, PhysicsAsset, ParentMeshComponent, SkeletalMeshComponent, 
+			&ControlData, &ControlSettings, &Result, bEnabled
+		](const FBodyInstance* BI)
+		{
+			if (USkeletalBodySetup* BodySetup = Cast<USkeletalBodySetup>(BI->BodySetup.Get()))
+			{
+				const FName ChildBoneName = PhysicsAsset->SkeletalBodySetups[BI->InstanceBodyIndex]->BoneName;
+
+				FName ParentBoneName;
+				if (ParentMeshComponent)
+				{
+					ParentBoneName = GetPhysicalParentBone(SkeletalMeshComponent, ChildBoneName);
+					if (ParentBoneName.IsNone())
+					{
+						return;
+					}
+				}
+				FName ControlName = MakeControl(
+					ParentMeshComponent, ParentBoneName, SkeletalMeshComponent, ChildBoneName,
+					ControlData, FPhysicsControlTarget(), ControlSettings, bEnabled);
+				if (!ControlName.IsNone())
+				{
+					Result.Add(ControlName);
+				}
+				else
+				{
+					UE_LOG(LogPhysicsControlComponent, Warning, 
+						TEXT("Failed to make control for %s"), *ChildBoneName.ToString());
+				}
+			}
+		});
+
+	return Result;
+}
+
+//======================================================================================================================
+// Note - some params passed by value to allow them to be unconnected and enable inline BP editing
+TArray<FName> UPhysicsControlComponent::MakeControlsFromSkeletalMesh(
+	USkeletalMeshComponent* SkeletalMeshComponent,
+	const TArray<FName>&    BoneNames,
+	EPhysicsControlType     ControlType,
+	FPhysicsControlData     ControlData,
+	FPhysicsControlSettings ControlSettings,
+	bool                    bEnabled)
+{
+	TArray<FName> Result;
+	UPhysicsAsset* PhysicsAsset = SkeletalMeshComponent ? SkeletalMeshComponent->GetPhysicsAsset() : nullptr;
+	if (!PhysicsAsset)
+	{
+		UE_LOG(LogPhysicsControlComponent, Warning, TEXT("No physics asset in skeletal mesh"));
+		return Result;
+	}
+
+	UMeshComponent* ParentMeshComponent =
+		(ControlType == EPhysicsControlType::ParentSpace) ? SkeletalMeshComponent : nullptr;
+
+	for (FName ChildBoneName : BoneNames)
+	{
+		FBodyInstance* ChildBone = SkeletalMeshComponent->GetBodyInstance(ChildBoneName);
+
+		FName ParentBoneName;
+		if (ParentMeshComponent)
+		{
+			ParentBoneName = GetPhysicalParentBone(SkeletalMeshComponent, ChildBoneName);
+			if (ParentBoneName.IsNone())
+			{
+				continue;
+			}
+		}
+		FName ControlName = MakeControl(
+			ParentMeshComponent, ParentBoneName, SkeletalMeshComponent, ChildBoneName,
+			ControlData, FPhysicsControlTarget(), ControlSettings, bEnabled);
+		if (!ControlName.IsNone())
+		{
+			Result.Add(ControlName);
+		}
+		else
+		{
+			UE_LOG(LogPhysicsControlComponent, Warning, TEXT("Failed to make control for %s"), *ChildBoneName.ToString());
+		}
+	}
+
+	return Result;
+}
+
+//======================================================================================================================
+// Note - some params passed by value to allow them to be unconnected and enable inline BP editing
+TMap<FName, FPhysicsControlNameArray> UPhysicsControlComponent::MakeControlsFromLimbBones(
+	FPhysicsControlNameArray&                    AllControls,
+	const TMap<FName, FPhysicsControlLimbBones>& LimbBones,
+	EPhysicsControlType                          ControlType,
+	FPhysicsControlData                          ControlData,
+	FPhysicsControlSettings                      ControlSettings,
+	bool                                         bEnabled)
+{
+	TMap<FName, FPhysicsControlNameArray> Result;
+	Result.Reserve(LimbBones.Num());
+
+	for (const TPair< FName, FPhysicsControlLimbBones>& LimbBoneEntry : LimbBones)
+	{
+		FName LimbName = LimbBoneEntry.Key;
+		const FPhysicsControlLimbBones& BonesInLimb = LimbBoneEntry.Value;
+
+		if (!BonesInLimb.SkeletalMeshComponent)
+		{
+			UE_LOG(LogPhysicsControlComponent, Warning, TEXT("No Skeletal mesh in limb %s"), *LimbName.ToString());
+			continue;
+		}
+
+		USkeletalMeshComponent* ParentMeshComponent =
+			(ControlType == EPhysicsControlType::ParentSpace) ? BonesInLimb.SkeletalMeshComponent : nullptr;
+
+		const int32 NumBonesInLimb = BonesInLimb.BoneNames.Num();
+
+		FPhysicsControlNameArray& LimbResult = Result.Add(LimbName);
+		LimbResult.Names.Reserve(NumBonesInLimb);
+		AllControls.Names.Reserve(AllControls.Names.Num() + NumBonesInLimb);
+
+		for (FName ChildBoneName : BonesInLimb.BoneNames)
+		{
+			FBodyInstance* ChildBone = BonesInLimb.SkeletalMeshComponent->GetBodyInstance(ChildBoneName);
+
+			FName ParentBoneName;
+			if (ParentMeshComponent)
+			{
+				ParentBoneName = GetPhysicalParentBone(ParentMeshComponent, ChildBoneName);
+				if (ParentBoneName.IsNone())
+				{
+					continue;
+				}
+			}
+			FName ControlName = MakeControl(
+				ParentMeshComponent, ParentBoneName, BonesInLimb.SkeletalMeshComponent, ChildBoneName,
+				ControlData, FPhysicsControlTarget(), ControlSettings, bEnabled);
+			if (!ControlName.IsNone())
+			{
+				LimbResult.Names.Add(ControlName);
+				AllControls.Names.Add(ControlName);
+			}
+			else
+			{
+				UE_LOG(LogPhysicsControlComponent, Warning, TEXT("Failed to make control for %s"), *ChildBoneName.ToString());
+			}
+		}
+	}
+	return Result;
+}
+
+
+//======================================================================================================================
+bool UPhysicsControlComponent::DestroyControl(const FName Name)
+{
+	FPhysicsControlRecord* PhysicsControlRecord = Implementation->FindControlRecord(Name);
+	if (PhysicsControlRecord)
+	{
+		if (PhysicsControlRecord->PhysicsControl.ParentMeshComponent && 
+			PhysicsControlRecord->PhysicsControl.ParentMeshComponent->IsA<USkeletalMeshComponent>())
+		{
+			Implementation->RemoveSkeletalMeshReference(Cast<USkeletalMeshComponent>(PhysicsControlRecord->PhysicsControl.ParentMeshComponent));
+		}
+		if (PhysicsControlRecord->PhysicsControl.ChildMeshComponent &&
+			PhysicsControlRecord->PhysicsControl.ChildMeshComponent->IsA<USkeletalMeshComponent>())
+		{
+			Implementation->RemoveSkeletalMeshReference(Cast<USkeletalMeshComponent>(PhysicsControlRecord->PhysicsControl.ChildMeshComponent));
+		}
+
+		PhysicsControlRecord->PhysicsControlState.bPendingDestroy = true;
+		PhysicsControlRecord->PhysicsControlState.bEnabled = false;
+		return true;
+	}
+	return false;
+}
+
+//======================================================================================================================
+bool UPhysicsControlComponent::SetControlEnabled(const FName Name, bool bEnable)
+{
+	FPhysicsControlRecord* Record = Implementation->FindControlRecord(Name);
+	if (Record)
+	{
+		Record->PhysicsControlState.bEnabled = bEnable;
+		return true;
+	}
+	return false;
+}
+
+//======================================================================================================================
+// Note - some params passed by value to allow them to be unconnected and enable inline BP editing
+bool UPhysicsControlComponent::SetControlData(
+	const FName             Name, 
+	FPhysicsControlData     ControlData, 
+	bool                    bEnableControl)
+{
+	FPhysicsControlRecord* Record = Implementation->FindControlRecord(Name);
+	if (Record)
+	{
+		Record->PhysicsControl.ControlData = ControlData;
+		Record->PhysicsControlState.bEnabled = bEnableControl;
+		return true;
+	}
+	return false;
+}
+
+//======================================================================================================================
+// Note - some params passed by value to allow them to be unconnected and enable inline BP editing
+bool UPhysicsControlComponent::SetControlMultipliers(
+	const FName                Name, 
+	FPhysicsControlMultipliers ControlMultipliers, 
+	bool                       bEnableControl)
+{
+	FPhysicsControlRecord* Record = Implementation->FindControlRecord(Name);
+	if (Record)
+	{
+		Record->PhysicsControl.ControlMultipliers = ControlMultipliers;
+		Record->PhysicsControlState.bEnabled = bEnableControl;
+		return true;
+	}
+	return false;
+}
+
+//======================================================================================================================
+bool UPhysicsControlComponent::SetControlLinearData(
+	const FName Name, float Strength, float DampingRatio, float ExtraDamping, float MaxForce, bool bEnableControl)
+{
+	FPhysicsControlRecord* Record = Implementation->FindControlRecord(Name);
+	if (Record)
+	{
+		Record->PhysicsControl.ControlData.LinearStrength = Strength;
+		Record->PhysicsControl.ControlData.LinearDampingRatio = DampingRatio;
+		Record->PhysicsControl.ControlData.LinearExtraDamping = ExtraDamping;
+		Record->PhysicsControl.ControlData.MaxForce = MaxForce;
+		Record->PhysicsControlState.bEnabled = bEnableControl;
+		return true;
+	}
+	return false;
+}
+
+//======================================================================================================================
+bool UPhysicsControlComponent::SetControlAngularData(
+	const FName Name, float Strength, float DampingRatio, float ExtraDamping, float MaxTorque, bool bEnableControl)
+{
+	FPhysicsControlRecord* Record = Implementation->FindControlRecord(Name);
+	if (Record)
+	{
+		Record->PhysicsControl.ControlData.AngularStrength = Strength;
+		Record->PhysicsControl.ControlData.AngularDampingRatio = DampingRatio;
+		Record->PhysicsControl.ControlData.AngularExtraDamping = ExtraDamping;
+		Record->PhysicsControl.ControlData.MaxTorque = MaxTorque;
+		Record->PhysicsControlState.bEnabled = bEnableControl;
+		return true;
+	}
+	return false;
+}
+
+//======================================================================================================================
+bool UPhysicsControlComponent::SetControlPoint(const FName Name, const FVector Position)
+{
+	FPhysicsControlRecord* Record = Implementation->FindControlRecord(Name);
+	if (Record)
+	{
+		Record->PhysicsControl.ControlSettings.ControlPoint = Position;
+		Record->UpdateConstraintControlPoint();
+		return true;
+	}
+	return false;
+}
+
+//======================================================================================================================
+bool UPhysicsControlComponent::ResetControlPoint(const FName Name)
+{
+	FPhysicsControlRecord* Record = Implementation->FindControlRecord(Name);
+	if (Record)
+	{
+		Record->ResetControlPoint();
+		return true;
+	}
+	return false;
+}
+
+//======================================================================================================================
+bool UPhysicsControlComponent::SetControlTarget(
+	const FName           Name, 
+	FPhysicsControlTarget ControlTarget, 
+	bool                  bEnableControl)
+{
+	FPhysicsControlRecord* Record = Implementation->FindControlRecord(Name);
+	if (Record)
+	{
+		Record->PhysicsControl.ControlTarget = ControlTarget;
+		Record->PhysicsControlState.bEnabled = bEnableControl;
+		return true;
+	}
+	return false;
+}
+
+//======================================================================================================================
+bool UPhysicsControlComponent::SetControlTargetTransform(
+	const FName Name, const FTransform Transform, float VelocityDeltaTime, 
+	bool bEnableControl, bool bApplyControlPointToTarget)
+{
+	FPhysicsControlRecord* Record = Implementation->FindControlRecord(Name);
+	if (Record)
+	{
+		SetControlTargetPosition(
+			Name, Transform.GetTranslation(), VelocityDeltaTime, bEnableControl, bApplyControlPointToTarget);
+		SetControlTargetOrientation(
+			Name, Transform.GetRotation().Rotator(), VelocityDeltaTime, bEnableControl, bApplyControlPointToTarget);
+	}
+	return false;
+
+}
+
+//======================================================================================================================
+bool UPhysicsControlComponent::SetControlTargetPosition(
+	const FName Name, const FVector Position, float VelocityDeltaTime, 
+	bool bEnableControl, bool bApplyControlPointToTarget)
+{
+	FPhysicsControlRecord* Record = Implementation->FindControlRecord(Name);
+	if (Record)
+	{
+		if (VelocityDeltaTime)
+		{
+			Record->PhysicsControl.ControlTarget.TargetVelocity =
+				(Position - Record->PhysicsControl.ControlTarget.TargetPosition) / VelocityDeltaTime;
+		}
+		else
+		{
+			Record->PhysicsControl.ControlTarget.TargetVelocity = FVector::ZeroVector;
+		}
+		Record->PhysicsControl.ControlTarget.TargetPosition = Position;
+		Record->PhysicsControl.ControlTarget.bApplyControlPointToTarget = bApplyControlPointToTarget;
+		Record->PhysicsControlState.bEnabled = bEnableControl;
+		return true;
+	}
+	return false;
+}
+
+//======================================================================================================================
+bool UPhysicsControlComponent::SetControlTargetOrientation(
+	const FName Name, const FRotator Orientation, float AngularVelocityDeltaTime, 
+	bool bEnableControl, bool bApplyControlPointToTarget)
+{
+	FPhysicsControlRecord* Record = Implementation->FindControlRecord(Name);
+	if (Record)
+	{
+		if (AngularVelocityDeltaTime)
+		{
+			FQuat OldQ = Record->PhysicsControl.ControlTarget.TargetOrientation.Quaternion();
+			FQuat OrientationQ = Orientation.Quaternion();
+			OldQ.EnforceShortestArcWith(OrientationQ);
+			// Note that quats multiply in the opposite order to TMs
+			FQuat DeltaQ = OrientationQ * OldQ.Inverse();
+			Record->PhysicsControl.ControlTarget.TargetAngularVelocity =
+				DeltaQ.ToRotationVector() / (UE_TWO_PI * AngularVelocityDeltaTime);
+		}
+		else
+		{
+			Record->PhysicsControl.ControlTarget.TargetAngularVelocity = FVector::ZeroVector;
+		}
+		Record->PhysicsControl.ControlTarget.TargetOrientation = Orientation;
+		Record->PhysicsControl.ControlTarget.bApplyControlPointToTarget = bApplyControlPointToTarget;
+		Record->PhysicsControlState.bEnabled = bEnableControl;
+		return true;
+	}
+	return false;
+}
+
+//======================================================================================================================
+bool UPhysicsControlComponent::SetControlTargetPoses(
+	const FName Name,
+	const FVector ParentPosition, const FRotator ParentOrientation,
+	const FVector ChildPosition, const FRotator ChildOrientation,
+	float VelocityDeltaTime, bool bEnableControl)
+{
+	FPhysicsControlRecord* Record = Implementation->FindControlRecord(Name);
+	if (Record)
+	{
+		FTransform ParentTM(ParentOrientation, ParentPosition, FVector::One());
+		FTransform ChildTM(ChildOrientation, ChildPosition, FVector::One());
+
+		FTransform OffsetTM = ChildTM * ParentTM.Inverse();
+		FVector Position = OffsetTM.GetTranslation();
+		FQuat OrientationQ = OffsetTM.GetRotation();
+
+		if (VelocityDeltaTime)
+		{
+			FQuat OldQ = Record->PhysicsControl.ControlTarget.TargetOrientation.Quaternion();
+			OldQ.EnforceShortestArcWith(OrientationQ);
+			// Note that quats multiply in the opposite order to TMs
+			FQuat DeltaQ = OrientationQ * OldQ.Inverse();
+			Record->PhysicsControl.ControlTarget.TargetAngularVelocity =
+				DeltaQ.ToRotationVector() / (UE_TWO_PI * VelocityDeltaTime);
+
+			Record->PhysicsControl.ControlTarget.TargetVelocity =
+				(Position - Record->PhysicsControl.ControlTarget.TargetPosition) / VelocityDeltaTime;
+		}
+		else
+		{
+			Record->PhysicsControl.ControlTarget.TargetAngularVelocity = FVector::ZeroVector;
+			Record->PhysicsControl.ControlTarget.TargetVelocity = FVector::ZeroVector;
+		}
+		Record->PhysicsControl.ControlTarget.TargetOrientation = OrientationQ.Rotator();
+		Record->PhysicsControl.ControlTarget.TargetPosition = Position;
+		Record->PhysicsControl.ControlTarget.bApplyControlPointToTarget = true;
+		Record->PhysicsControlState.bEnabled = bEnableControl;
+		return true;
+	}
+	return false;
+}
+
+//======================================================================================================================
+bool UPhysicsControlComponent::SetControlUseSkeletalAnimation(
+	const FName Name,
+	bool        bUseSkeletalAnimation,
+	float       SkeletalAnimationVelocityMultiplier)
+{
+	FPhysicsControl* PhysicsControl = Implementation->FindControl(Name);
+	if (PhysicsControl)
+	{
+		PhysicsControl->ControlSettings.bUseSkeletalAnimation = bUseSkeletalAnimation;
+		PhysicsControl->ControlSettings.SkeletalAnimationVelocityMultiplier = SkeletalAnimationVelocityMultiplier;
+		return true;
+	}
+	return false;
+}
+
+//======================================================================================================================
+bool UPhysicsControlComponent::SetControlAutoDisable(const FName Name, bool bAutoDisable)
+{
+	FPhysicsControl* PhysicsControl = Implementation->FindControl(Name);
+	if (PhysicsControl)
+	{
+		PhysicsControl->ControlSettings.bAutoDisable = bAutoDisable;
+		return true;
+	}
+	return false;
+}
+
+//======================================================================================================================
+bool UPhysicsControlComponent::GetControl(const FName Name, FPhysicsControl& Control) const
+{
+	FPhysicsControl* PhysicsControl = Implementation->FindControl(Name);
+	if (PhysicsControl)
+	{
+		Control = *PhysicsControl;
+		return true;
+	}
+	return false;
+}
+
+//======================================================================================================================
+bool UPhysicsControlComponent::GetControlData(const FName Name, FPhysicsControlData& ControlData) const
+{
+	FPhysicsControl* PhysicsControl = Implementation->FindControl(Name);
+	if (PhysicsControl)
+	{
+		ControlData = PhysicsControl->ControlData;
+		return true;
+	}
+	return false;
+}
+
+//======================================================================================================================
+bool UPhysicsControlComponent::GetControlMultipliers(const FName Name, FPhysicsControlMultipliers& ControlMultipliers) const
+{
+	FPhysicsControl* PhysicsControl = Implementation->FindControl(Name);
+	if (PhysicsControl)
+	{
+		ControlMultipliers = PhysicsControl->ControlMultipliers;
+		return true;
+	}
+	return false;
+}
+
+//======================================================================================================================
+bool UPhysicsControlComponent::GetControlTarget(const FName Name, FPhysicsControlTarget& ControlTarget) const
+{
+	FPhysicsControl* PhysicsControl = Implementation->FindControl(Name);
+	if (PhysicsControl)
+	{
+		ControlTarget = PhysicsControl->ControlTarget;
+		return true;
+	}
+	return false;
+}
+
+//======================================================================================================================
+bool UPhysicsControlComponent::GetControlAutoDisable(const FName Name) const
+{
+	FPhysicsControl* PhysicsControl = Implementation->FindControl(Name);
+	if (PhysicsControl)
+	{
+		return PhysicsControl->ControlSettings.bAutoDisable;
+	}
+	return false;
+}
+
+//======================================================================================================================
+bool UPhysicsControlComponent::GetControlEnabled(const FName Name) const
+{
+	FPhysicsControlRecord* Record = Implementation->FindControlRecord(Name);
+	if (Record)
+	{
+		return Record->PhysicsControlState.bEnabled;
+	}
+	return false;
+}
+
+//======================================================================================================================
+FName UPhysicsControlComponent::MakeBodyModifier(
+	UMeshComponent*       MeshComponent,
+	const FName           BoneName,
+	EPhysicsMovementType  MovementType, 
+	float                 GravityMultiplier)
+{
+	FName Name = Implementation->GetUniqueBodyModifierName(BoneName);
+	if (MakeNamedBodyModifier(Name, MeshComponent, BoneName, MovementType, GravityMultiplier))
+	{
+		return Name;
+	}
+	return FName();
+}
+
+//======================================================================================================================
+bool UPhysicsControlComponent::MakeNamedBodyModifier(
+	const FName           Name,
+	UMeshComponent*       MeshComponent,
+	const FName           BoneName,
+	EPhysicsMovementType  MovementType,
+	float                 GravityMultiplier)
+{
+	FPhysicsBodyModifier* BodyModifier = Implementation->FindBodyModifier(Name);
+	if (BodyModifier)
+	{
+		return false;
+	}
+
+	if (!MeshComponent)
+	{
+		UE_LOG(LogPhysicsControlComponent, Warning,
+			TEXT("Unable to make a PhysicsBodyModifier as the mesh component has not been set"));
+		return false;
+	}
+
+	USkeletalMeshComponent* SkeletalMeshComponent = Cast<USkeletalMeshComponent>(MeshComponent);
+	if (SkeletalMeshComponent)
+	{
+		Implementation->AddSkeletalMeshReference(SkeletalMeshComponent);
+		SkeletalMeshComponent->bUpdateMeshWhenKinematic = true;
+	}
+
+	Implementation->PhysicsBodyModifiers.Add(
+		Name, FPhysicsBodyModifier(MeshComponent, BoneName, MovementType, GravityMultiplier));
+
+	return true;
+}
+
+//======================================================================================================================
+TArray<FName> UPhysicsControlComponent::MakeBodyModifiersFromSkeletalMeshBelow(
+	USkeletalMeshComponent* SkeletalMeshComponent,
+	FName                   BoneName,
+	bool                    bIncludeSelf,
+	EPhysicsMovementType    MovementType,
+	float                   GravityMultiplier)
+{
+	TArray<FName> Result;
+	UPhysicsAsset* PhysicsAsset = SkeletalMeshComponent ? SkeletalMeshComponent->GetPhysicsAsset() : nullptr;
+	if (!PhysicsAsset)
+	{
+		return Result;
+	}
+
+	SkeletalMeshComponent->ForEachBodyBelow(
+		BoneName, bIncludeSelf, /*bSkipCustomType=*/false,
+		[this, PhysicsAsset, SkeletalMeshComponent, MovementType, GravityMultiplier, &Result](const FBodyInstance* BI)
+		{
+			if (USkeletalBodySetup* BodySetup = Cast<USkeletalBodySetup>(BI->BodySetup.Get()))
+			{
+				const FName BoneName = PhysicsAsset->SkeletalBodySetups[BI->InstanceBodyIndex]->BoneName;
+				FName BodyModifierName = MakeBodyModifier(
+					SkeletalMeshComponent, BoneName, MovementType, GravityMultiplier);
+				Result.Add(BodyModifierName);
+			}
+		});
+
+	return Result;
+}
+
+//======================================================================================================================
+TMap<FName, FPhysicsControlNameArray> UPhysicsControlComponent::MakeBodyModifiersFromLimbBones(
+	FPhysicsControlNameArray&                    AllBodyModifiers,
+	const TMap<FName, FPhysicsControlLimbBones>& LimbBones,
+	EPhysicsMovementType                         MovementType,
+	float                                        GravityMultiplier)
+{
+	TMap<FName, FPhysicsControlNameArray> Result;
+	Result.Reserve(LimbBones.Num());
+
+	for (const TPair<FName, FPhysicsControlLimbBones>& LimbBoneEntry : LimbBones)
+	{
+		FName LimbName = LimbBoneEntry.Key;
+		const FPhysicsControlLimbBones& BonesInLimb = LimbBoneEntry.Value;
+
+		if (!BonesInLimb.SkeletalMeshComponent)
+		{
+			UE_LOG(LogPhysicsControlComponent, Warning, TEXT("No Skeletal mesh in limb %s"), *LimbName.ToString());
+			continue;
+		}
+
+		const int32 NumBonesInLimb = BonesInLimb.BoneNames.Num();
+
+		FPhysicsControlNameArray& LimbResult = Result.Add(LimbName);
+		LimbResult.Names.Reserve(NumBonesInLimb);
+		AllBodyModifiers.Names.Reserve(AllBodyModifiers.Names.Num() + NumBonesInLimb);
+
+		for (FName BoneName : BonesInLimb.BoneNames)
+		{
+			FName BodyModifierName = MakeBodyModifier(
+				BonesInLimb.SkeletalMeshComponent, BoneName, MovementType, GravityMultiplier);
+			if (!BodyModifierName.IsNone())
+			{
+				LimbResult.Names.Add(BodyModifierName);
+				AllBodyModifiers.Names.Add(BodyModifierName);
+			}
+			else
+			{
+				UE_LOG(LogPhysicsControlComponent, Warning, TEXT("Failed to make body modifier for %s"), *BoneName.ToString());
+			}
+		}
+	}
+	return Result;
+}
+
+//======================================================================================================================
+bool UPhysicsControlComponent::DestroyBodyModifier(const FName Name)
+{
+	FPhysicsBodyModifier* PhysicsBodyModifier = Implementation->FindBodyModifier(Name);
+	if (PhysicsBodyModifier)
+	{
+		if (PhysicsBodyModifier->MeshComponent &&
+			PhysicsBodyModifier->MeshComponent->IsA<USkeletalMeshComponent>())
+		{
+			Implementation->RemoveSkeletalMeshReference(Cast<USkeletalMeshComponent>(PhysicsBodyModifier->MeshComponent));
+		}
+
+		PhysicsBodyModifier->bPendingDestroy = true;
+		return true;
+	}
+	return false;
+}
+
+//======================================================================================================================
+bool UPhysicsControlComponent::SetBodyModifier(
+	const FName Name, EPhysicsMovementType MovementType, float GravityMultiplier)
+{
+	FPhysicsBodyModifier* PhysicsBodyModifier = Implementation->FindBodyModifier(Name);
+	if (PhysicsBodyModifier)
+	{
+		PhysicsBodyModifier->MovementType = MovementType;
+		PhysicsBodyModifier->GravityMultiplier = GravityMultiplier;
+		return true;
+	}
+	return false;
+}
+
+#if WITH_EDITOR
+//======================================================================================================================
+void UPhysicsControlComponent::DebugDraw(FPrimitiveDrawInterface* PDI) const
+{
+	if (bShowDebugVisualization && VisualizationSizeScale > 0)
+	{
+		for (const TPair<FName, FPhysicsControlRecord>& PhysicsControlRecordPair : 
+			Implementation->PhysicsControlRecords)
+		{
+			const FName& Name = PhysicsControlRecordPair.Key;
+			const FPhysicsControlRecord& Record = PhysicsControlRecordPair.Value;
+			DebugDrawControl(PDI, Record, Name);
+		} 
+	}
+}
+
+//======================================================================================================================
+void UPhysicsControlComponent::DebugDrawControl(
+	FPrimitiveDrawInterface* PDI, const FPhysicsControlRecord& Record, FName ControlName) const
+{
+	const float GizmoWidthScale = 0.02f * VisualizationSizeScale;
+	const FColor CurrentToTargetColor(255, 0, 0);
+	const FColor TargetColor(0, 255, 0);
+	const FColor CurrentColor(0, 0, 255);
+
+	const FConstraintInstance* ConstraintInstance = Record.PhysicsControlState.ConstraintInstance.Get();
+
+	bool bHaveLinear = Record.PhysicsControl.ControlData.LinearStrength > 0;
+	bool bHaveAngular = Record.PhysicsControl.ControlData.AngularStrength > 0;
+
+	if (Record.PhysicsControlState.bEnabled && ConstraintInstance)
+	{
+		FBodyInstance* ChildBodyInstance = GetBodyInstance(
+			Record.PhysicsControl.ChildMeshComponent, Record.PhysicsControl.ChildBoneName);
+		if (!ChildBodyInstance)
+		{
+			return;
+		}
+		FTransform ChildBodyTM = ChildBodyInstance->GetUnrealWorldTransform();
+
+		FBodyInstance* ParentBodyInstance = GetBodyInstance(
+			Record.PhysicsControl.ParentMeshComponent, Record.PhysicsControl.ParentBoneName);
+		FTransform ParentBodyTM = ParentBodyInstance ? ParentBodyInstance->GetUnrealWorldTransform() : FTransform();
+
+		FTransform TargetTM;
+		FVector TargetVelocity;
+		FVector TargetAngularVelocity;
+		Implementation->CalculateControlTargetData(TargetTM, TargetVelocity, TargetAngularVelocity, Record, true);
+
+		// WorldChildFrameTM is the world-space transform of the child (driven) constraint frame
+		FTransform WorldChildFrameTM = ConstraintInstance->GetRefFrame(EConstraintFrame::Frame1) * ChildBodyTM;
+
+		// WorldParentFrameTM is the world-space transform of the parent constraint frame
+		FTransform WorldParentFrameTM = ConstraintInstance->GetRefFrame(EConstraintFrame::Frame2) * ParentBodyTM;
+
+		FTransform WorldTargetTM = TargetTM * WorldParentFrameTM;
+		FTransform WorldCurrentTM = WorldChildFrameTM;
+
+		if (!bHaveLinear)
+		{
+			WorldTargetTM.SetTranslation(WorldCurrentTM.GetTranslation());
+		}
+		if (!bHaveAngular)
+		{
+			WorldTargetTM.SetRotation(WorldCurrentTM.GetRotation());
+		}
+
+		FVector WorldTargetVelocity = WorldParentFrameTM.GetRotation() * TargetVelocity;
+		FVector WorldTargetAngularVelocity = WorldParentFrameTM.GetRotation() * TargetAngularVelocity;
+
+		// Indicate the velocities by predicting the TargetTM
+		FTransform PredictedTargetTM = WorldTargetTM;
+		PredictedTargetTM.AddToTranslation(WorldTargetVelocity * VelocityPredictionTime);
+
+		// Draw the target and current positions/orientations
+		if (bHaveAngular)
+		{
+			FQuat AngularVelocityQ = FQuat::MakeFromRotationVector(WorldTargetAngularVelocity * VelocityPredictionTime);
+			PredictedTargetTM.SetRotation(AngularVelocityQ * WorldTargetTM.GetRotation());
+
+			DrawCoordinateSystem(
+				PDI, WorldCurrentTM.GetTranslation(), WorldCurrentTM.Rotator(), 
+				VisualizationSizeScale, SDPG_Foreground, 1.0f * GizmoWidthScale);
+			DrawCoordinateSystem(
+				PDI, WorldTargetTM.GetTranslation(), WorldTargetTM.Rotator(),
+				VisualizationSizeScale, SDPG_Foreground, 4.0f * GizmoWidthScale);
+			if (VelocityPredictionTime)
+			{
+				DrawCoordinateSystem(
+					PDI, PredictedTargetTM.GetTranslation(), PredictedTargetTM.Rotator(),
+					VisualizationSizeScale * 0.5, SDPG_Foreground, 4.0f * GizmoWidthScale);
+			}
+		}
+		else
+		{
+			DrawWireSphere(
+				PDI, WorldCurrentTM, CurrentColor, 
+				VisualizationSizeScale, 8, SDPG_Foreground, 1.0f * GizmoWidthScale);
+			DrawWireSphere(
+				PDI, WorldTargetTM, TargetColor, 
+				VisualizationSizeScale, 8, SDPG_Foreground, 3.0f * GizmoWidthScale);
+			if (VelocityPredictionTime)
+			{
+				DrawWireSphere(
+					PDI, PredictedTargetTM, TargetColor, 
+					VisualizationSizeScale * 0.5, 8, SDPG_Foreground, 3.0f * GizmoWidthScale);
+			}
+		}
+
+
+		if (VelocityPredictionTime)
+		{
+			PDI->DrawLine(
+				WorldTargetTM.GetTranslation(), 
+				WorldTargetTM.GetTranslation() + WorldTargetVelocity * VelocityPredictionTime, 
+				TargetColor, SDPG_Foreground);
+		}
+
+		// Connect current to target
+		DrawDashedLine(
+			PDI,
+			WorldTargetTM.GetTranslation(), WorldCurrentTM.GetTranslation(), 
+			CurrentToTargetColor, VisualizationSizeScale * 0.2f, SDPG_Foreground);
+	}
+}
+
+#endif
+
diff --git a/Engine/Plugins/Experimental/PhysicsControl/Source/PhysicsControl/Private/PhysicsControlComponentHelpers.cpp b/Engine/Plugins/Experimental/PhysicsControl/Source/PhysicsControl/Private/PhysicsControlComponentHelpers.cpp
new file mode 100644
--- /dev/null
+++ b/Engine/Plugins/Experimental/PhysicsControl/Source/PhysicsControl/Private/PhysicsControlComponentHelpers.cpp
@@ -0,0 +1,63 @@
+#include "PhysicsControlComponentHelpers.h"
+#include "Components/StaticMeshComponent.h"
+#include "Components/SkeletalMeshComponent.h"
+
+//======================================================================================================================
+void ConvertSpringParams(
+	double& OutSpring, double& OutDamping, 
+	double InStrength, double InDampingRatio, double InExtraDamping)
+{
+	double AngularFrequency = InStrength * 2.0 * PI;
+	double Stiffness = AngularFrequency * AngularFrequency;
+
+	OutSpring = Stiffness;
+	OutDamping = 2.0 * InDampingRatio * AngularFrequency;
+	OutDamping += InExtraDamping;
+}
+
+//======================================================================================================================
+void ConvertSpringParams(
+	FVector& OutSpring, FVector& OutDamping, 
+	const FVector& InStrength, float InDampingRatio, const FVector& InExtraDamping)
+{
+	ConvertSpringParams(OutSpring.X, OutDamping.X, InStrength.X, InDampingRatio, InExtraDamping.X);
+	ConvertSpringParams(OutSpring.Y, OutDamping.Y, InStrength.Y, InDampingRatio, InExtraDamping.Y);
+	ConvertSpringParams(OutSpring.Z, OutDamping.Z, InStrength.Z, InDampingRatio, InExtraDamping.Z);
+}
+
+//======================================================================================================================
+FBodyInstance* GetBodyInstance(UMeshComponent* MeshComponent, const FName BoneName)
+{
+	UStaticMeshComponent* StaticMeshComponent = Cast<UStaticMeshComponent>(MeshComponent);
+	if (StaticMeshComponent)
+	{
+		return StaticMeshComponent->GetBodyInstance();
+	}
+	USkeletalMeshComponent* SkeletalMeshComponent = Cast<USkeletalMeshComponent>(MeshComponent);
+	if (SkeletalMeshComponent)
+	{
+		return SkeletalMeshComponent->GetBodyInstance(BoneName);
+	}
+	return nullptr;
+}
+
+//======================================================================================================================
+FName GetPhysicalParentBone(USkeletalMeshComponent* SkeletalMeshComponent, FName BoneName)
+{
+	while (true)
+	{
+		FName ParentBoneName = SkeletalMeshComponent->GetParentBone(BoneName);
+		if (ParentBoneName.IsNone() || ParentBoneName == BoneName)
+		{
+			return FName();
+		}
+		const FBodyInstance* ParentBodyInstance = GetBodyInstance(SkeletalMeshComponent, ParentBoneName);
+		if (ParentBodyInstance)
+		{
+			return ParentBoneName;
+		}
+		BoneName = ParentBoneName;
+	}
+	return FName();
+}
+
diff --git a/Engine/Plugins/Experimental/PhysicsControl/Source/PhysicsControl/Private/PhysicsControlComponentHelpers.h b/Engine/Plugins/Experimental/PhysicsControl/Source/PhysicsControl/Private/PhysicsControlComponentHelpers.h
new file mode 100644
--- /dev/null
+++ b/Engine/Plugins/Experimental/PhysicsControl/Source/PhysicsControl/Private/PhysicsControlComponentHelpers.h
@@ -0,0 +1,37 @@
+// Copyright Epic Games, Inc. All Rights Reserved.
+
+#pragma once
+
+#include "UObject/NameTypes.h"
+
+struct FBodyInstance;
+class UMeshComponent;
+class USkeletalMeshComponent;
+
+/** 
+ * Converts strength/damping ratio/extra damping into spring stiffness/damping.
+ */
+void ConvertSpringParams(
+	double& OutSpring, double& OutDamping, 
+	double InStrength, double InDampingRatio, double InExtraDamping);
+
+/** 
+ * Converts strength/damping ratio/extra damping into spring stiffness/damping 
+ */
+void ConvertSpringParams(
+	FVector& OutSpring, FVector& OutDamping, 
+	const FVector& InStrength, float InDampingRatio, const FVector& InExtraDamping);
+
+/** 
+ * Attempts to find a BodyInstance from the mesh. If it is a static mesh the single body instance
+ * will be returned. If it is a skeletal mesh then if BoneName can be found, the body instance corresponding
+ * to that bone will be returned. Otherwise it will return nullptr if the bone can't be found.
+ */
+FBodyInstance* GetBodyInstance(UMeshComponent* MeshComponent, const FName BoneName);
+
+/**
+ * Attempts to find the parent physical bone given a skeletal mesh and starting bone. This walks up
+ * the hierarchy, ignoring non-physical bones, until either a physical bone is found, or it has reached
+ * the root without finding a physical bone.
+ */
+FName GetPhysicalParentBone(USkeletalMeshComponent* SkeletalMeshComponent, FName BoneName);
diff --git a/Engine/Plugins/Experimental/PhysicsControl/Source/PhysicsControl/Private/PhysicsControlComponentImpl.cpp b/Engine/Plugins/Experimental/PhysicsControl/Source/PhysicsControl/Private/PhysicsControlComponentImpl.cpp
new file mode 100644
--- /dev/null
+++ b/Engine/Plugins/Experimental/PhysicsControl/Source/PhysicsControl/Private/PhysicsControlComponentImpl.cpp
@@ -0,0 +1,633 @@
+// Copyright Epic Games, Inc. All Rights Reserved.
+
+#include "PhysicsControlComponentImpl.h"
+#include "PhysicsControlComponentLog.h"
+#include "PhysicsControlComponent.h"
+#include "PhysicsControlComponentHelpers.h"
+
+#include "Components/SkeletalMeshComponent.h"
+
+//======================================================================================================================
+bool FPhysicsControlComponentImpl::GetBoneData(
+	FCachedSkeletalMeshData::FBoneData& OutBoneData,
+	const USkeletalMeshComponent*       InSkeletalMeshComponent,
+	const FName                         InBoneName) const
+{
+	check(InSkeletalMeshComponent);
+	const FReferenceSkeleton& RefSkeleton = InSkeletalMeshComponent->GetSkeletalMeshAsset()->GetRefSkeleton();
+	const int32 BoneIndex = RefSkeleton.FindBoneIndex(InBoneName);
+
+	if (BoneIndex == INDEX_NONE)
+	{
+		UE_LOG(LogPhysicsControlComponent, Warning, TEXT("Failed to find BoneIndex for %s"), *InBoneName.ToString());
+		return false;
+	}
+
+	const FCachedSkeletalMeshData* CachedSkeletalMeshData = CachedSkeletalMeshDatas.Find(InSkeletalMeshComponent);
+	if (CachedSkeletalMeshData &&
+		CachedSkeletalMeshData->ReferenceCount > 0 &&
+		!CachedSkeletalMeshData->BoneData.IsEmpty())
+	{
+		if (BoneIndex < CachedSkeletalMeshData->BoneData.Num())
+		{
+			OutBoneData = CachedSkeletalMeshData->BoneData[BoneIndex];
+			return true;
+		}
+		UE_LOG(LogPhysicsControlComponent, Warning, TEXT("BoneIndex is out of range"));
+
+	}
+	UE_LOG(LogPhysicsControlComponent, Warning, TEXT("Failed to find bone data for %s"), *InBoneName.ToString());
+	return false;
+}
+
+//======================================================================================================================
+FPhysicsControlRecord* FPhysicsControlComponentImpl::FindControlRecord(const FName Name)
+{
+	if (PhysicsControlRecords.IsEmpty())
+	{
+		return nullptr;
+	}
+	if (Name.IsNone())
+	{
+		return &PhysicsControlRecords.CreateIterator().Value();
+	}
+	if (FPhysicsControlRecord* Record = PhysicsControlRecords.Find(Name))
+	{
+		return Record;
+	}
+	return nullptr;
+}
+
+//======================================================================================================================
+FPhysicsControl* FPhysicsControlComponentImpl::FindControl(const FName Name)
+{
+	if (FPhysicsControlRecord* ControlRecord = FindControlRecord(Name))
+	{
+		return &ControlRecord->PhysicsControl;
+	}
+	return nullptr;
+}
+
+//======================================================================================================================
+bool FPhysicsControlComponentImpl::DetectTeleport(
+	const FTransform& OldComponentTM, const FTransform& NewComponentTM) const
+{
+	if (Owner->TeleportDistanceThreshold > 0)
+	{
+		double Distance = FVector::Distance(OldComponentTM.GetTranslation(), NewComponentTM.GetTranslation());
+		if (Distance > Owner->TeleportDistanceThreshold)
+		{
+			return true;
+		}
+	}
+	if (Owner->TeleportRotationThreshold > 0)
+	{
+		double Radians = OldComponentTM.GetRotation().AngularDistance(NewComponentTM.GetRotation());
+		if (FMath::RadiansToDegrees(Radians) > Owner->TeleportRotationThreshold)
+		{
+			return true;
+		}
+	}
+	return false;
+}
+
+
+//======================================================================================================================
+void FPhysicsControlComponentImpl::UpdateCachedSkeletalBoneData(float Dt)
+{
+	for (TPair<TObjectPtr<USkeletalMeshComponent>, FCachedSkeletalMeshData>& CachedSkeletalMeshDataPair :
+		CachedSkeletalMeshDatas)
+	{
+		FCachedSkeletalMeshData& CachedSkeletalMeshData = CachedSkeletalMeshDataPair.Value;
+		if (!CachedSkeletalMeshData.ReferenceCount)
+		{
+			continue;
+		}
+
+		TObjectPtr<USkeletalMeshComponent> SkeletalMeshComponent = CachedSkeletalMeshDataPair.Key;
+		if (USkeletalMeshComponent* SkeletalMesh = SkeletalMeshComponent.Get())
+		{
+			FTransform ComponentTM = SkeletalMesh->GetComponentToWorld();
+			const TArray<FTransform>& TMs = SkeletalMesh->GetEditableComponentSpaceTransforms();
+			if (TMs.Num() == CachedSkeletalMeshData.BoneData.Num() &&
+				!DetectTeleport(CachedSkeletalMeshData.ComponentTM, ComponentTM))
+			{
+				for (int32 Index = 0; Index != TMs.Num(); ++Index)
+				{
+					FTransform TM = TMs[Index] * ComponentTM;
+					CachedSkeletalMeshData.BoneData[Index].Update(TM.GetTranslation(), TM.GetRotation(), Dt);
+				}
+			}
+			else
+			{
+				CachedSkeletalMeshData.BoneData.Empty(TMs.Num());
+				for (const FTransform& BoneTM : TMs)
+				{
+					FTransform TM = BoneTM * ComponentTM;
+					CachedSkeletalMeshData.BoneData.Emplace(TM.GetTranslation(), TM.GetRotation());
+				}
+			}
+			CachedSkeletalMeshData.ComponentTM = ComponentTM;
+		}
+		else
+		{
+			CachedSkeletalMeshData.BoneData.Empty();
+		}
+	}
+}
+
+//======================================================================================================================
+void FPhysicsControlComponentImpl::ResetControls(bool bKeepControlRecords)
+{
+	for (TPair<FName, FPhysicsControlRecord>& PhysicsControlRecordPair : PhysicsControlRecords)
+	{
+		FPhysicsControlRecord& Record = PhysicsControlRecordPair.Value;
+		Record.PhysicsControlState.Reset();
+	}
+
+	if (!bKeepControlRecords)
+	{
+		PhysicsControlRecords.Empty();
+	}
+}
+
+//======================================================================================================================
+// Currently this looks for world-space targets from the controls, and forms a strength-weighted
+// average of them if there are multiple targets. However, it would probably be better to replace
+// this with an explicit kinematic target on each body modifier, as it is a little unintuitive to
+// make zero strength physical controls. UE-159655
+void FPhysicsControlComponentImpl::ApplyKinematicTarget(const FPhysicsBodyModifier& BodyModifier) const
+{
+	FBodyInstance* BodyInstance = GetBodyInstance(BodyModifier.MeshComponent, BodyModifier.BoneName);
+	if (!BodyInstance)
+	{
+		return;
+	}
+
+	// First find any controls that are (a) acting in world space (b) driving the modified object
+	float PositionWeight = 0.0f;
+	float OrientationWeight = 0.0f;
+	FVector WeightedPosition(ForceInitToZero);
+	FVector4 WeightedOrientation(ForceInitToZero);
+	for (const TPair<FName, FPhysicsControlRecord>& PhysicsControlRecordPair : PhysicsControlRecords)
+	{
+		const FPhysicsControlRecord& Record = PhysicsControlRecordPair.Value;
+		FBodyInstance* ParentBodyInstance = GetBodyInstance(
+			Record.PhysicsControl.ParentMeshComponent, Record.PhysicsControl.ParentBoneName);
+
+		if (!ParentBodyInstance)
+		{
+			FBodyInstance* ChildBodyInstance = GetBodyInstance(
+				Record.PhysicsControl.ChildMeshComponent, Record.PhysicsControl.ChildBoneName);
+			if (BodyInstance == ChildBodyInstance)
+			{
+				FTransform TargetTM;
+				FVector TargetVelocity;
+				FVector TargetAngularVelocity;
+				CalculateControlTargetData(TargetTM, TargetVelocity, TargetAngularVelocity, Record, false);
+
+				// TargetTM is actually the target for the control point, but we will be setting the body TM
+				// so we need to remove the offset
+				TargetTM.AddToTranslation(TargetTM.GetRotation() * -Record.PhysicsControl.ControlSettings.ControlPoint);
+
+				// TODO note that this isn't using the multipliers, or the force limits etc. Using an explicit
+				// kinematic target will solve this
+				float LinearWeight = Record.PhysicsControl.ControlData.LinearStrength + UE_SMALL_NUMBER;
+				float AngularWeight = Record.PhysicsControl.ControlData.AngularStrength + UE_SMALL_NUMBER;
+
+				WeightedPosition += TargetTM.GetTranslation() * LinearWeight;
+				PositionWeight += LinearWeight;
+
+				FQuat Q = TargetTM.GetRotation();
+				Q.EnforceShortestArcWith(FQuat::Identity);
+				WeightedOrientation += FVector4(Q.X, Q.Y, Q.Z, Q.W) * AngularWeight;
+				OrientationWeight += AngularWeight;
+			}
+		}
+	}
+
+	if (PositionWeight <= 0 && OrientationWeight <= 0)
+	{
+		return;
+	}
+
+	// Seems like static and skeletal meshes need to be handled differently
+	if (BodyModifier.MeshComponent.IsA<USkeletalMeshComponent>())
+	{
+		FTransform NewTM = BodyInstance->GetUnrealWorldTransform();
+		if (PositionWeight)
+		{
+			FVector NewPosition = WeightedPosition / PositionWeight;
+			NewTM.SetLocation(NewPosition);
+		}
+		if (OrientationWeight)
+		{
+			WeightedOrientation /= OrientationWeight;
+			FQuat TargetOrientation = FQuat(
+				WeightedOrientation.X, WeightedOrientation.Y, WeightedOrientation.Z, WeightedOrientation.W);
+			TargetOrientation.Normalize();
+			NewTM.SetRotation(TargetOrientation);
+		}
+
+		BodyInstance->SetBodyTransform(NewTM, ETeleportType::None);
+	}
+	else
+	{
+		if (PositionWeight)
+		{
+			FVector TargetPosition = WeightedPosition / PositionWeight;
+			BodyModifier.MeshComponent->SetWorldLocation(TargetPosition, false, nullptr, ETeleportType::None);
+		}
+
+		if (OrientationWeight)
+		{
+			WeightedOrientation /= OrientationWeight;
+			FQuat TargetOrientation = FQuat(
+				WeightedOrientation.X, WeightedOrientation.Y, WeightedOrientation.Z, WeightedOrientation.W);
+			TargetOrientation.Normalize();
+			BodyModifier.MeshComponent->SetWorldRotation(TargetOrientation, false, nullptr, ETeleportType::None);
+		}
+	}
+}
+
+//======================================================================================================================
+void FPhysicsControlComponentImpl::AddSkeletalMeshReference(USkeletalMeshComponent* InSkeletalMeshComponent)
+{
+	check(InSkeletalMeshComponent);
+	for (TPair<TObjectPtr<USkeletalMeshComponent>, FCachedSkeletalMeshData>& CachedSkeletalMeshDataPair :
+		CachedSkeletalMeshDatas)
+	{
+		TObjectPtr<USkeletalMeshComponent> SkeletalMeshComponent = CachedSkeletalMeshDataPair.Key;
+		if (SkeletalMeshComponent == InSkeletalMeshComponent)
+		{
+			FCachedSkeletalMeshData& CachedSkeletalMeshData = CachedSkeletalMeshDataPair.Value;
+			++CachedSkeletalMeshData.ReferenceCount;
+			return;
+		}
+	}
+	FCachedSkeletalMeshData& Data = CachedSkeletalMeshDatas.Add(InSkeletalMeshComponent);
+	Data.ReferenceCount = 1;
+	Owner->PrimaryComponentTick.AddPrerequisite(InSkeletalMeshComponent, InSkeletalMeshComponent->PrimaryComponentTick);
+}
+
+//======================================================================================================================
+void FPhysicsControlComponentImpl::RemoveSkeletalMeshReference(USkeletalMeshComponent* InSkeletalMeshComponent)
+{
+	check(InSkeletalMeshComponent);
+	if (!InSkeletalMeshComponent)
+	{
+		UE_LOG(LogPhysicsControlComponent, Warning, TEXT("Invalid skeletal mesh component"));
+		return;
+	}
+
+	for (auto It = CachedSkeletalMeshDatas.CreateIterator(); It; ++It)
+	{
+		TObjectPtr<USkeletalMeshComponent> SkeletalMeshComponent = It.Key();
+		FCachedSkeletalMeshData& CachedSkeletalMeshData = It.Value();
+		if (SkeletalMeshComponent == InSkeletalMeshComponent)
+		{
+			if (--CachedSkeletalMeshData.ReferenceCount == 0)
+			{
+				Owner->PrimaryComponentTick.RemovePrerequisite(
+					InSkeletalMeshComponent, InSkeletalMeshComponent->PrimaryComponentTick);
+				It.RemoveCurrent();
+			}
+			return;
+		}
+	}
+	UE_LOG(LogPhysicsControlComponent, Warning, TEXT("Failed to remove skeletal mesh component dependency"));
+}
+
+//======================================================================================================================
+FName FPhysicsControlComponentImpl::GetUniqueControlName(const FName ParentBoneName, const FName ChildBoneName) const
+{
+	FString NameBase = TEXT("");
+	if (!ParentBoneName.IsNone())
+	{
+		NameBase += ParentBoneName.ToString() + TEXT("_");
+	}
+	if (!ChildBoneName.IsNone())
+	{
+		NameBase += ChildBoneName.ToString() + TEXT("_");
+	}
+
+	TSet<FName> Keys;
+	PhysicsControlRecords.GetKeys(Keys);
+	// If the number gets too large, almost certainly we're in some nasty situation where this is
+	// getting called in a loop. Better to quit and fail, rather than allow the constraint set to
+	// increase without bound. 
+	for (int32 Index = 0; Index < Owner->MaxNumControlsOrModifiersPerName; ++Index)
+	{
+		FString NameStr = FString::Format(TEXT("{0}{1}"), { NameBase, Index });
+		FName Name(NameStr);
+		if (!Keys.Find(Name))
+		{
+			return Name;
+		}
+	}
+	UE_LOG(LogPhysicsControlComponent, Warning,
+		TEXT("Unable to find a suitable Control name - the limit of MaxNumControlsOrModifiersPerName (%d) has been exceeded"),
+		Owner->MaxNumControlsOrModifiersPerName);
+	return FName();
+}
+
+//======================================================================================================================
+FName FPhysicsControlComponentImpl::GetUniqueBodyModifierName(const FName BoneName) const
+{
+	FString NameBase = TEXT("");
+	if (!BoneName.IsNone())
+	{
+		NameBase += BoneName.ToString() + TEXT("_");
+	}
+	else
+	{
+		NameBase = TEXT("Body_");
+	}
+
+	TSet<FName> Keys;
+	PhysicsBodyModifiers.GetKeys(Keys);
+	// If the number gets too large, almost certainly we're in some nasty situation where this is
+	// getting called in a loop. Better to quit and fail, rather than allow the modifier set to
+	// increase without bound. 
+	for (int32 Index = 0; Index != Owner->MaxNumControlsOrModifiersPerName; ++Index)
+	{
+		FString NameStr = FString::Format(TEXT("{0}{1}"), { NameBase, Index });
+		FName Name(NameStr);
+		if (!Keys.Find(Name))
+		{
+			return Name;
+		}
+	}
+	UE_LOG(LogPhysicsControlComponent, Warning,
+		TEXT("Unable to find a suitable Body Modifier name - the limit of MaxNumControlsOrModifiersPerName (%d) has been exceeded"),
+		Owner->MaxNumControlsOrModifiersPerName);
+	return FName();
+}
+
+//======================================================================================================================
+void FPhysicsControlComponentImpl::CalculateControlTargetData(
+	FTransform&                  OutTargetTM,
+	FVector&                     OutTargetVelocity,
+	FVector&                     OutTargetAngularVelocity,
+	const FPhysicsControlRecord& Record,
+	bool                         bCalculateVelocity) const
+{
+	const FPhysicsControlTarget& Target = Record.PhysicsControl.ControlTarget;
+
+	// Calculate the authored target position/orientation - i.e. not using the skeletal animation
+	FQuat TargetOrientationQ = Target.TargetOrientation.Quaternion();
+
+	// Incorporate the offset from the control point
+	FVector ExtraTargetPosition =
+		Record.PhysicsControl.ControlTarget.bApplyControlPointToTarget
+		? Record.PhysicsControl.ControlSettings.ControlPoint
+		: FVector::ZeroVector;
+	FVector ExtraTargetPositionWorld = TargetOrientationQ * ExtraTargetPosition;
+	FVector TargetPosition = Target.TargetPosition + ExtraTargetPositionWorld;
+
+	if (bCalculateVelocity)
+	{
+		// Note that Target.TargetAngularVelocity is in revs per second (as it's user-facing)
+		OutTargetAngularVelocity = Target.TargetAngularVelocity * UE_TWO_PI;
+		FVector ExtraVelocity = OutTargetAngularVelocity.Cross(ExtraTargetPositionWorld);
+		OutTargetVelocity = Target.TargetVelocity + ExtraVelocity;
+	}
+	else
+	{
+		OutTargetVelocity.Set(0, 0, 0);
+		OutTargetAngularVelocity.Set(0, 0, 0);
+	}
+
+	// OutTargetTM is the target transform of the constraint's child frame relative to the
+	// constraint's parent frame
+	OutTargetTM = FTransform(TargetOrientationQ, TargetPosition);
+
+	// Adjust based on any skeletal action
+	if (Record.PhysicsControl.ControlSettings.bUseSkeletalAnimation)
+	{
+		FCachedSkeletalMeshData::FBoneData ChildBoneData, ParentBoneData;
+		bool bHaveChildBoneData = false;
+		bool bHaveParentBoneData = false;
+
+		if (USkeletalMeshComponent* ChildSkeletalMeshComponent =
+			Cast<USkeletalMeshComponent>(Record.PhysicsControl.ChildMeshComponent.Get()))
+		{
+			bHaveChildBoneData = GetBoneData(
+				ChildBoneData, ChildSkeletalMeshComponent, Record.PhysicsControl.ChildBoneName);
+		}
+
+		if (USkeletalMeshComponent* ParentSkeletalMeshComponent =
+			Cast<USkeletalMeshComponent>(Record.PhysicsControl.ParentMeshComponent.Get()))
+		{
+			bHaveParentBoneData = GetBoneData(
+				ParentBoneData, ParentSkeletalMeshComponent, Record.PhysicsControl.ParentBoneName);
+		}
+
+		// Note that the TargetTM/velocity calculated so far are supposed to be interpreted as
+		// expressed relative to the skeletal animation pose.
+		//
+		// Also note that the velocities calculated in the bone data are the strict rates of change
+		// of the transform position/orientation - not of the center of mass (which is what physics
+		// bodies often use for velocity).
+		if (bHaveChildBoneData)
+		{
+			FTransform ChildBoneTM = ChildBoneData.GetTM();
+			if (bHaveParentBoneData)
+			{
+				FTransform ParentBoneTM = ParentBoneData.GetTM();
+				FTransform SkeletalDeltaTM = ChildBoneTM * ParentBoneTM.Inverse();
+				// This puts TargetTM in the space of the ParentBone
+				OutTargetTM = OutTargetTM * SkeletalDeltaTM;
+
+				FQuat ParentBoneQ = ParentBoneTM.GetRotation();
+				FQuat ParentBoneInvQ = ParentBoneQ.Inverse();
+
+				if (bCalculateVelocity)
+				{
+					OutTargetVelocity = SkeletalDeltaTM.GetRotation() * OutTargetVelocity;
+					OutTargetAngularVelocity = SkeletalDeltaTM.GetRotation() * OutTargetAngularVelocity;
+
+					if (Record.PhysicsControl.ControlSettings.SkeletalAnimationVelocityMultiplier)
+					{
+						// Offset of the control point from the target child bone TM, in world space.
+						FVector WorldControlPointOffset =
+							ChildBoneTM.GetRotation() * Record.PhysicsControl.ControlSettings.ControlPoint;
+						// World space position of the target control point
+						FVector WorldChildControlPointPosition = ChildBoneTM.GetTranslation() + WorldControlPointOffset;
+
+						// World-space velocity of the control point due to the motion of the parent
+						// linear and angular velocity.
+						FVector ChildTargetVelocityDueToParent =
+							ParentBoneData.Velocity + ParentBoneData.AngularVelocity.Cross(
+								WorldChildControlPointPosition - ParentBoneTM.GetTranslation());
+						// World-space velocity of the control point due to the motion of the child
+						// linear and angular velocity
+						FVector ChildTargetVelocity =
+							ChildBoneData.Velocity + ChildBoneData.AngularVelocity.Cross(WorldControlPointOffset);
+
+						// Pull out just the motion in the child that isn't due to the parent
+						FVector SkeletalTargetVelocity =
+							ParentBoneInvQ * (ChildTargetVelocity - ChildTargetVelocityDueToParent);
+						OutTargetVelocity += SkeletalTargetVelocity *
+							Record.PhysicsControl.ControlSettings.SkeletalAnimationVelocityMultiplier;
+
+						FVector SkeletalTargetAngularVelocity =
+							ParentBoneInvQ * (ChildBoneData.AngularVelocity - ParentBoneData.AngularVelocity);
+						OutTargetAngularVelocity += SkeletalTargetAngularVelocity *
+							Record.PhysicsControl.ControlSettings.SkeletalAnimationVelocityMultiplier;
+					}
+				}
+			}
+			else
+			{
+				OutTargetTM = OutTargetTM * ChildBoneTM;
+
+				if (bCalculateVelocity)
+				{
+					OutTargetVelocity = ChildBoneTM.GetRotation() * OutTargetVelocity;
+					OutTargetAngularVelocity = ChildBoneTM.GetRotation() * OutTargetAngularVelocity;
+
+					if (Record.PhysicsControl.ControlSettings.SkeletalAnimationVelocityMultiplier)
+					{
+						FVector WorldControlPointOffset =
+							ChildBoneTM.GetRotation() * Record.PhysicsControl.ControlSettings.ControlPoint;
+						FVector WorldChildControlPointPosition = ChildBoneTM.GetTranslation() + WorldControlPointOffset;
+
+						// World-space velocity of the control point due to the motion of the child
+						FVector ChildTargetVelocity =
+							ChildBoneData.Velocity + ChildBoneData.AngularVelocity.Cross(WorldControlPointOffset);
+
+						OutTargetVelocity += ChildTargetVelocity *
+							Record.PhysicsControl.ControlSettings.SkeletalAnimationVelocityMultiplier;
+
+						OutTargetAngularVelocity += ChildBoneData.AngularVelocity *
+							Record.PhysicsControl.ControlSettings.SkeletalAnimationVelocityMultiplier;
+					}
+				}
+			}
+		}
+	}
+}
+
+//======================================================================================================================
+bool FPhysicsControlComponentImpl::ApplyControlStrengths(
+	FPhysicsControlRecord& Record, FConstraintInstance* ConstraintInstance)
+{
+	const FPhysicsControlData& Data = Record.PhysicsControl.ControlData;
+	const FPhysicsControlMultipliers& Multipliers = Record.PhysicsControl.ControlMultipliers;
+
+	double AngularSpring;
+	double AngularDamping;
+	double MaxTorque = Data.MaxTorque * Multipliers.MaxTorqueMultiplier;
+
+	FVector LinearSpring;
+	FVector LinearDamping;
+	FVector MaxForce = Data.MaxForce * Multipliers.MaxForceMultiplier;
+
+	ConvertSpringParams(
+		AngularSpring, AngularDamping,
+		Data.AngularStrength * Multipliers.AngularStrengthMultiplier,
+		Data.AngularDampingRatio,
+		Data.AngularExtraDamping * Multipliers.AngularExtraDampingMultiplier);
+	ConvertSpringParams(
+		LinearSpring, LinearDamping,
+		Data.LinearStrength * Multipliers.LinearStrengthMultiplier,
+		Data.LinearDampingRatio,
+		Data.LinearExtraDamping * Multipliers.LinearExtraDampingMultiplier);
+
+	if (Multipliers.MaxTorqueMultiplier <= 0.0)
+	{
+		AngularSpring = 0.0;
+		AngularDamping = 0.0;
+	}
+	if (Multipliers.MaxForceMultiplier.X <= 0.0)
+	{
+		LinearSpring.X = 0.0;
+		LinearDamping.X = 0.0;
+	}
+	if (Multipliers.MaxForceMultiplier.Y <= 0.0)
+	{
+		LinearSpring.Y = 0.0;
+		LinearDamping.Y = 0.0;
+	}
+	if (Multipliers.MaxForceMultiplier.Z <= 0.0)
+	{
+		LinearSpring.Z = 0.0;
+		LinearDamping.Z = 0.0;
+	}
+
+	ConstraintInstance->SetAngularDriveParams(AngularSpring, AngularDamping, MaxTorque);
+	ConstraintInstance->SetLinearDriveParams(LinearSpring, LinearDamping, MaxForce);
+
+	double TestAngular = (AngularSpring + AngularDamping) * FMath::Max(UE_SMALL_NUMBER, MaxTorque);
+	FVector TestLinear = (LinearSpring + LinearDamping) * FVector(
+		FMath::Max(UE_SMALL_NUMBER, MaxForce.X),
+		FMath::Max(UE_SMALL_NUMBER, MaxForce.Y),
+		FMath::Max(UE_SMALL_NUMBER, MaxForce.Z));
+	double Test = TestAngular + TestLinear.GetMax();
+	return Test > 0.0;
+}
+
+//======================================================================================================================
+void FPhysicsControlComponentImpl::ApplyControl(FPhysicsControlRecord& Record)
+{
+	FConstraintInstance* ConstraintInstance = Record.PhysicsControlState.ConstraintInstance.Get();
+
+	if (!ConstraintInstance || !Record.PhysicsControlState.bEnabled)
+	{
+		return;
+	}
+
+	FBodyInstance* ParentBodyInstance = GetBodyInstance(
+		Record.PhysicsControl.ParentMeshComponent, Record.PhysicsControl.ParentBoneName);
+
+	FBodyInstance* ChildBodyInstance = GetBodyInstance(
+		Record.PhysicsControl.ChildMeshComponent, Record.PhysicsControl.ChildBoneName);
+
+	if (!ParentBodyInstance && !ChildBodyInstance)
+	{
+		return;
+	}
+
+	// Set strengths etc
+	if (ApplyControlStrengths(Record, ConstraintInstance))
+	{
+		FTransform TargetTM;
+		FVector TargetVelocity;
+		FVector TargetAngularVelocity;
+		CalculateControlTargetData(TargetTM, TargetVelocity, TargetAngularVelocity, Record, true);
+
+		ConstraintInstance->SetLinearPositionTarget(TargetTM.GetTranslation());
+		ConstraintInstance->SetAngularOrientationTarget(TargetTM.GetRotation());
+		ConstraintInstance->SetLinearVelocityTarget(TargetVelocity);
+		ConstraintInstance->SetAngularVelocityTarget(TargetAngularVelocity / UE_TWO_PI); // In rev/sec
+
+		if (ParentBodyInstance)
+		{
+			ParentBodyInstance->WakeInstance();
+		}
+		if (ChildBodyInstance)
+		{
+			ChildBodyInstance->WakeInstance();
+		}
+	}
+}
+
+//======================================================================================================================
+FPhysicsBodyModifier* FPhysicsControlComponentImpl::FindBodyModifier(const FName Name)
+{
+	if (PhysicsBodyModifiers.IsEmpty())
+	{
+		return nullptr;
+	}
+	if (Name.IsNone())
+	{
+		return &PhysicsBodyModifiers.CreateIterator().Value();
+	}
+	if (FPhysicsBodyModifier* Record = PhysicsBodyModifiers.Find(Name))
+	{
+		return Record;
+	}
+	return nullptr;
+}
+
diff --git a/Engine/Plugins/Experimental/PhysicsControl/Source/PhysicsControl/Private/PhysicsControlComponentImpl.h b/Engine/Plugins/Experimental/PhysicsControl/Source/PhysicsControl/Private/PhysicsControlComponentImpl.h
new file mode 100644
--- /dev/null
+++ b/Engine/Plugins/Experimental/PhysicsControl/Source/PhysicsControl/Private/PhysicsControlComponentImpl.h
@@ -0,0 +1,117 @@
+// Copyright Epic Games, Inc. All Rights Reserved.
+
+#pragma once
+
+#include "PhysicsControlRecord.h"
+#include "UObject/ObjectPtr.h"
+#include "UObject/NameTypes.h"
+
+class USkeletalMeshComponent;
+class UPhysicsControlComponent;
+struct FPhysicsControl;
+
+/**
+ * Internal/private data/implementation for UPhysicsControlComponent
+ */
+struct FPhysicsControlComponentImpl
+{
+public:
+
+	FPhysicsControlComponentImpl(UPhysicsControlComponent* InOwner) : Owner(InOwner) {}
+
+	/**
+	 * Retrieves the bone data for the specified bone given the skeletal mesh component.
+	 * 
+	 * @param OutBoneData If successful, this will contain the output bone data.
+	 * @param InSkeletalMeshComponent Required to be a valid pointer to a skeletal mesh component
+	 * @param InBoneName The name of the bone to retrieve data for
+	 * @return true if the bone and data were found, false if not (in which case warnings will be logged)
+	 */
+	bool GetBoneData(
+		FCachedSkeletalMeshData::FBoneData& OutBoneData,
+		const USkeletalMeshComponent*       InSkeletalMeshComponent,
+		const FName                         InBoneName) const;
+
+	/**
+	 * Retrieves the control for the name. Note that if Name is blank then the first control will be returned,
+	 * assuming there is one.
+	 */
+	FPhysicsControl* FindControl(const FName Name);
+	FPhysicsControlRecord* FindControlRecord(const FName Name);
+
+	/**
+	 * Updates the world-space bone positions etc for each skeleton we're tracking
+	 */
+	void UpdateCachedSkeletalBoneData(float DeltaTime);
+
+	/**
+	 * @return true if the difference between the old and new TMs exceeds the teleport
+	 * translation/rotation thresholds
+	 */
+	bool DetectTeleport(const FTransform& OldComponentTM, const FTransform& NewComponentTM) const;
+
+	/**
+	 * Terminates the underlying physical constraints, resets our internal stored state for each control,
+	 * and optionally deletes all record of the controls.
+	 */
+	void ResetControls(bool bKeepControlRecords);
+
+	/**
+	 * Starts caching skeletal mesh poses, and registers for a tick pre-requisite
+	 */
+	void AddSkeletalMeshReference(USkeletalMeshComponent* SkeletalMeshComponent);
+
+	/**
+	 * Stops caching skeletal mesh poses (if this is the last one), and deregisters for a tick pre-requisite
+	 */
+	void RemoveSkeletalMeshReference(USkeletalMeshComponent* SkeletalMeshComponent);
+
+	/** Update the constraint based on the record. */
+	void ApplyControl(FPhysicsControlRecord& Record);
+
+	/**
+	 * Updates the constraint strengths. Returns true if there is some strength, 
+	 * false otherwise (e.g. to skip setting targets)
+	 */
+	bool ApplyControlStrengths(FPhysicsControlRecord& Record, FConstraintInstance* ConstraintInstance);
+
+	/**
+	 * Calculates the Target TM and velocities that will get passed to the constraint - so this 
+	 * is a target that is defined in the space of the parent body (or in world space, if it doesn't exist).
+	 */
+	void CalculateControlTargetData(
+		FTransform&                  OutTargetTM, 
+		FVector&                     OutTargetVelocity,
+		FVector&                     OutTargetAngularVelocity,
+		const FPhysicsControlRecord& Record,
+		bool                         bCalculateVelocity) const;
+
+	/**
+	 * This will set the kinematic target for the appropriate body based on the weighted target position and 
+	 * orientation (and whether any were found) for any controls that are related to the body modifier.
+	 */
+	void ApplyKinematicTarget(const FPhysicsBodyModifier& BodyModifier) const;
+
+	/**
+	 * Retrieves the body modifier for the name. Note that if Name is blank then the first modifier will be returned,
+	 * assuming there is one.
+	 */
+	FPhysicsBodyModifier* FindBodyModifier(const FName Name);
+
+	FName GetUniqueControlName(const FName ParentBoneName, const FName ChildBoneName) const;
+	FName GetUniqueBodyModifierName(const FName BoneName) const;
+
+public:
+
+	// Cached transforms from each skeletal mesh we're working with. Will be updated at the
+	// beginning of each tick.
+	TMap<TObjectPtr<USkeletalMeshComponent>, FCachedSkeletalMeshData> CachedSkeletalMeshDatas;
+
+	TMap<FName, FPhysicsControlRecord> PhysicsControlRecords;
+
+	TMap<FName, FPhysicsBodyModifier> PhysicsBodyModifiers;
+
+private:
+	UPhysicsControlComponent* Owner;
+};
+
diff --git a/Engine/Plugins/Experimental/PhysicsControl/Source/PhysicsControl/Private/PhysicsControlComponentLog.cpp b/Engine/Plugins/Experimental/PhysicsControl/Source/PhysicsControl/Private/PhysicsControlComponentLog.cpp
new file mode 100644
--- /dev/null
+++ b/Engine/Plugins/Experimental/PhysicsControl/Source/PhysicsControl/Private/PhysicsControlComponentLog.cpp
@@ -0,0 +1,5 @@
+// Copyright Epic Games, Inc. All Rights Reserved.
+
+#include "PhysicsControlComponentLog.h"
+
+DEFINE_LOG_CATEGORY(LogPhysicsControlComponent);
diff --git a/Engine/Plugins/Experimental/PhysicsControl/Source/PhysicsControl/Private/PhysicsControlComponentLog.h b/Engine/Plugins/Experimental/PhysicsControl/Source/PhysicsControl/Private/PhysicsControlComponentLog.h
new file mode 100644
--- /dev/null
+++ b/Engine/Plugins/Experimental/PhysicsControl/Source/PhysicsControl/Private/PhysicsControlComponentLog.h
@@ -0,0 +1,7 @@
+// Copyright Epic Games, Inc. All Rights Reserved.
+
+#pragma once
+
+#include "Logging/LogMacros.h"
+
+DECLARE_LOG_CATEGORY_EXTERN(LogPhysicsControlComponent, Log, All);
diff --git a/Engine/Plugins/Experimental/PhysicsControl/Source/PhysicsControl/Private/PhysicsControlModule.cpp b/Engine/Plugins/Experimental/PhysicsControl/Source/PhysicsControl/Private/PhysicsControlModule.cpp
new file mode 100644
--- /dev/null
+++ b/Engine/Plugins/Experimental/PhysicsControl/Source/PhysicsControl/Private/PhysicsControlModule.cpp
@@ -0,0 +1,18 @@
+// Copyright Epic Games, Inc. All Rights Reserved.
+
+#include "PhysicsControlModule.h"
+
+#define LOCTEXT_NAMESPACE "FPhysicsControlModule"
+
+void FPhysicsControlModule::StartupModule()
+{
+}
+
+void FPhysicsControlModule::ShutdownModule()
+{
+}
+
+#undef LOCTEXT_NAMESPACE
+	
+IMPLEMENT_MODULE(FPhysicsControlModule, PhysicsControl)
+
diff --git a/Engine/Plugins/Experimental/PhysicsControl/Source/PhysicsControl/Private/PhysicsControlRecord.cpp b/Engine/Plugins/Experimental/PhysicsControl/Source/PhysicsControl/Private/PhysicsControlRecord.cpp
new file mode 100644
--- /dev/null
+++ b/Engine/Plugins/Experimental/PhysicsControl/Source/PhysicsControl/Private/PhysicsControlRecord.cpp
@@ -0,0 +1,115 @@
+// Copyright Epic Games, Inc. All Rights Reserved.
+
+#include "PhysicsControlRecord.h"
+#include "PhysicsControlComponentLog.h"
+#include "PhysicsControlComponentHelpers.h"
+
+#include "PhysicsEngine/ConstraintInstance.h"
+#include "PhysicsEngine/BodyInstance.h"
+
+//======================================================================================================================
+void FPhysicsControlState::Reset()
+{
+	if (ConstraintInstance)
+	{
+		ConstraintInstance->TermConstraint();
+	}
+	*this = FPhysicsControlState();
+}
+
+//======================================================================================================================
+// A constraint created with identity frames will just have the child frame at the mesh origin,
+// which is not necessarily where the center of gravity is.
+FConstraintInstance* FPhysicsControlRecord::CreateConstraint(UObject* ConstraintDebugOwner)
+{
+	if (!PhysicsControlState.ConstraintInstance)
+	{
+		PhysicsControlState.ConstraintInstance = MakeShared<FConstraintInstance>();
+	}
+	FConstraintInstance* ConstraintInstance = PhysicsControlState.ConstraintInstance.Get();
+
+	FBodyInstance* ParentBody = GetBodyInstance(
+		PhysicsControl.ParentMeshComponent, PhysicsControl.ParentBoneName);
+	FBodyInstance* ChildBody = GetBodyInstance(
+		PhysicsControl.ChildMeshComponent, PhysicsControl.ChildBoneName);
+
+	if (PhysicsControl.ParentMeshComponent && !PhysicsControl.ParentBoneName.IsNone() && !ParentBody)
+	{
+		UE_LOG(LogPhysicsControlComponent, Warning, TEXT("Failed to find expected parent body when making constraint"));
+		return nullptr;
+	}
+	if (PhysicsControl.ChildMeshComponent && !PhysicsControl.ChildBoneName.IsNone() && !ChildBody)
+	{
+		UE_LOG(LogPhysicsControlComponent, Warning, TEXT("Failed to find expected child body when making constraint"));
+		return nullptr;
+	}
+
+	ConstraintInstance->InitConstraint(ChildBody, ParentBody, 1.0f, ConstraintDebugOwner);
+
+	// Ensure the control point is set
+	UpdateConstraintControlPoint();
+
+	ConstraintInstance->SetLinearXMotion(ELinearConstraintMotion::LCM_Free);
+	ConstraintInstance->SetLinearYMotion(ELinearConstraintMotion::LCM_Free);
+	ConstraintInstance->SetLinearZMotion(ELinearConstraintMotion::LCM_Free);
+	ConstraintInstance->SetAngularSwing1Motion(EAngularConstraintMotion::ACM_Free);
+	ConstraintInstance->SetAngularSwing2Motion(EAngularConstraintMotion::ACM_Free);
+	ConstraintInstance->SetAngularTwistMotion(EAngularConstraintMotion::ACM_Free);
+
+	ConstraintInstance->SetLinearPositionDrive(true, true, true);
+	ConstraintInstance->SetLinearVelocityDrive(true, true, true);
+	ConstraintInstance->SetAngularDriveMode(EAngularDriveMode::SLERP);
+	ConstraintInstance->SetOrientationDriveSLERP(true);
+	ConstraintInstance->SetAngularVelocityDriveSLERP(true);
+
+	return ConstraintInstance;
+}
+
+//======================================================================================================================
+// Note that, by default, the constraint frames are simply identity. We only modify Frame1, which 
+// corresponds to the child frame. Frame2 will always be identity, because we never change it.
+void FPhysicsControlRecord::UpdateConstraintControlPoint()
+{
+	FConstraintInstance* ConstraintInstance = PhysicsControlState.ConstraintInstance.Get();
+	if (ConstraintInstance)
+	{
+		// Constraints are child then parent
+		FTransform Frame1 = ConstraintInstance->GetRefFrame(EConstraintFrame::Frame1);
+		Frame1.SetTranslation(PhysicsControl.ControlSettings.ControlPoint);
+		ConstraintInstance->SetRefFrame(EConstraintFrame::Frame1, Frame1);
+	}
+}
+
+//======================================================================================================================
+void FPhysicsControlRecord::ResetControlPoint()
+{
+	FBodyInstance* ChildBodyInstance = GetBodyInstance(
+		PhysicsControl.ChildMeshComponent, PhysicsControl.ChildBoneName);
+
+	PhysicsControl.ControlSettings.ControlPoint = ChildBodyInstance
+		? ChildBodyInstance->GetMassSpaceLocal().GetTranslation() : FVector::ZeroVector;
+
+	UpdateConstraintControlPoint();
+}
+
+//======================================================================================================================
+void FCachedSkeletalMeshData::FBoneData::Update(const FVector& InPosition, const FQuat& InOrientation, float Dt)
+{
+	if (Dt > 0)
+	{
+		Velocity = (InPosition - Position) / Dt;
+		Orientation.EnforceShortestArcWith(InOrientation);
+		// Note that quats multiply in the opposite order to TMs
+		FQuat DeltaQ = InOrientation * Orientation.Inverse();
+		AngularVelocity = DeltaQ.ToRotationVector() / Dt;
+	}
+	else
+	{
+		Velocity = FVector::ZeroVector;
+		AngularVelocity = FVector::ZeroVector;
+	}
+	Position = InPosition;
+	Orientation = InOrientation;
+}
+
+
diff --git a/Engine/Plugins/Experimental/PhysicsControl/Source/PhysicsControl/Private/PhysicsControlRecord.h b/Engine/Plugins/Experimental/PhysicsControl/Source/PhysicsControl/Private/PhysicsControlRecord.h
new file mode 100644
--- /dev/null
+++ b/Engine/Plugins/Experimental/PhysicsControl/Source/PhysicsControl/Private/PhysicsControlRecord.h
@@ -0,0 +1,136 @@
+// Copyright Epic Games, Inc. All Rights Reserved.
+
+#pragma once
+
+#include "UObject/ObjectMacros.h"
+#include "PhysicsControlData.h"
+
+struct FConstraintInstance;
+struct FBodyInstance;
+class UMeshComponent;
+
+/**
+ * The basic state of a physics control, created for every control record.
+ */
+struct FPhysicsControlState
+{
+	FPhysicsControlState()
+		: bEnabled(false), bPendingDestroy(false)
+	{
+	}
+
+	/** Removes any constraint and resets the state */
+	void Reset();
+
+	TSharedPtr<FConstraintInstance> ConstraintInstance;
+
+	bool bEnabled;
+	bool bPendingDestroy;
+};
+
+/**
+ * There will be a PhysicsControlRecord created at runtime for every Control that has been created
+ */
+struct FPhysicsControlRecord : public TSharedFromThis<FPhysicsControlRecord>
+{
+	FPhysicsControlRecord(const FPhysicsControl& InControl)
+		: PhysicsControl(InControl)
+	{
+	}
+
+	/** 
+	 * Creates the constraint (if necessary) and stores it in the state. ConstraintDebugOwner is passed 
+	 * to the constraint on creation. 
+	 */
+	FConstraintInstance* CreateConstraint(UObject* ConstraintDebugOwner);
+
+	/** Ensures the constraint frame matches the control point in the record. */
+	void UpdateConstraintControlPoint();
+
+	/** Sets the control point to the center of mass of the child mesh (or to zero if that fails). */
+	void ResetControlPoint();
+
+	/** The configuration data */
+	FPhysicsControl PhysicsControl;
+
+	/** The instance/runtime state - instantiated and kept up to date (during the tick) with PhysicsControl. */
+	FPhysicsControlState PhysicsControlState;
+};
+
+/**
+ * There will be a PhysicsBodyModifier created at runtime for every BodyInstance involved in the component
+ */
+struct FPhysicsBodyModifier : public TSharedFromThis<FPhysicsBodyModifier>
+{
+	FPhysicsBodyModifier(
+		TObjectPtr<UMeshComponent> InMeshComponent, 
+		const FName&               InBoneName, 
+		EPhysicsMovementType       InMovementType, 
+		float                      InGravityMultiplier)
+	: MeshComponent(InMeshComponent), BoneName(InBoneName)
+	, MovementType(InMovementType), GravityMultiplier(InGravityMultiplier)
+	, bPendingDestroy(false)
+	{}
+
+	/**  The mesh that will be modified. */
+	TObjectPtr<UMeshComponent> MeshComponent;
+
+	/** The name of the skeletal mesh bone or the name of the static mesh body that will be modified. */
+	FName                      BoneName;
+
+	EPhysicsMovementType       MovementType;
+	float                      GravityMultiplier;
+
+	bool bPendingDestroy;
+};
+
+/**
+ * Used internally/only at runtime to cache skeletal transforms at the start of the tick, to avoid
+ * calculating them separately for every control.
+ */
+struct FCachedSkeletalMeshData
+{
+public:
+	FCachedSkeletalMeshData() : ReferenceCount(0) {}
+
+public:
+	struct FBoneData
+	{
+		FBoneData()
+			: Position(FVector::ZeroVector), Orientation(FQuat::Identity)
+			, Velocity(FVector::ZeroVector), AngularVelocity(FVector::ZeroVector) {}
+		FBoneData(const FVector& InPosition, const FQuat& InOrientation)
+			: Position(InPosition), Orientation(InOrientation),
+			Velocity(FVector::ZeroVector), AngularVelocity(FVector::ZeroVector) {}
+
+		/**
+		 * Sets position and velocity, and calculates velocity if Dt > 0 (otherwise sets it to zero)
+		 */
+		void Update(const FVector& InPosition, const FQuat& InOrientation, float Dt);
+		FTransform GetTM() const { return FTransform(Orientation, Position); }
+
+		FVector Position;
+		FQuat   Orientation;
+		FVector Velocity;
+		FVector AngularVelocity;
+	};
+
+public:
+
+	/**
+	 * The cached skeletal data, updated at the start of each tick
+	 */
+	TArray<FBoneData> BoneData;
+
+	/**
+	 * The component transform. This is only stored so we can detect teleports
+	 */
+	FTransform ComponentTM;
+
+	/**
+	 * Track when skeletal meshes are going to be used so this entry can be removed, and also so we
+	 * can add a tick dependency
+	 */
+	int32 ReferenceCount;
+};
+
diff --git a/Engine/Plugins/Experimental/PhysicsControl/Source/PhysicsControl/Public/PhysicsControlComponent.h b/Engine/Plugins/Experimental/PhysicsControl/Source/PhysicsControl/Public/PhysicsControlComponent.h
new file mode 100644
--- /dev/null
+++ b/Engine/Plugins/Experimental/PhysicsControl/Source/PhysicsControl/Public/PhysicsControlComponent.h
@@ -0,0 +1,712 @@
+// Copyright Epic Games, Inc. All Rights Reserved.
+
+#pragma once
+
+#include "PhysicsControlData.h"
+#include "PhysicsControlLimbData.h"
+
+#include "UObject/ObjectMacros.h"
+#include "Components/ActorComponent.h"
+#include "EngineDefines.h"
+#include "Templates/PimplPtr.h"
+
+#include "PhysicsControlComponent.generated.h"
+
+class FPrimitiveDrawInterface;
+struct FPhysicsControlComponentImpl;
+struct FPhysicsControlRecord;
+struct FPhysicsBodyModifier;
+struct FConstraintInstance;
+
+/**
+ * Specifies the type of control that is created when making controls from a skeleton or a set of limbs. 
+ * Note that if controls are made individually then other options are available - i.e. in a character, 
+ * any body part can be controlled relative to any other part, or indeed any other object.
+ */
+UENUM(BlueprintType)
+enum class EPhysicsControlType : uint8
+{
+	/** Control is done in world space, so each object/part is driven independently */
+	WorldSpace,
+	/** Control is done in the space of the parent of each object */
+	ParentSpace,
+};
+
+/**
+ * This is the main Physics Control Component class which manages Controls and Body Modifiers associated 
+ * with one or more static or skeletal meshes. You can add this as a component to an actor containing a 
+ * mesh and then use it to create, configure and destroy Controls/Body Modifiers:
+ * 
+ * Controls are used to control one physics body relative to another (or the world). These controls are done
+ * through physical spring/damper drives.
+ * 
+ * Body Modifiers are used to update the most important physical properties of physics bodies such as whether 
+ * they are simulated vs kinematic, or whether they experience gravity.
+ * 
+ * Note that Controls and Body Modifiers are given names (which are predictable). These names can then be stored 
+ * (perhaps in arrays) to make it easy to quickly change multiple Controls/Body Modifiers.
+ */
+UCLASS(meta = (BlueprintSpawnableComponent), ClassGroup = Physics, Experimental)
+class PHYSICSCONTROL_API UPhysicsControlComponent : public UActorComponent
+{
+	GENERATED_UCLASS_BODY()
+
+public:
+	/**
+	 * Makes a new control for mesh components
+	 * 
+	 * @param ControlData   Describes the initial strength etc of the new control
+	 * 
+	 * @param ControlTarget Describes the initial target for the new control
+	 * 
+	 * @param ControlSettings General settings for the control
+	 * 
+	 * @param bEnabled      If true then the control will be enabled immediately. If false you will need to call
+	 *                      SetControlEnabled(true) in order to activate it.
+	 *
+	 * @return The name of the new control
+	 */
+	UFUNCTION(BlueprintCallable, Category = PhysicsControl)
+	FName MakeControl(
+		UMeshComponent*         ParentMeshComponent,
+		const FName             ParentBoneName,
+		UMeshComponent*         ChildMeshComponent,
+		const FName             ChildBoneName,
+		FPhysicsControlData     ControlData,
+		FPhysicsControlTarget   ControlTarget, 
+		FPhysicsControlSettings ControlSettings,
+		bool                    bEnabled = true);
+
+	/**
+	 * Makes a new control for mesh components
+	 * 
+	 * @param ControlData   Describes the initial strength etc of the new control
+	 * 
+	 * @param ControlTarget Describes the initial target for the new control
+	 * 
+	 * @param ControlSettings General settings for the control
+	 * 
+	 * @param bEnabled      If true then the control will be enabled immediately. If false you will need to call
+	 *                      SetControlEnabled(true) in order to enable it.
+	 * 
+	 * @return True if a new control was created, false if a control of the specified name already exists
+	 */
+	UFUNCTION(BlueprintCallable, Category = PhysicsControl)
+	bool MakeNamedControl(
+		const FName             Name,
+		UMeshComponent*         ParentMeshComponent,
+		const FName             ParentBoneName,
+		UMeshComponent*         ChildMeshComponent,
+		const FName             ChildBoneName,
+		FPhysicsControlData     ControlData, 
+		FPhysicsControlTarget   ControlTarget, 
+		FPhysicsControlSettings ControlSettings, 
+		bool                    bEnabled = true);
+
+	/**
+	 * Makes a collection of controls controlling a skeletal mesh
+	 * 
+	 * @param SkeletalMeshComponent The skeletal mesh which will have controls
+	 * 
+	 * @param BoneName The name of the bone below which controls should be created. Each bone will be the child in a control
+	 * 
+	 * @param bIncludeSelf Whether or not to include BoneName when creating controls
+	 * 
+	 * @param ControlType What type of control to create. This determines what the parent will be for each control
+	 *
+	 * @param ControlData   Describes the initial strength etc of the new control
+	 * 
+	 * @param ControlSettings General settings for the control
+	 *
+	 * @param bEnabled      If true then the control will be enabled immediately. If false you will need to call
+	 *                      SetControlEnabled(true) in order to enable it.
+	 *
+	 * @return An array of the controls that have been created
+	 */
+	UFUNCTION(BlueprintCallable, Category = PhysicsControl)
+	TArray<FName> MakeControlsFromSkeletalMeshBelow(
+		USkeletalMeshComponent* SkeletalMeshComponent,
+		FName                   BoneName,
+		bool                    bIncludeSelf,
+		EPhysicsControlType     ControlType,
+		FPhysicsControlData     ControlData,
+		FPhysicsControlSettings ControlSettings,
+		bool                    bEnabled = true);
+
+	/**
+	 * Makes a collection of controls controlling a skeletal mesh
+	 *
+	 * @param SkeletalMeshComponent The skeletal mesh which will have controls
+	 *
+	 * @param BoneNames The names of bones for which controls should be created. Each bone will be the child in a control
+	 *
+	 * @param ControlType What type of control to create. This determines what the parent will be for each control
+	 *
+	 * @param ControlData   Describes the initial strength etc of the new control
+	 *
+	 * @param ControlSettings General settings for the control
+	 *
+	 * @param bEnabled      If true then the control will be enabled immediately. If false you will need to call
+	 *                      SetControlEnabled(true) in order to enable it.
+	 *
+	 * @return An array of the controls that have been created
+	 */
+	UFUNCTION(BlueprintCallable, Category = PhysicsControl)
+	TArray<FName> MakeControlsFromSkeletalMesh(
+		USkeletalMeshComponent* SkeletalMeshComponent,
+		const TArray<FName>&    BoneNames,
+		EPhysicsControlType     ControlType,
+		FPhysicsControlData     ControlData,
+		FPhysicsControlSettings ControlSettings,
+		bool                    bEnabled = true);
+
+	/**
+	 * Calculates which bones belong to which limb in a skeletal mesh
+	 * 
+	 * @param SkeletalMeshComponent The skeletal mesh which will be analyzed
+	 * 
+	 * @param LimbSetupData This needs to be filled in with the list of limbs to "discover". Note that the 
+	 *                      limbs should be listed starting at the "leaf" (i.e. outer) parts of the skeleton 
+	 *                      first, typically finishing with the spine. In addition, the spine limb is typically 
+	 *                      specified using the first spine bone, but flagging it to include its parent 
+	 *                      (normally the pelvis). 
+	 * 
+	 * @return A map of limb names to bones
+	 */
+	UFUNCTION(BlueprintCallable, Category = PhysicsControl)
+	TMap<FName, FPhysicsControlLimbBones> GetLimbBonesFromSkeletalMesh(
+		USkeletalMeshComponent* SkeletalMeshComponent,
+		const TArray<FPhysicsControlLimbSetupData>& LimbSetupData) const;
+
+	/**
+	 * Makes a collection of controls controlling a skeletal mesh, grouped together in limbs
+	 *
+	 * @param AllControls A single container for all the controls that have been created
+	 *
+	 * @param LimbBones A map relating the limbs and the bones that they contain. Typically create this 
+	 *                  using GetLimbBonesFromSkeletalMesh
+	 *
+	 * @param ControlType What type of control to create. This determines what the parent will be for each control
+	 *
+	 * @param ControlData Describes the initial strength etc of the new control
+	 *
+	 * @param ControlSettings General settings for the control
+	 *
+	 * @param bEnabled If true then the control will be enabled immediately. If false you will need to call
+	 *                 SetControlEnabled(true) in order to enable it.
+	 *
+	 * @return A map containing the controls for each limb
+	 */
+	UFUNCTION(BlueprintCallable, Category = PhysicsControl)
+	TMap<FName, FPhysicsControlNameArray> MakeControlsFromLimbBones(
+		FPhysicsControlNameArray&                    AllControls,
+		const TMap<FName, FPhysicsControlLimbBones>& LimbBones,
+		EPhysicsControlType                          ControlType,
+		FPhysicsControlData                          ControlData,
+		FPhysicsControlSettings                      ControlSettings,
+		bool                                         bEnabled = true);
+
+	/**
+	 * Destroys a control
+	 *
+	 * @param Name The name of the control to destroy. If blank, then this will destroy the first
+	 *             control, whatever its name.
+	 *
+	 * @return     Returns true if the control was found and destroyed, false if not
+	 */
+	UFUNCTION(BlueprintCallable, Category = PhysicsControl)
+	bool DestroyControl(const FName Name);
+
+	/**
+	 * Modifies an existing control data - i.e. the strengths etc of the control driving towards the target
+	 *
+	 * @param Name The name of the control to modify. If blank, then this will access the first 
+	 *             control, whatever its name.
+	 * 
+	 * @param ControlData The new control data
+	 * 
+	 * @param bEnableControl Enables/disables the control
+	 * 
+	 * @return Returns true if the control was found and modified, false if not
+	 */
+	UFUNCTION(BlueprintCallable, Category = PhysicsControl)
+	bool SetControlData(const FName Name, FPhysicsControlData ControlData, bool bEnableControl = true);
+
+	/**
+	 * Modifies an existing control data using the multipliers
+	 *
+	 * @param Name The name of the control to modify. If blank, then this will access the first
+	 *             control, whatever its name.
+	 *
+	 * @param ControlMultipliers The new control multipliers
+	 *
+	 * @param bEnableControl Enables/disables the control
+	 *
+	 * @return Returns true if the control was found and modified, false if not
+	 */
+	UFUNCTION(BlueprintCallable, Category = PhysicsControl)
+	bool SetControlMultipliers(const FName Name, FPhysicsControlMultipliers ControlMultipliers, bool bEnableControl = true);
+
+	/**
+	 * Modifies an existing control's linear data - i.e. the strengths etc of the control driving towards the target
+	 *
+	 * @param Name The name of the control to modify. If blank, then this will access the first
+	 *             control, whatever its name.
+	 *
+	 * @param Strength The strength used to drive linear motion
+	 *
+	 * @param DampingRatio The amount of damping associated with the linear strength. 1 Results in critically damped motion
+	 *
+	 * @param ExtraDamping The amount of additional linear damping
+	 *
+	 * @param MaxForce The maximum force used to drive the linear motion. Zero indicates no limit.
+	 *
+	 * @param bEnableControl Enables/disables the control
+	 * 
+	 * @return Returns true if the control was found and modified, false if not
+	 */
+	UFUNCTION(BlueprintCallable, Category = PhysicsControl)
+	bool SetControlLinearData(
+		const FName Name,
+		float Strength = 1.0f, 
+		float DampingRatio = 1.0f, 
+		float ExtraDamping = 0.0f, 
+		float MaxForce = 0.0f, 
+		bool  bEnableControl = true);
+
+	/**
+	 * Modifies an existing control's angular data - i.e. the strengths etc of the control driving towards the target
+	 *
+	 * @param Name The name of the control to modify. If blank, then this will access the first
+	 *             control, whatever its name.
+	 *
+	 * @param Strength The strength used to drive angular motion
+	 *
+	 * @param DampingRatio The amount of damping associated with the angular strength. 1 Results in critically damped motion
+	 *
+	 * @param ExtraDamping The amount of additional angular damping
+	 *
+	 * @param MaxTorque The maximum torque used to drive the angular motion. Zero indicates no limit.
+	 * 
+	 * @param bEnableControl Enables/disables the control
+	 *
+	 * @return Returns true if the control was found and modified, false if not
+	 */
+	UFUNCTION(BlueprintCallable, Category = PhysicsControl)
+	bool SetControlAngularData(
+		const FName Name,
+		float Strength = 1.0f, 
+		float DampingRatio = 1.0f, 
+		float ExtraDamping = 0.0f, 
+		float MaxTorque = 0.0f, 
+		bool  bEnableControl = true);
+
+	/**
+	 * Sets the point at which controls will "push" the child object.
+	 * 
+	 * @param Name The name of the control to modify. If blank, then this will access the first
+	 *             control, whatever its name.
+	 * 
+	 * @param Position The position of the control point on the child mesh object (only relevant if that 
+	 *                 object is in use and is being simulated)
+	 */
+	UFUNCTION(BlueprintCallable, Category = PhysicsControl)
+	bool SetControlPoint(const FName Name, const FVector Position);
+
+	/**
+	 * Resets the control point to the center of mass of the mesh
+	 *
+	 * @param Name    The name of the control to modify. If blank, then this will access the first
+	 *                control, whatever its name.
+	 */
+	UFUNCTION(BlueprintCallable, Category = PhysicsControl)
+	bool ResetControlPoint(const FName Name);
+
+	/**
+	 * Modifies an existing control target - i.e. what it is driving towards, relative to the parent object
+	 *
+	 * @param Name The name of the control to modify. If blank, then this will access the first 
+	 *             control, whatever its name.
+	 * 
+	 * @param ControlTarget The new target for the control
+	 * 
+	 * @param bEnableControl Enables/disables the control
+	 * 
+	 * @return Returns true if the control was found and modified, false if not
+	 */
+	UFUNCTION(BlueprintCallable, Category = PhysicsControl)
+	bool SetControlTarget(const FName Name, FPhysicsControlTarget ControlTarget, bool bEnableControl = true);
+
+	/**
+	 * Modifies an existing control target - i.e. what it is driving towards, relative to the parent object
+	 *
+	 * @param Name The name of the control to modify. If blank, then this will access the first
+	 *             control, whatever its name.
+	 *
+	 * @param Transform The new transform target for the control
+	 * 
+	 * @param VelocityDeltaTime If non-zero, the target velocity will be calculated using the current target 
+	 *                          position. If zero, the target velocity will be set to zero.    
+	 *
+	 * @param bEnableControl Enables/disables the control
+	 * 
+	 * @param bApplyControlPointToTarget If true, then the target position/orientation is treated as
+	 *                                   a "virtual" object, where the system attempts to move the object
+	 *                                   to match the pose of this "virtual" object that has been placed at
+	 *                                   the target transform. Use this when you want to specify the target
+	 *                                   transform for the object as a whole. If false, then the target transform
+	 *                                   is used as is, and the system drives the control point towards this
+	 *                                   transform.
+	 *
+	 * @return                  Returns true if the control was found and modified, false if not
+	 */
+	UFUNCTION(BlueprintCallable, Category = PhysicsControl)
+	bool SetControlTargetTransform(
+		const FName      Name,
+		const FTransform Transform, 
+		float            VelocityDeltaTime, 
+		bool             bEnableControl = true,
+		bool             bApplyControlPointToTarget = true);
+
+	/**
+	 * Modifies an existing control target - i.e. what it is driving towards, relative to the parent object
+	 *
+	 * @param Name              The name of the control to modify. If blank, then this will access the first
+	 *                          control, whatever its name.
+	 *
+	 * @param Position          The new position target for the control
+	 * 
+	 * @param VelocityDeltaTime If non-zero, the target velocity will be calculated using the current target 
+	 *                          position. If zero, the target velocity will be set to zero.    
+	 *
+	 * @param bEnableControl    Enables/disables the control
+	 * 
+	 * @param bApplyControlPointToTarget If true, then the target position/orientation is treated as
+	 *                                   a "virtual" object, where the system attempts to move the object
+	 *                                   to match the pose of this "virtual" object that has been placed at
+	 *                                   the target transform. Use this when you want to specify the target
+	 *                                   transform for the object as a whole. If false, then the target transform
+	 *                                   is used as is, and the system drives the control point towards this
+	 *                                   transform.
+	 *
+	 * @return                  Returns true if the control was found and modified, false if not
+	 */
+	UFUNCTION(BlueprintCallable, Category = PhysicsControl)
+	bool SetControlTargetPosition(
+		const FName Name,
+		const FVector  Position, 
+		float          VelocityDeltaTime, 
+		bool           bEnableControl = true,
+		bool           bApplyControlPointToTarget = true);
+
+	/**
+	 * Modifies an existing control target - i.e. what it is driving towards, relative to the parent object
+	 *
+	 * @param Name                      The name of the control to modify. If blank, then this will access the first
+	 *                                  control, whatever its name.
+	 *							        
+	 * @param Orientation               The new orientation target for the control
+	 *
+	 * @param AngularVelocityDeltaTime  If non-zero, the target angular velocity will be calculated using the current 
+	 *                                  target position. If zero, the target velocity will be set to zero.
+	 * 
+	 * @param bEnableControl Enables/disables the control
+	 *
+	 * @param bApplyControlPointToTarget If true, then the target position/orientation is treated as
+	 *                                   a "virtual" object, where the system attempts to move the object
+	 *                                   to match the pose of this "virtual" object that has been placed at
+	 *                                   the target transform. Use this when you want to specify the target
+	 *                                   transform for the object as a whole. If false, then the target transform
+	 *                                   is used as is, and the system drives the control point towards this
+	 *                                   transform.
+	 * 
+	 * @return                          Returns true if the control was found and modified, false if not
+	 */
+	UFUNCTION(BlueprintCallable, Category = PhysicsControl)
+	bool SetControlTargetOrientation(
+		const FName Name,
+		const FRotator Orientation, 
+		float          AngularVelocityDeltaTime, 
+		bool           bEnableControl = true,
+		bool           bApplyControlPointToTarget = true);
+
+	/**
+	 * Calculates and sets an existing control target. This takes the "virtual" position/orientation of the parent 
+	 * and child and calculates the relative control. Note that this will set bApplyControlPointToTarget to true.
+	 * 
+	 * @param Name              The name of the control to modify. If blank, then this will access the first
+	 *                          control, whatever its name.
+	 *
+	 * @param ParentPosition    The virtual/target parent position
+	 * 
+	 * @param ParentOrientation The virtual/target parent orientation
+	 *
+	 * @param ChildPosition     The virtual/target child position
+	 *
+	 * @param ChildOrientation  The virtual/target child orientation
+	 *
+	 * @param VelocityDeltaTime If non-zero, the target velocity will be calculated using the current target
+	 *                          position. If zero, the target velocity will be set to zero.    
+	 *
+	 * @param bEnableControl Enables/disables the control
+	 * 
+	 * @return                  Returns true if the control was found and modified, false if not
+	 */
+	UFUNCTION(BlueprintCallable, Category = PhysicsControl)
+	bool SetControlTargetPoses(
+		const FName Name,
+		const FVector ParentPosition, const FRotator ParentOrientation,
+		const FVector ChildPosition, const FRotator ChildOrientation,
+		float VelocityDeltaTime, bool bEnableControl = true);
+
+
+	/**
+	 * Sets whether or not the control should use skeletal animation for the targets
+	 *
+	 * @param Name              The name of the control to modify. If blank, then this will access the first
+	 *                          control, whatever its name.
+	 * 
+	 * @param bUseSkeletalAnimation If true then the targets will be a combination of the skeletal animation (if
+	 *                              there is any) and the control target that has been set
+	 * 
+	 * @param SkeletalAnimationVelocityMultiplier If skeletal animation is being used, then this determines the amount of 
+	 *                              velocity extracted from the animation that is used as targets for the controls
+	 */
+	UFUNCTION(BlueprintCallable, Category = PhysicsControl)
+	bool SetControlUseSkeletalAnimation(
+		const FName Name,
+		bool        bUseSkeletalAnimation = true,
+		float       SkeletalAnimationVelocityMultiplier = 1.0f);
+
+	/**
+	 * Activates or deactivates a control
+	 *
+	 * @param Name     The name of the control to modify. If blank, then this will access the first
+	 *                 control, whatever its name.
+	 *
+	 * @param bEnable  The control to enable/disable
+	 *
+	 * @return         Returns true if the control was found and modified, false if not
+	 */
+	UFUNCTION(BlueprintCallable, Category = PhysicsControl)
+	bool SetControlEnabled(const FName Name, bool bEnable);
+
+	/**
+	 * @param Name The name of the control to modify. If blank, then this will access the first 
+	 *             control, whatever its name.
+	 * 
+	 * @param bAutoDisable If set then the control will automatically deactivate after each tick.
+	 * 
+	 * @return     Returns true if the control was found and modified, false if not
+	 */
+	UFUNCTION(BlueprintCallable, Category = PhysicsControl)
+	bool SetControlAutoDisable(const FName Name, bool bAutoDisable);
+
+	/**
+	 * @param Name     The name of the control to access. If blank, then this will access the first 
+	 *                 control, whatever its name.
+	 *
+	 * @param Control  The control that will be filled in, if found
+	 *
+	 * @return         Returns true if the control was found, false if not
+	 */
+	//UFUNCTION(BlueprintCallable, Category = PhysicsControl)
+	bool GetControl(const FName Name, FPhysicsControl& Control) const;
+
+	/**
+	 * @param Name     The name of the control to access. If blank, then this will access the first
+	 *                 control, whatever its name.
+	 *
+	 * @param Control  The control data that will be filled in if found
+	 *
+	 * @return         Returns true if the control was found, false if not
+	 */
+	UFUNCTION(BlueprintCallable, Category = PhysicsControl)
+	bool GetControlData(const FName Name, FPhysicsControlData& ControlData) const;
+
+	/**
+	 * @param Name     The name of the control to access. If blank, then this will access the first
+	 *                 control, whatever its name.
+	 *
+	 * @param Control  The control multipliers that will be filled in if found
+	 *
+	 * @return         Returns true if the control was found, false if not
+	 */
+	UFUNCTION(BlueprintCallable, Category = PhysicsControl)
+	bool GetControlMultipliers(const FName Name, FPhysicsControlMultipliers& ControlMultipliers) const;
+
+	/**
+	 * @param Name     The name of the control to access. If blank, then this will access the first
+	 *                 control, whatever its name.
+	 *
+	 * @param Control  The control target, if found
+	 *
+	 * @return         Returns true if the control was found, false if not
+	 */
+	UFUNCTION(BlueprintCallable, Category = PhysicsControl)
+	bool GetControlTarget(const FName Name, FPhysicsControlTarget& ControlTarget) const;
+
+	/**
+	 * @param Name        The name of the control to access. If blank, then this will access the first
+	 *                    control, whatever its name.
+	 *
+	 * @return            Returns true if the control is marked to automatically disable after each tick
+	 */
+	UFUNCTION(BlueprintCallable, Category = PhysicsControl)
+	bool GetControlAutoDisable(const FName Name) const;
+
+	/**
+	 * @param Name        The name of the control to access. If blank, then this will access the first
+	 *                    control, whatever its name.
+	 *
+	 * @return            Returns true if the control is enabled
+	 */
+	UFUNCTION(BlueprintCallable, Category = PhysicsControl)
+	bool GetControlEnabled(const FName Name) const;
+
+	/**
+	 * Makes a new body modifier for mesh components
+	 */
+	UFUNCTION(BlueprintCallable, Category = PhysicsControl)
+	FName MakeBodyModifier(
+		UMeshComponent*       MeshComponent,
+		const FName           BoneName,
+		EPhysicsMovementType  MovementType = EPhysicsMovementType::Simulated, 
+		float                 GravityMultiplier = 1.0f);
+
+	/**
+	 * Makes a new body modifier for mesh components
+	 */
+	UFUNCTION(BlueprintCallable, Category = PhysicsControl)
+	bool MakeNamedBodyModifier(
+		const FName           Name,
+		UMeshComponent*       MeshComponent,
+		const FName           BoneName,
+		EPhysicsMovementType  MovementType = EPhysicsMovementType::Simulated,
+		float                 GravityMultiplier = 1.0f);
+
+	/**
+	 * Makes new body modifiers for skeletal mesh components
+	 */
+	UFUNCTION(BlueprintCallable, Category = PhysicsControl)
+	TArray<FName> MakeBodyModifiersFromSkeletalMeshBelow(
+		USkeletalMeshComponent* SkeletalMeshComponent,
+		FName                   BoneName,
+		bool                    bIncludeSelf,
+		EPhysicsMovementType    MovementType = EPhysicsMovementType::Simulated,
+		float                   GravityMultiplier = 1.0f);
+
+
+	/**
+	 * Makes a collection of controls controlling a skeletal mesh, grouped together in limbs
+	 *
+	 * @param AllControls A single container for all the controls that have been created
+	 *
+	 * @param LimbBones A map relating the limbs and the bones that they contain. Typically create this 
+	 *                  using GetLimbBonesFromSkeletalMesh
+	 *
+	 * @param ControlType What type of control to create. This determines what the parent will be for each control
+	 *
+	 * @param ControlData   Describes the initial strength etc of the new control
+	 *
+	 * @param ControlSettings General settings for the control
+	 *
+	 * @param bEnabled      If true then the control will be enabled immediately. If false you will need to call
+	 *                      SetControlEnabled(true) in order to enable it.
+	 *
+	 * @return A map containing the controls for each limb
+	 */
+	UFUNCTION(BlueprintCallable, Category = PhysicsControl)
+	TMap<FName, FPhysicsControlNameArray> MakeBodyModifiersFromLimbBones(
+		FPhysicsControlNameArray&                    AllBodyModifiers,
+		const TMap<FName, FPhysicsControlLimbBones>& LimbBones,
+		EPhysicsMovementType                         MovementType = EPhysicsMovementType::Simulated,
+		float                                        GravityMultiplier = 1.0f);
+
+	/**
+	 * Destroys a BodyModifier
+	 *
+	 * @param Name        The name of the body modifier to destroy. If blank, then this will destroy the first
+	 *                    body modifier, whatever its name.
+	 *
+	 * @return            Returns true if the body modifier was found and destroyed, false if not
+	 */
+	UFUNCTION(BlueprintCallable, Category = PhysicsControl)
+	bool DestroyBodyModifier(const FName Name);
+
+	/**
+	 * @param Name        The name of the body modifier to access. If blank, then this will access the first
+	 *                    modifier, whatever its name.
+	 *
+	 * @param bSimulate   Whether to enable/disable simulation on the body
+	 * 
+	 * @param GravityMultiplier The amount of gravity to apply when simulating
+	 * 
+	 * @return            Returns true if the body modifier was found and modified
+	 */
+	UFUNCTION(BlueprintCallable, Category = PhysicsControl)
+	bool SetBodyModifier(
+		const FName          Name,
+		EPhysicsMovementType MovementType = EPhysicsMovementType::Simulated, 
+		float                GravityMultiplier = 1.0f);
+
+public:
+
+	/**
+	 * If the component moves by more than this distance then it is treated as a teleport,
+	 * which prevents velocities being used for a frame. Zero or negative disables.
+	 */
+	UPROPERTY(EditAnywhere, BlueprintReadWrite, Category = Teleport)
+	float TeleportDistanceThreshold;
+
+	/**
+	 * If the component rotates by more than this angle (in degrees) then it is treated as a teleport,
+	 * which prevents velocities being used for a frame. Zero or negative disables.
+	 */
+	UPROPERTY(EditAnywhere, BlueprintReadWrite, Category = Teleport)
+	float TeleportRotationThreshold;
+
+	/** Visualize the controls when this actor/component is selected */
+	UPROPERTY(EditAnywhere, BlueprintReadWrite, Category = Debug)
+	bool bShowDebugVisualization;
+
+	/** Size of the gizmos etc used during visualization */
+	UPROPERTY(EditAnywhere, BlueprintReadWrite, Category = Debug)
+	float VisualizationSizeScale;
+
+	/**
+	 * The time used when "predicting" the target position/orientation. Zero will disable the visualization
+	 * of this.
+	 */
+	UPROPERTY(EditAnywhere, BlueprintReadWrite, Category = Debug)
+	float VelocityPredictionTime;
+
+	/**
+	 * Upper limit on the number of controls or modifiers that will be created using the same name (which
+	 * will get a numerical postfix). When this limit is reached a warning will be issued  and the control 
+	 * or modifier won't be created. This is to avoid problems if controls or modifiers are being created 
+	 * dynamically, and can generally be a "moderately large" number, depending on how many controls or 
+	 * modifiers you expect to create.
+	 */
+	UPROPERTY(EditAnywhere, BlueprintReadWrite, Category = Settings)
+	int32 MaxNumControlsOrModifiersPerName;
+
+protected:
+
+	// Begin UActorComponent Interface
+	virtual void InitializeComponent() override;
+	virtual void BeginDestroy() override;
+	virtual void TickComponent(float DeltaTime, enum ELevelTick TickType, FActorComponentTickFunction* ThisTickFunction) override;
+	// End UActorComponent Interface
+
+public:
+#if WITH_EDITOR
+	// Used by the component visualizer
+	void DebugDraw(FPrimitiveDrawInterface* PDI) const;
+	void DebugDrawControl(FPrimitiveDrawInterface* PDI, const FPhysicsControlRecord& Record, FName ControlName) const;
+#endif
+
+protected:
+
+	TPimplPtr<FPhysicsControlComponentImpl> Implementation;
+};
diff --git a/Engine/Plugins/Experimental/PhysicsControl/Source/PhysicsControl/Public/PhysicsControlData.h b/Engine/Plugins/Experimental/PhysicsControl/Source/PhysicsControl/Public/PhysicsControlData.h
new file mode 100644
--- /dev/null
+++ b/Engine/Plugins/Experimental/PhysicsControl/Source/PhysicsControl/Public/PhysicsControlData.h
@@ -0,0 +1,304 @@
+// Copyright Epic Games, Inc. All Rights Reserved.
+
+#pragma once
+
+#include "UObject/ObjectMacros.h"
+#include "Components/ActorComponent.h"
+#include "EngineDefines.h"
+#include "PhysicsControlData.generated.h"
+
+class UMeshComponent;
+
+/**
+ * Used by Body Modifiers to specify how the physical bodies should move.  
+ */
+UENUM(BlueprintType)
+enum class EPhysicsMovementType : uint8
+{
+	// Default means that the simulation mode and transform won't be changed when set in a Body Modifier
+	Default,
+	// Kinematic means that the object won't be simulated, but will be moved to the average of any
+	// world-space control targets (if there are any), weighted by strength.
+	Kinematic,
+	// Simulated means that the object will be controlled by the physics solver
+	Simulated
+};
+
+/**
+ * Contains data associated with how physical bodies should be controlled/directed towards their targets. 
+ * The underlying control is done through damped springs, so the parameters here relate to that.
+ */
+USTRUCT(BlueprintType)
+struct PHYSICSCONTROL_API FPhysicsControlData
+{
+	GENERATED_BODY()
+
+	FPhysicsControlData()
+		: LinearStrength(1.0f)
+		, LinearDampingRatio(1.0f)
+		, LinearExtraDamping(0.0f)
+		, MaxForce(0.0f)
+		, AngularStrength(1.0f)
+		, AngularDampingRatio(1.0f)
+		, AngularExtraDamping(0.0f)
+		, MaxTorque(0.0f)
+	{
+	}
+
+	/** The strength used to drive linear motion */
+	UPROPERTY(EditAnywhere, BlueprintReadWrite, Category = PhysicsControl, meta = (ClampMin = "0.0"))
+	float LinearStrength;
+
+	/** 
+	 * The amount of damping associated with the linear strength. A value of 1 Results in critically 
+	 * damped motion where the control drives as quickly as possible to the target without overshooting. 
+	 * Values > 1 result in more damped motion, and values below 1 result in faster, but more "wobbly" motion.
+	 */
+	UPROPERTY(EditAnywhere, BlueprintReadWrite, Category = PhysicsControl, meta = (ClampMin = "0.0"))
+	float LinearDampingRatio;
+
+	/** 
+	 * The amount of additional linear damping. This is added to the damping that comes from LinearDampingRatio
+	 * and can be useful when you want damping even when LinearStrength is zero.
+	 */
+	UPROPERTY(EditAnywhere, BlueprintReadWrite, Category = PhysicsControl, meta = (ClampMin = "0.0"))
+	float LinearExtraDamping;
+
+	/** The maximum force used to drive the linear motion. Zero indicates no limit. */
+	UPROPERTY(EditAnywhere, BlueprintReadWrite, Category = PhysicsControl, meta = (ClampMin = "0.0"))
+	float MaxForce;
+
+	/** The strength used to drive angular motion */
+	UPROPERTY(EditAnywhere, BlueprintReadWrite, Category = PhysicsControl, meta = (ClampMin = "0.0"))
+	float AngularStrength;
+
+	/** 
+	 * The amount of damping associated with the angular strength. A value of 1 Results in critically 
+	 * damped motion where the control drives as quickly as possible to the target without overshooting. 
+	 * Values > 1 result in more damped motion, and values below 1 result in faster, but more "wobbly" motion.
+	 */
+	UPROPERTY(EditAnywhere, BlueprintReadWrite, Category = PhysicsControl, meta = (ClampMin = "0.0"))
+	float AngularDampingRatio;
+
+	/** 
+	 * The amount of additional angular damping. This is added to the damping that comes from AngularDampingRatio
+	 * and can be useful when you want damping even when AngularStrength is zero.
+	 */
+	UPROPERTY(EditAnywhere, BlueprintReadWrite, Category = PhysicsControl, meta = (ClampMin = "0.0"))
+	float AngularExtraDamping;
+
+	/** The maximum torque used to drive the angular motion. Zero indicates no limit. */
+	UPROPERTY(EditAnywhere, BlueprintReadWrite, Category = PhysicsControl, meta = (ClampMin = "0.0"))
+	float MaxTorque;
+};
+
+/**
+ * These parameters allow modification of the parameters in FPhysicsControlData for two reasons:
+ * 1. They allow per-axis settings for the linear components (e.g. so you can drive an object 
+ *    horizontally but still let it fall under gravity)
+ * 2. They make it easy to create the controls with "default" strength/damping (e.g. taken from the
+ *    physics asset) in FPhysicsControlData, and then the strength/damping etc can be scaled every 
+ *    tick (typically between 0 and 1, though that is up to the user).
+ */
+USTRUCT(BlueprintType)
+struct PHYSICSCONTROL_API FPhysicsControlMultipliers
+{
+	GENERATED_BODY()
+
+	FPhysicsControlMultipliers()
+		: LinearStrengthMultiplier(1.0)
+		, LinearExtraDampingMultiplier(1.0)
+		, MaxForceMultiplier(1.0)
+		, AngularStrengthMultiplier(1.0)
+		, AngularExtraDampingMultiplier(1.0)
+		, MaxTorqueMultiplier(1.0)
+	{
+	}
+
+	// Per-direction multiplier on the linear strength.
+	UPROPERTY(EditAnywhere, BlueprintReadWrite, Category = PhysicsControl, meta = (ClampMin = "0.0"))
+	FVector LinearStrengthMultiplier;
+
+	// Per-direction multiplier on the linear extra damping.
+	UPROPERTY(EditAnywhere, BlueprintReadWrite, Category = PhysicsControl, meta = (ClampMin = "0.0"))
+	FVector LinearExtraDampingMultiplier;
+
+	// Per-direction multiplier on the maximum force that can be applied. Note that zero means zero force.
+	UPROPERTY(EditAnywhere, BlueprintReadWrite, Category = PhysicsControl, meta = (ClampMin = "0.0"))
+	FVector MaxForceMultiplier;
+
+	// Multiplier on the angular strength.
+	UPROPERTY(EditAnywhere, BlueprintReadWrite, Category = PhysicsControl, meta = (ClampMin = "0.0"))
+	float AngularStrengthMultiplier;
+
+	// Multiplier on the angular extra damping.
+	UPROPERTY(EditAnywhere, BlueprintReadWrite, Category = PhysicsControl, meta = (ClampMin = "0.0"))
+	float AngularExtraDampingMultiplier;
+
+	// Per-direction multiplier on the maximum torque that can be applied. Note that zero means zero torque.
+	UPROPERTY(EditAnywhere, BlueprintReadWrite, Category = PhysicsControl, meta = (ClampMin = "0.0"))
+	float MaxTorqueMultiplier;
+};
+
+/**
+ * Defines a target position and orientation, and also the target velocity and angular velocity.
+ * In many cases the velocities will be calculated automatically (e.g. when setting the target position,
+ * the component will optionally calculate an implied velocity. However, the user can also specify a 
+ * target velocity directly. Note that the velocity influences the control through the damping parameters
+ * in FPhysicsControlData
+ */
+USTRUCT(BlueprintType)
+struct PHYSICSCONTROL_API FPhysicsControlTarget
+{
+	GENERATED_BODY()
+
+	FPhysicsControlTarget()
+		: TargetPosition(ForceInitToZero)
+		, TargetVelocity(ForceInitToZero)
+		, TargetOrientation(ForceInitToZero)
+		, TargetAngularVelocity(ForceInitToZero)
+		, bApplyControlPointToTarget(true)
+	{
+	}
+
+	/** The target position of the child body, relative to the parent body */
+	UPROPERTY(EditAnywhere, BlueprintReadWrite, Category = PhysicsControl)
+	FVector TargetPosition;
+
+	/** The target velocity of the child body, relative to the parent body */
+	UPROPERTY(EditAnywhere, BlueprintReadWrite, Category = PhysicsControl)
+	FVector TargetVelocity;
+
+	/** The target orientation of the child body, relative to the parent body */
+	UPROPERTY(EditAnywhere, BlueprintReadWrite, Category = PhysicsControl)
+	FRotator TargetOrientation;
+
+	/** The target angular velocity (revolutions per second) of the child body, relative to the parent body */
+	UPROPERTY(EditAnywhere, BlueprintReadWrite, Category = PhysicsControl)
+	FVector TargetAngularVelocity;
+
+	/** 
+	 * Whether to use the ControlPoint as an offset for the target transform, as well as the 
+	 * physical body. If true then the target TM is treated as a target transform for the actual 
+	 * object, though the control is still applied through the control point (which is at the 
+	 * center of mass by default). If false then it is treated as a target transform for the 
+	 * control point on the object.
+	 */
+	UPROPERTY(EditAnywhere, BlueprintReadWrite, Category = PhysicsControl)
+	bool bApplyControlPointToTarget;
+};
+
+/**
+ * General settings for a control
+ */
+USTRUCT(BlueprintType)
+struct PHYSICSCONTROL_API FPhysicsControlSettings
+{
+	GENERATED_BODY()
+
+	FPhysicsControlSettings()
+		: ControlPoint(ForceInitToZero)
+		, bUseSkeletalAnimation(true)
+		, SkeletalAnimationVelocityMultiplier(1.0f)
+		, bAutoDisable(false)
+	{
+	}
+
+	/** The position of the control point relative to the child mesh. */
+	UPROPERTY(EditAnywhere, BlueprintReadWrite, Category = PhysicsControl)
+	FVector ControlPoint;
+
+	/** If true then the target will be applied on top of the skeletal animation (if there is any) */
+	UPROPERTY(EditAnywhere, BlueprintReadWrite, Category = PhysicsControl)
+	bool bUseSkeletalAnimation;
+
+	/** The amount of skeletal animation velocity to use in the targets */
+	UPROPERTY(EditAnywhere, BlueprintReadWrite, Category = PhysicsControl)
+	float SkeletalAnimationVelocityMultiplier;
+
+	/**
+	 * Whether or not this control should automatically disable itself at the end of each tick. This
+	 * can be useful when it is more convenient to have a branch that handles some condition (e.g. character 
+	 * is flailing) by updating/enabling the control, and to then have the control automatically get 
+	 * disabled when the branch is no longer taken.
+	 */
+	UPROPERTY(EditAnywhere, BlueprintReadWrite, Category = PhysicsControl)
+	bool bAutoDisable;
+};
+
+/**
+ * Structure that determines a "control" - this contains all the information needed to drive (with spring-dampers)
+ * a child body relative to a parent body. These bodies will be associated with either a static or skeletal mesh.
+ */
+USTRUCT(BlueprintType)
+struct PHYSICSCONTROL_API FPhysicsControl
+{
+	GENERATED_BODY()
+
+	FPhysicsControl() {}
+
+	FPhysicsControl(
+		UMeshComponent*              InParentMeshComponent,
+		const FName&                 InParentBoneName,
+		UMeshComponent*              InChildMeshComponent,
+		const FName&                 InChildBoneName,
+		const FPhysicsControlData&   InControlData,
+		const FPhysicsControlTarget& InControlTarget,
+		const FPhysicsControlSettings& InControlSettings)
+		: ParentMeshComponent(InParentMeshComponent)
+		, ParentBoneName(InParentBoneName)
+		, ChildMeshComponent(InChildMeshComponent)
+		, ChildBoneName(InChildBoneName)
+		, ControlData(InControlData)
+		, ControlTarget(InControlTarget)
+		, ControlSettings(InControlSettings)
+	{
+	}
+
+	/**  The mesh that will be doing the driving. Blank/non-existent means it will happen in world space */
+	UPROPERTY(EditAnywhere, BlueprintReadWrite, Category = PhysicsControl)
+	TObjectPtr<UMeshComponent> ParentMeshComponent;
+
+	/** The name of the skeletal mesh bone or the name of the static mesh body that will be doing the driving. */
+	UPROPERTY(EditAnywhere, BlueprintReadWrite, Category = PhysicsControl)
+	FName ParentBoneName;
+
+	/** The mesh that the control will be driving. */
+	UPROPERTY(EditAnywhere, BlueprintReadWrite, Category = PhysicsControl)
+	TObjectPtr<UMeshComponent> ChildMeshComponent;
+
+	/** 
+	 * The name of the skeletal mesh bone or the name of the static mesh body that the control 
+	 * will be driving. */
+	UPROPERTY(EditAnywhere, BlueprintReadWrite, Category = PhysicsControl)
+	FName ChildBoneName;
+
+	/** 
+	 * Strength and damping parameters. Can be modified at any time, but will sometimes have 
+	 * been set once during initialization 
+	 */
+	UPROPERTY(EditAnywhere, BlueprintReadWrite, Category = PhysicsControl)
+	FPhysicsControlData ControlData;
+
+	/**
+	 * Multipliers for the ControlData. These will typically be modified dynamically, and also expose the ability
+	 * to set directional strengths
+	 */
+	UPROPERTY(EditAnywhere, BlueprintReadWrite, Category = PhysicsControl)
+	FPhysicsControlMultipliers ControlMultipliers;
+
+	/**
+	 * The position/orientation etc targets for the controls. These are procedural/explicit control targets -
+	 * skeletal meshes have the option to use skeletal animation as well, in which case these targets are 
+	 * expressed as relative to that animation.
+	 */
+	UPROPERTY(EditAnywhere, BlueprintReadWrite, Category = PhysicsControl)
+	FPhysicsControlTarget ControlTarget;
+
+	/**
+	 * More general settings for the control
+	 */
+	UPROPERTY(EditAnywhere, BlueprintReadWrite, Category = PhysicsControl)
+	FPhysicsControlSettings ControlSettings;
+};
+
diff --git a/Engine/Plugins/Experimental/PhysicsControl/Source/PhysicsControl/Public/PhysicsControlLimbData.h b/Engine/Plugins/Experimental/PhysicsControl/Source/PhysicsControl/Public/PhysicsControlLimbData.h
new file mode 100644
--- /dev/null
+++ b/Engine/Plugins/Experimental/PhysicsControl/Source/PhysicsControl/Public/PhysicsControlLimbData.h
@@ -0,0 +1,72 @@
+// Copyright Epic Games, Inc. All Rights Reserved.
+
+#pragma once
+
+#include "UObject/ObjectMacros.h"
+#include "EngineDefines.h"
+#include "PhysicsControlData.h"
+#include "PhysicsControlLimbData.generated.h"
+
+/**
+ * Setup data that is used to create the representation of a single limb. A limb is an array of 
+ * contiguous bones (e.g. left arm, or the spine etc). We can define it has the set of bones that
+ * are children of a start bone (plus the start bone itself), plus optionally the parent of that start 
+ * bone (this is useful when defining the spine, since you will want to include the pelvis, but you 
+ * don't to include all the children of the pelvis since that would include the legs), but excluding 
+ * any bones that are already part of another limb. This implies limbs should be constructed in order
+ * from leaf to root.
+ */
+USTRUCT(BlueprintType)
+struct PHYSICSCONTROL_API FPhysicsControlLimbSetupData
+{
+	GENERATED_BODY()
+
+	/* The name of the limb that this will be used to create */
+	UPROPERTY(EditAnywhere, BlueprintReadWrite, Category = PhysicsControl)
+	FName LimbName;
+
+	/* 
+	 * Normally the root-most bone of the limb (e.g. left clavicle when defining the left arm) - so the
+	 * limb will contain children of this bone (plus this bone itself).
+	 */
+	UPROPERTY(EditAnywhere, BlueprintReadWrite, Category = PhysicsControl)
+	FName StartBone;
+
+	/* 
+	 * Whether or not to include the parent of the start bone. This is intended to be used for limbs like 
+	 * the spine, where you would set StartBone = spine_01 but also expect to include the pelvis (parent 
+	 * of spine_01) in the spine limb.
+	 */
+	UPROPERTY(EditAnywhere, BlueprintReadWrite, Category = PhysicsControl)
+	bool bIncludeParentBone;
+};
+
+/**
+ * Wrapper for the array of bone names, so that the limbs can we can work with a map of arrays of names. 
+ */
+USTRUCT(BlueprintType)
+struct PHYSICSCONTROL_API FPhysicsControlLimbBones
+{
+	GENERATED_BODY()
+
+	/** The Skeletal mesh that this limb is associated with */
+	TObjectPtr<USkeletalMeshComponent> SkeletalMeshComponent;
+
+	/** The names of the bones in the limb */
+	TArray<FName> BoneNames;
+};
+
+/**
+ * Wrapper for the array of names, so that the limbs can we can work with a map of arrays of names. 
+ * Note that in practice, the names will be of either controls or body modifiers.
+ */
+USTRUCT(BlueprintType)
+struct PHYSICSCONTROL_API FPhysicsControlNameArray
+{
+	GENERATED_BODY()
+
+	/** The names of either controls of body modifiers (depending on context) */
+	UPROPERTY(EditAnywhere, BlueprintReadWrite, Category = PhysicsControl)
+	TArray<FName> Names;
+};
+
diff --git a/Engine/Plugins/Experimental/PhysicsControl/Source/PhysicsControl/Public/PhysicsControlModule.h b/Engine/Plugins/Experimental/PhysicsControl/Source/PhysicsControl/Public/PhysicsControlModule.h
new file mode 100644
--- /dev/null
+++ b/Engine/Plugins/Experimental/PhysicsControl/Source/PhysicsControl/Public/PhysicsControlModule.h
@@ -0,0 +1,14 @@
+// Copyright Epic Games, Inc. All Rights Reserved.
+
+#pragma once
+
+#include "Modules/ModuleManager.h"
+
+class FPhysicsControlModule : public IModuleInterface
+{
+public:
+
+	/** IModuleInterface implementation */
+	virtual void StartupModule() override;
+	virtual void ShutdownModule() override;
+};
diff --git a/Engine/Plugins/Experimental/PhysicsControl/Source/PhysicsControlVisualizer/PhysicsControlVisualizer.Build.cs b/Engine/Plugins/Experimental/PhysicsControl/Source/PhysicsControlVisualizer/PhysicsControlVisualizer.Build.cs
new file mode 100644
--- /dev/null
+++ b/Engine/Plugins/Experimental/PhysicsControl/Source/PhysicsControlVisualizer/PhysicsControlVisualizer.Build.cs
@@ -0,0 +1,53 @@
+// Copyright Epic Games, Inc. All Rights Reserved.
+
+using UnrealBuildTool;
+
+public class PhysicsControlVisualizer : ModuleRules
+{
+	public PhysicsControlVisualizer(ReadOnlyTargetRules Target) : base(Target)
+	{
+		PCHUsage = ModuleRules.PCHUsageMode.UseExplicitOrSharedPCHs;
+		
+		PublicIncludePaths.AddRange(
+			new string[] {
+				// ... add public include paths required here ...
+			}
+			);
+				
+		
+		PrivateIncludePaths.AddRange(
+			new string[] {
+
+				// ... add other private include paths required here ...
+			}
+			);
+			
+		
+		PublicDependencyModuleNames.AddRange(
+			new string[]
+			{
+				"Core"
+				// ... add other public dependencies that you statically link with here ...
+			}
+			);
+			
+		
+		PrivateDependencyModuleNames.AddRange(
+			new string[]
+			{
+				"CoreUObject",
+				"UnrealEd",
+				"ComponentVisualizers",
+				"PhysicsControl"
+			}
+			);
+		
+		
+		DynamicallyLoadedModuleNames.AddRange(
+			new string[]
+			{
+				// ... add any modules that your module loads dynamically here ...
+			}
+			);
+	}
+}
diff --git a/Engine/Plugins/Experimental/PhysicsControl/Source/PhysicsControlVisualizer/Private/PhysicsControlComponentVisualizer.cpp b/Engine/Plugins/Experimental/PhysicsControl/Source/PhysicsControlVisualizer/Private/PhysicsControlComponentVisualizer.cpp
new file mode 100644
--- /dev/null
+++ b/Engine/Plugins/Experimental/PhysicsControl/Source/PhysicsControlVisualizer/Private/PhysicsControlComponentVisualizer.cpp
@@ -0,0 +1,16 @@
+// Copyright Epic Games, Inc. All Rights Reserved.
+
+#include "PhysicsControlComponentVisualizer.h"
+#include "PhysicsControlComponent.h"
+
+//======================================================================================================================
+void FPhysicsControlComponentVisualizer::DrawVisualization(
+	const UActorComponent*   Component,
+	const FSceneView*        View,
+	FPrimitiveDrawInterface* PDI)
+{
+	if (const UPhysicsControlComponent* PCC = Cast<const UPhysicsControlComponent>(Component))
+	{
+		PCC->DebugDraw(PDI);
+	}
+}
diff --git a/Engine/Plugins/Experimental/PhysicsControl/Source/PhysicsControlVisualizer/Private/PhysicsControlComponentVisualizer.h b/Engine/Plugins/Experimental/PhysicsControl/Source/PhysicsControlVisualizer/Private/PhysicsControlComponentVisualizer.h
new file mode 100644
--- /dev/null
+++ b/Engine/Plugins/Experimental/PhysicsControl/Source/PhysicsControlVisualizer/Private/PhysicsControlComponentVisualizer.h
@@ -0,0 +1,13 @@
+// Copyright Epic Games, Inc. All Rights Reserved.
+
+#pragma once
+#include "ComponentVisualizer.h"
+
+class PHYSICSCONTROLVISUALIZER_API FPhysicsControlComponentVisualizer : public FComponentVisualizer
+{
+private:
+	virtual void DrawVisualization(
+		const UActorComponent*   Component, 
+		const FSceneView*        View, 
+		FPrimitiveDrawInterface* PDI) override;
+};
\ No newline at end of file
diff --git a/Engine/Plugins/Experimental/PhysicsControl/Source/PhysicsControlVisualizer/Private/PhysicsControlVisualizerModule.cpp b/Engine/Plugins/Experimental/PhysicsControl/Source/PhysicsControlVisualizer/Private/PhysicsControlVisualizerModule.cpp
new file mode 100644
--- /dev/null
+++ b/Engine/Plugins/Experimental/PhysicsControl/Source/PhysicsControlVisualizer/Private/PhysicsControlVisualizerModule.cpp
@@ -0,0 +1,38 @@
+// Copyright Epic Games, Inc. All Rights Reserved.
+
+#include "PhysicsControlVisualizerModule.h"
+#include "PhysicsControlComponent.h"
+#include "PhysicsControlComponentVisualizer.h"
+#include "UnrealEdGlobals.h"
+#include "Editor/UnrealEdEngine.h"
+
+#define LOCTEXT_NAMESPACE "FPhysicsControlModule"
+
+void FPhysicsControlVisualizerModule::StartupModule()
+{
+	if (GUnrealEd)
+	{
+		TSharedPtr<FPhysicsControlComponentVisualizer> Visualizer = MakeShared<FPhysicsControlComponentVisualizer>();
+		GUnrealEd->RegisterComponentVisualizer(UPhysicsControlComponent::StaticClass()->GetFName(), Visualizer);
+		// This call should maybe be inside the RegisterComponentVisualizer call above, but since it's not,
+		// we'll put it here.
+		Visualizer->OnRegister();
+		VisualizersToUnregisterOnShutdown.Add(UPhysicsControlComponent::StaticClass()->GetFName());
+	}
+}
+
+void FPhysicsControlVisualizerModule::ShutdownModule()
+{
+	if (GUnrealEd)
+	{
+		for (const FName& Name : VisualizersToUnregisterOnShutdown)
+		{
+			GUnrealEd->UnregisterComponentVisualizer(Name);
+		}
+	}
+}
+
+#undef LOCTEXT_NAMESPACE
+	
+IMPLEMENT_MODULE(FPhysicsControlVisualizerModule, PhysicsControl)
+
diff --git a/Engine/Plugins/Experimental/PhysicsControl/Source/PhysicsControlVisualizer/Public/PhysicsControlVisualizerModule.h b/Engine/Plugins/Experimental/PhysicsControl/Source/PhysicsControlVisualizer/Public/PhysicsControlVisualizerModule.h
new file mode 100644
--- /dev/null
+++ b/Engine/Plugins/Experimental/PhysicsControl/Source/PhysicsControlVisualizer/Public/PhysicsControlVisualizerModule.h
@@ -0,0 +1,18 @@
+// Copyright Epic Games, Inc. All Rights Reserved.
+
+#pragma once
+
+#include "CoreMinimal.h"
+#include "Modules/ModuleManager.h"
+
+class FPhysicsControlVisualizerModule : public IModuleInterface
+{
+public:
+
+	/** IModuleInterface implementation */
+	virtual void StartupModule() override;
+	virtual void ShutdownModule() override;
+
+private:
+	TArray<FName> VisualizersToUnregisterOnShutdown;
+};
