Physics Control Component: Use explicit targets for kinematic control. Complete functionality for creating/manipulating sets of controls.Initialise from constraint profiles. Restore body state on Modifier deletetion.

diff --git a/Engine/Plugins/Experimental/PhysicsControl/Source/PhysicsControl/Private/PhysicsControlComponent.cpp b/Engine/Plugins/Experimental/PhysicsControl/Source/PhysicsControl/Private/PhysicsControlComponent.cpp
--- a/Engine/Plugins/Experimental/PhysicsControl/Source/PhysicsControl/Private/PhysicsControlComponent.cpp
+++ b/Engine/Plugins/Experimental/PhysicsControl/Source/PhysicsControl/Private/PhysicsControlComponent.cpp
@@ -45,10 +45,18 @@ void UPhysicsControlComponent::BeginDestroy()
 {
 	for (TPair<FName, FPhysicsControlRecord>& PhysicsControlRecordPair : Implementation->PhysicsControlRecords)
 	{
+		DestroyControl(PhysicsControlRecordPair.Key);
 		FPhysicsControlRecord& Record = PhysicsControlRecordPair.Value;
 		Record.PhysicsControlState.Reset();
 	}
 	Implementation->PhysicsControlRecords.Empty();
+
+	for (TPair<FName, FPhysicsBodyModifier>& PhysicsBodyModifierPair : Implementation->PhysicsBodyModifiers)
+	{
+		DestroyBodyModifier(PhysicsBodyModifierPair.Key);
+	}
+	Implementation->PhysicsBodyModifiers.Empty();
+
 	Super::BeginDestroy();
 }
 
@@ -124,16 +132,24 @@ void UPhysicsControlComponent::TickComponent(
 	for (TPair<FName, FPhysicsBodyModifier>& BodyModifierPair : Implementation->PhysicsBodyModifiers)
 	{
 		FPhysicsBodyModifier& BodyModifier = BodyModifierPair.Value;
-		FBodyInstance* BodyInstance = GetBodyInstance(BodyModifier.MeshComponent, BodyModifier.BoneName);
-		// Note that there is a Default type, which means we don't change anything
-		if (BodyModifier.MovementType == EPhysicsMovementType::Simulated)
-		{
-			BodyInstance->SetInstanceSimulatePhysics(true);
-		}
-		else if (BodyModifier.MovementType == EPhysicsMovementType::Kinematic)
+		FBodyInstance* BodyInstance = UE::PhysicsControlComponent::GetBodyInstance(
+			BodyModifier.MeshComponent, BodyModifier.BoneName);
+
+		switch (BodyModifier.MovementType)
 		{
+		case EPhysicsMovementType::Static:
+			BodyInstance->SetInstanceSimulatePhysics(false, true);
+			break;
+		case EPhysicsMovementType::Kinematic:
 			BodyInstance->SetInstanceSimulatePhysics(false, true);
 			Implementation->ApplyKinematicTarget(BodyModifier);
+			break;
+		case EPhysicsMovementType::Simulated:
+			BodyInstance->SetInstanceSimulatePhysics(true);
+			break;
+		default:
+			UE_LOG(LogPhysicsControlComponent, Warning, TEXT("Invalid movement type %d"), BodyModifier.MovementType);
+			break;
 		}
 
 		if (BodyInstance->IsInstanceSimulatingPhysics())
@@ -182,14 +198,20 @@ TMap<FName, FPhysicsControlLimbBones> UPhysicsControlComponent::GetLimbBonesFrom
 
 		if (LimbSetupData.bIncludeParentBone)
 		{
+			LimbBones.bFirstBoneIsAdditional = true;
 			FName ParentBoneName;
-			ParentBoneName = GetPhysicalParentBone(SkeletalMeshComponent, LimbSetupData.StartBone);
+			ParentBoneName = UE::PhysicsControlComponent::GetPhysicalParentBone(
+				SkeletalMeshComponent, LimbSetupData.StartBone);
 			if (!ParentBoneName.IsNone())
 			{
 				LimbBones.BoneNames.Add(ParentBoneName);
 				AllBones.Add(ParentBoneName);
 			}
 		}
+		else
+		{
+			LimbBones.bFirstBoneIsAdditional = false;
+		}
 
 		SkeletalMeshComponent->ForEachBodyBelow(
 			LimbSetupData.StartBone, true, /*bSkipCustomType=*/false,
@@ -257,13 +279,13 @@ bool UPhysicsControlComponent::MakeNamedControl(
 		return false;
 	}
 
-	if (ParentMeshComponent && ParentMeshComponent->IsA<USkeletalMeshComponent>())
+	if (USkeletalMeshComponent* SkeletalMeshComponent = Cast<USkeletalMeshComponent>(ParentMeshComponent))
 	{
-		Implementation->AddSkeletalMeshReference(Cast<USkeletalMeshComponent>(ParentMeshComponent));
+		Implementation->AddSkeletalMeshReferenceForCaching(SkeletalMeshComponent);
 	}
-	if (ChildMeshComponent && ChildMeshComponent->IsA<USkeletalMeshComponent>())
+	if (USkeletalMeshComponent* SkeletalMeshComponent = Cast<USkeletalMeshComponent>(ChildMeshComponent))
 	{
-		Implementation->AddSkeletalMeshReference(Cast<USkeletalMeshComponent>(ChildMeshComponent));
+		Implementation->AddSkeletalMeshReferenceForCaching(SkeletalMeshComponent);
 	}
 
 	FPhysicsControlRecord& NewRecord = Implementation->PhysicsControlRecords.Add(
@@ -313,7 +335,8 @@ TArray<FName> UPhysicsControlComponent::MakeControlsFromSkeletalMeshBelow(
 				FName ParentBoneName;
 				if (ParentMeshComponent)
 				{
-					ParentBoneName = GetPhysicalParentBone(SkeletalMeshComponent, ChildBoneName);
+					ParentBoneName = UE::PhysicsControlComponent::GetPhysicalParentBone(
+						SkeletalMeshComponent, ChildBoneName);
 					if (ParentBoneName.IsNone())
 					{
 						return;
@@ -337,6 +360,76 @@ TArray<FName> UPhysicsControlComponent::MakeControlsFromSkeletalMeshBelow(
 	return Result;
 }
 
+//======================================================================================================================
+// Note - some params passed by value to allow them to be unconnected and enable inline BP editing
+TArray<FName> UPhysicsControlComponent::MakeControlsFromSkeletalMeshAndConstraintProfileBelow(
+	USkeletalMeshComponent* SkeletalMeshComponent,
+	FName                   BoneName,
+	bool                    bIncludeSelf,
+	FName                   ConstraintProfile,
+	bool                    bEnabled)
+{
+	TArray<FName> Result;
+	UPhysicsAsset* PhysicsAsset = SkeletalMeshComponent ? SkeletalMeshComponent->GetPhysicsAsset() : nullptr;
+	if (!PhysicsAsset)
+	{
+		UE_LOG(LogPhysicsControlComponent, Warning, TEXT("No physics asset in skeletal mesh"));
+		return Result;
+	}
+
+	FPhysicsControlSettings ControlSettings;
+	// This is to match the skeletal mesh component velocity drive, which does not use the
+	// target animation velocity.
+	ControlSettings.SkeletalAnimationVelocityMultiplier = 0;
+
+	SkeletalMeshComponent->ForEachBodyBelow(
+		BoneName, bIncludeSelf, /*bSkipCustomType=*/false,
+		[
+			this, PhysicsAsset, SkeletalMeshComponent,
+			ConstraintProfile, &ControlSettings, &Result, bEnabled
+		](const FBodyInstance* BI)
+		{
+			if (USkeletalBodySetup* BodySetup = Cast<USkeletalBodySetup>(BI->BodySetup.Get()))
+			{
+				const FName ChildBoneName = PhysicsAsset->SkeletalBodySetups[BI->InstanceBodyIndex]->BoneName;
+
+				FName ParentBoneName;
+				ParentBoneName = UE::PhysicsControlComponent::GetPhysicalParentBone(
+					SkeletalMeshComponent, ChildBoneName);
+				if (ParentBoneName.IsNone())
+				{
+					return;
+				}
+
+				FPhysicsControlData ControlData;
+				FConstraintProfileProperties ProfileProperties;
+				if (!SkeletalMeshComponent->GetConstraintProfilePropertiesOrDefault(
+					ProfileProperties, ChildBoneName, ConstraintProfile))
+				{
+					UE_LOG(LogPhysicsControlComponent, Warning, TEXT("Failed get constraint profile for %s"), *ChildBoneName.ToString());
+					return;
+				}
+
+				UE::PhysicsControlComponent::ConvertConstraintProfileToControlData(ControlData, ProfileProperties);
+
+				FName ControlName = MakeControl(
+					SkeletalMeshComponent, ParentBoneName, SkeletalMeshComponent, ChildBoneName,
+					ControlData, FPhysicsControlTarget(), ControlSettings, bEnabled);
+				if (!ControlName.IsNone())
+				{
+					Result.Add(ControlName);
+				}
+				else
+				{
+					UE_LOG(LogPhysicsControlComponent, Warning,
+						TEXT("Failed to make control for %s"), *ChildBoneName.ToString());
+				}
+			}
+		});
+
+	return Result;
+}
+
 //======================================================================================================================
 // Note - some params passed by value to allow them to be unconnected and enable inline BP editing
 TArray<FName> UPhysicsControlComponent::MakeControlsFromSkeletalMesh(
@@ -365,7 +458,8 @@ TArray<FName> UPhysicsControlComponent::MakeControlsFromSkeletalMesh(
 		FName ParentBoneName;
 		if (ParentMeshComponent)
 		{
-			ParentBoneName = GetPhysicalParentBone(SkeletalMeshComponent, ChildBoneName);
+			ParentBoneName = UE::PhysicsControlComponent::GetPhysicalParentBone(
+				SkeletalMeshComponent, ChildBoneName);
 			if (ParentBoneName.IsNone())
 			{
 				continue;
@@ -387,6 +481,65 @@ TArray<FName> UPhysicsControlComponent::MakeControlsFromSkeletalMesh(
 	return Result;
 }
 
+//======================================================================================================================
+// Note - some params passed by value to allow them to be unconnected and enable inline BP editing
+TArray<FName> UPhysicsControlComponent::MakeControlsFromSkeletalMeshAndConstraintProfile(
+	USkeletalMeshComponent* SkeletalMeshComponent,
+	const TArray<FName>&    BoneNames,
+	FName                   ConstraintProfile,
+	bool                    bEnabled)
+{
+	TArray<FName> Result;
+	UPhysicsAsset* PhysicsAsset = SkeletalMeshComponent ? SkeletalMeshComponent->GetPhysicsAsset() : nullptr;
+	if (!PhysicsAsset)
+	{
+		UE_LOG(LogPhysicsControlComponent, Warning, TEXT("No physics asset in skeletal mesh"));
+		return Result;
+	}
+
+	FPhysicsControlSettings ControlSettings;
+	// This is to match the skeletal mesh component velocity drive, which does not use the
+	// target animation velocity.
+	ControlSettings.SkeletalAnimationVelocityMultiplier = 0;
+
+	for (FName ChildBoneName : BoneNames)
+	{
+		FBodyInstance* ChildBone = SkeletalMeshComponent->GetBodyInstance(ChildBoneName);
+
+		FName ParentBoneName;
+		ParentBoneName = UE::PhysicsControlComponent::GetPhysicalParentBone(SkeletalMeshComponent, ChildBoneName);
+		if (ParentBoneName.IsNone())
+		{
+			continue;
+		}
+
+		FPhysicsControlData ControlData;
+		FConstraintProfileProperties ProfileProperties;
+		if (!SkeletalMeshComponent->GetConstraintProfilePropertiesOrDefault(
+			ProfileProperties, ChildBoneName, ConstraintProfile))
+		{
+			UE_LOG(LogPhysicsControlComponent, Warning, TEXT("Failed get constraint profile for %s"), *ChildBoneName.ToString());
+			continue;
+		}
+
+		UE::PhysicsControlComponent::ConvertConstraintProfileToControlData(ControlData, ProfileProperties);
+
+		FName ControlName = MakeControl(
+			SkeletalMeshComponent, ParentBoneName, SkeletalMeshComponent, ChildBoneName,
+			ControlData, FPhysicsControlTarget(), ControlSettings, bEnabled);
+		if (!ControlName.IsNone())
+		{
+			Result.Add(ControlName);
+		}
+		else
+		{
+			UE_LOG(LogPhysicsControlComponent, Warning, TEXT("Failed to make control for %s"), *ChildBoneName.ToString());
+		}
+	}
+
+	return Result;
+}
+
 //======================================================================================================================
 // Note - some params passed by value to allow them to be unconnected and enable inline BP editing
 TMap<FName, FPhysicsControlNameArray> UPhysicsControlComponent::MakeControlsFromLimbBones(
@@ -420,14 +573,22 @@ TMap<FName, FPhysicsControlNameArray> UPhysicsControlComponent::MakeControlsFrom
 		LimbResult.Names.Reserve(NumBonesInLimb);
 		AllControls.Names.Reserve(AllControls.Names.Num() + NumBonesInLimb);
 
-		for (FName ChildBoneName : BonesInLimb.BoneNames)
+		for (int32 BoneIndex = 0 ; BoneIndex != NumBonesInLimb ; ++BoneIndex)
 		{
+			// Don't create the parent space control if it's the first bone in a limb that had bIncludeParentBone
+			if (BoneIndex == 0 && BonesInLimb.bFirstBoneIsAdditional && ControlType == EPhysicsControlType::ParentSpace)
+			{
+				continue;
+			}
+
+			FName ChildBoneName = BonesInLimb.BoneNames[BoneIndex];
 			FBodyInstance* ChildBone = BonesInLimb.SkeletalMeshComponent->GetBodyInstance(ChildBoneName);
 
 			FName ParentBoneName;
 			if (ParentMeshComponent)
 			{
-				ParentBoneName = GetPhysicalParentBone(ParentMeshComponent, ChildBoneName);
+				ParentBoneName = UE::PhysicsControlComponent::GetPhysicalParentBone(
+					ParentMeshComponent, ChildBoneName);
 				if (ParentBoneName.IsNone())
 				{
 					continue;
@@ -450,6 +611,88 @@ TMap<FName, FPhysicsControlNameArray> UPhysicsControlComponent::MakeControlsFrom
 	return Result;
 }
 
+//======================================================================================================================
+TMap<FName, FPhysicsControlNameArray> UPhysicsControlComponent::MakeControlsFromLimbBonesAndConstraintProfile(
+	FPhysicsControlNameArray&                    AllControls,
+	const TMap<FName, FPhysicsControlLimbBones>& LimbBones,
+	FName                                        ConstraintProfile,
+	bool                                         bEnabled)
+{
+	TMap<FName, FPhysicsControlNameArray> Result;
+	Result.Reserve(LimbBones.Num());
+	for (const TPair< FName, FPhysicsControlLimbBones>& LimbBoneEntry : LimbBones)
+	{
+		FName LimbName = LimbBoneEntry.Key;
+		const FPhysicsControlLimbBones& BonesInLimb = LimbBoneEntry.Value;
+
+		USkeletalMeshComponent* SkeletalMeshComponent = BonesInLimb.SkeletalMeshComponent;
+		if (!SkeletalMeshComponent)
+		{
+			UE_LOG(LogPhysicsControlComponent, Warning, TEXT("No Skeletal mesh in limb %s"), *LimbName.ToString());
+			continue;
+		}
+		UPhysicsAsset* PhysicsAsset = SkeletalMeshComponent ? SkeletalMeshComponent->GetPhysicsAsset() : nullptr;
+		if (!PhysicsAsset)
+		{
+			UE_LOG(LogPhysicsControlComponent, Warning, TEXT("No physics asset in skeletal mesh"));
+			return Result;
+		}
+
+		const int32 NumBonesInLimb = BonesInLimb.BoneNames.Num();
+
+		FPhysicsControlNameArray& LimbResult = Result.Add(LimbName);
+		LimbResult.Names.Reserve(NumBonesInLimb);
+		AllControls.Names.Reserve(AllControls.Names.Num() + NumBonesInLimb);
+
+		FPhysicsControlSettings ControlSettings;
+		// This is to match the skeletal mesh component velocity drive, which does not use the
+		// target animation velocity.
+		ControlSettings.SkeletalAnimationVelocityMultiplier = 0;
+
+		for (int32 BoneIndex = 0; BoneIndex != NumBonesInLimb; ++BoneIndex)
+		{
+			// Don't create the parent space control if it's the first bone in a limb that had bIncludeParentBone
+			if (BoneIndex == 0 && BonesInLimb.bFirstBoneIsAdditional)
+			{
+				continue;
+			}
+
+			FName ChildBoneName = BonesInLimb.BoneNames[BoneIndex];
+			FBodyInstance* ChildBone = SkeletalMeshComponent->GetBodyInstance(ChildBoneName);
+			FName ParentBoneName;
+			ParentBoneName = UE::PhysicsControlComponent::GetPhysicalParentBone(SkeletalMeshComponent, ChildBoneName);
+			if (ParentBoneName.IsNone())
+			{
+				continue;
+			}
+
+			FPhysicsControlData ControlData;
+			FConstraintProfileProperties ProfileProperties;
+			if (!SkeletalMeshComponent->GetConstraintProfilePropertiesOrDefault(
+				ProfileProperties, ChildBoneName, ConstraintProfile))
+			{
+				UE_LOG(LogPhysicsControlComponent, Warning, TEXT("Failed get constraint profile for %s"), *ChildBoneName.ToString());
+				continue;
+			}
+
+			UE::PhysicsControlComponent::ConvertConstraintProfileToControlData(ControlData, ProfileProperties);
+
+			FName ControlName = MakeControl(
+				SkeletalMeshComponent, ParentBoneName, SkeletalMeshComponent, ChildBoneName,
+				ControlData, FPhysicsControlTarget(), ControlSettings, bEnabled);
+			if (!ControlName.IsNone())
+			{
+				LimbResult.Names.Add(ControlName);
+				AllControls.Names.Add(ControlName);
+			}
+			else
+			{
+				UE_LOG(LogPhysicsControlComponent, Warning, TEXT("Failed to make control for %s"), *ChildBoneName.ToString());
+			}
+		}
+	}
+	return Result;
+}
 
 //======================================================================================================================
 bool UPhysicsControlComponent::DestroyControl(const FName Name)
@@ -457,15 +700,15 @@ bool UPhysicsControlComponent::DestroyControl(const FName Name)
 	FPhysicsControlRecord* PhysicsControlRecord = Implementation->FindControlRecord(Name);
 	if (PhysicsControlRecord)
 	{
-		if (PhysicsControlRecord->PhysicsControl.ParentMeshComponent && 
-			PhysicsControlRecord->PhysicsControl.ParentMeshComponent->IsA<USkeletalMeshComponent>())
+		if (USkeletalMeshComponent* SkeletalMeshComponent = 
+			Cast<USkeletalMeshComponent>(PhysicsControlRecord->PhysicsControl.ParentMeshComponent))
 		{
-			Implementation->RemoveSkeletalMeshReference(Cast<USkeletalMeshComponent>(PhysicsControlRecord->PhysicsControl.ParentMeshComponent));
+			Implementation->RemoveSkeletalMeshReferenceForCaching(SkeletalMeshComponent);
 		}
-		if (PhysicsControlRecord->PhysicsControl.ChildMeshComponent &&
-			PhysicsControlRecord->PhysicsControl.ChildMeshComponent->IsA<USkeletalMeshComponent>())
+		if (USkeletalMeshComponent* SkeletalMeshComponent =
+			Cast<USkeletalMeshComponent>(PhysicsControlRecord->PhysicsControl.ChildMeshComponent))
 		{
-			Implementation->RemoveSkeletalMeshReference(Cast<USkeletalMeshComponent>(PhysicsControlRecord->PhysicsControl.ChildMeshComponent));
+			Implementation->RemoveSkeletalMeshReferenceForCaching(SkeletalMeshComponent);
 		}
 
 		PhysicsControlRecord->PhysicsControlState.bPendingDestroy = true;
@@ -475,6 +718,15 @@ bool UPhysicsControlComponent::DestroyControl(const FName Name)
 	return false;
 }
 
+//======================================================================================================================
+void UPhysicsControlComponent::DestroyAllControls(const TArray<FName>& Names)
+{
+	for (FName Name : Names)
+	{
+		DestroyControl(Name);
+	}
+}
+
 //======================================================================================================================
 bool UPhysicsControlComponent::SetControlEnabled(const FName Name, bool bEnable)
 {
@@ -487,12 +739,21 @@ bool UPhysicsControlComponent::SetControlEnabled(const FName Name, bool bEnable)
 	return false;
 }
 
+//======================================================================================================================
+void UPhysicsControlComponent::SetAllControlsEnabled(const TArray<FName>& Names, bool bEnable)
+{
+	for (FName Name : Names)
+	{
+		SetControlEnabled(Name, bEnable);
+	}
+}
+
 //======================================================================================================================
 // Note - some params passed by value to allow them to be unconnected and enable inline BP editing
 bool UPhysicsControlComponent::SetControlData(
-	const FName             Name, 
-	FPhysicsControlData     ControlData, 
-	bool                    bEnableControl)
+	const FName         Name, 
+	FPhysicsControlData ControlData, 
+	bool                bEnableControl)
 {
 	FPhysicsControlRecord* Record = Implementation->FindControlRecord(Name);
 	if (Record)
@@ -504,6 +765,19 @@ bool UPhysicsControlComponent::SetControlData(
 	return false;
 }
 
+//======================================================================================================================
+void UPhysicsControlComponent::SetAllControlDatas(
+	const TArray<FName>& Names, 
+	FPhysicsControlData  ControlData, 
+	bool                 bEnableControl)
+{
+	for (FName Name : Names)
+	{
+		SetControlData(Name, ControlData, bEnableControl);
+	}
+}
+
+
 //======================================================================================================================
 // Note - some params passed by value to allow them to be unconnected and enable inline BP editing
 bool UPhysicsControlComponent::SetControlMultipliers(
@@ -521,6 +795,19 @@ bool UPhysicsControlComponent::SetControlMultipliers(
 	return false;
 }
 
+//======================================================================================================================
+void UPhysicsControlComponent::SetAllControlMultipliers(
+	const TArray<FName>&       Names,
+	FPhysicsControlMultipliers ControlMultipliers, 
+	bool                       bEnableControl)
+{
+	for (FName Name : Names)
+	{
+		SetControlMultipliers(Name, ControlMultipliers, bEnableControl);
+	}
+}
+
+
 //======================================================================================================================
 bool UPhysicsControlComponent::SetControlLinearData(
 	const FName Name, float Strength, float DampingRatio, float ExtraDamping, float MaxForce, bool bEnableControl)
@@ -596,6 +883,18 @@ bool UPhysicsControlComponent::SetControlTarget(
 	return false;
 }
 
+//======================================================================================================================
+void UPhysicsControlComponent::SetAllControlTargets(
+	const TArray<FName>&  Names, 
+	FPhysicsControlTarget ControlTarget, 
+	bool                  bEnableControl)
+{
+	for (FName Name : Names)
+	{
+		SetControlTarget(Name, ControlTarget, bEnableControl);
+	}
+}
+
 //======================================================================================================================
 bool UPhysicsControlComponent::SetControlTargetTransform(
 	const FName Name, const FTransform Transform, float VelocityDeltaTime, 
@@ -727,6 +1026,18 @@ bool UPhysicsControlComponent::SetControlUseSkeletalAnimation(
 	return false;
 }
 
+//======================================================================================================================
+void UPhysicsControlComponent::SetAllControlsUseSkeletalAnimation(
+	const TArray<FName>& Names,
+	bool                 bUseSkeletalAnimation,
+	float                SkeletalAnimationVelocityMultiplier)
+{
+	for (FName Name : Names)
+	{
+		SetControlUseSkeletalAnimation(Name, bUseSkeletalAnimation, SkeletalAnimationVelocityMultiplier);
+	}
+}
+
 //======================================================================================================================
 bool UPhysicsControlComponent::SetControlAutoDisable(const FName Name, bool bAutoDisable)
 {
@@ -811,13 +1122,14 @@ bool UPhysicsControlComponent::GetControlEnabled(const FName Name) const
 
 //======================================================================================================================
 FName UPhysicsControlComponent::MakeBodyModifier(
-	UMeshComponent*       MeshComponent,
-	const FName           BoneName,
-	EPhysicsMovementType  MovementType, 
-	float                 GravityMultiplier)
+	UMeshComponent*      MeshComponent,
+	const FName          BoneName,
+	EPhysicsMovementType MovementType, 
+	float                GravityMultiplier,
+	bool                 bUseSkeletalAnimation)
 {
 	FName Name = Implementation->GetUniqueBodyModifierName(BoneName);
-	if (MakeNamedBodyModifier(Name, MeshComponent, BoneName, MovementType, GravityMultiplier))
+	if (MakeNamedBodyModifier(Name, MeshComponent, BoneName, MovementType, GravityMultiplier, bUseSkeletalAnimation))
 	{
 		return Name;
 	}
@@ -826,11 +1138,12 @@ FName UPhysicsControlComponent::MakeBodyModifier(
 
 //======================================================================================================================
 bool UPhysicsControlComponent::MakeNamedBodyModifier(
-	const FName           Name,
-	UMeshComponent*       MeshComponent,
-	const FName           BoneName,
-	EPhysicsMovementType  MovementType,
-	float                 GravityMultiplier)
+	const FName          Name,
+	UMeshComponent*      MeshComponent,
+	const FName          BoneName,
+	EPhysicsMovementType MovementType,
+	float                GravityMultiplier,
+	bool                 bUseSkeletalAnimation)
 {
 	FPhysicsBodyModifier* BodyModifier = Implementation->FindBodyModifier(Name);
 	if (BodyModifier)
@@ -845,26 +1158,27 @@ bool UPhysicsControlComponent::MakeNamedBodyModifier(
 		return false;
 	}
 
+	FPhysicsBodyModifier& Modifier = Implementation->PhysicsBodyModifiers.Add(
+		Name, FPhysicsBodyModifier(MeshComponent, BoneName, MovementType, GravityMultiplier, bUseSkeletalAnimation));
+
 	USkeletalMeshComponent* SkeletalMeshComponent = Cast<USkeletalMeshComponent>(MeshComponent);
 	if (SkeletalMeshComponent)
 	{
-		Implementation->AddSkeletalMeshReference(SkeletalMeshComponent);
-		SkeletalMeshComponent->bUpdateMeshWhenKinematic = true;
+		Implementation->AddSkeletalMeshReferenceForCaching(SkeletalMeshComponent);
+		Implementation->AddSkeletalMeshReferenceForModifier(SkeletalMeshComponent);
 	}
 
-	Implementation->PhysicsBodyModifiers.Add(
-		Name, FPhysicsBodyModifier(MeshComponent, BoneName, MovementType, GravityMultiplier));
-
 	return true;
 }
 
 //======================================================================================================================
 TArray<FName> UPhysicsControlComponent::MakeBodyModifiersFromSkeletalMeshBelow(
 	USkeletalMeshComponent* SkeletalMeshComponent,
-	FName                   BoneName,
-	bool                    bIncludeSelf,
-	EPhysicsMovementType    MovementType,
-	float                   GravityMultiplier)
+	FName                BoneName,
+	bool                 bIncludeSelf,
+	EPhysicsMovementType MovementType,
+	float                GravityMultiplier,
+	bool                 bUseSkeletalAnimation)
 {
 	TArray<FName> Result;
 	UPhysicsAsset* PhysicsAsset = SkeletalMeshComponent ? SkeletalMeshComponent->GetPhysicsAsset() : nullptr;
@@ -875,13 +1189,16 @@ TArray<FName> UPhysicsControlComponent::MakeBodyModifiersFromSkeletalMeshBelow(
 
 	SkeletalMeshComponent->ForEachBodyBelow(
 		BoneName, bIncludeSelf, /*bSkipCustomType=*/false,
-		[this, PhysicsAsset, SkeletalMeshComponent, MovementType, GravityMultiplier, &Result](const FBodyInstance* BI)
+		[
+			this, PhysicsAsset, SkeletalMeshComponent, MovementType, 
+			GravityMultiplier, bUseSkeletalAnimation, &Result
+		](const FBodyInstance* BI)
 		{
 			if (USkeletalBodySetup* BodySetup = Cast<USkeletalBodySetup>(BI->BodySetup.Get()))
 			{
 				const FName BoneName = PhysicsAsset->SkeletalBodySetups[BI->InstanceBodyIndex]->BoneName;
 				FName BodyModifierName = MakeBodyModifier(
-					SkeletalMeshComponent, BoneName, MovementType, GravityMultiplier);
+					SkeletalMeshComponent, BoneName, MovementType, GravityMultiplier, bUseSkeletalAnimation);
 				Result.Add(BodyModifierName);
 			}
 		});
@@ -894,7 +1211,8 @@ TMap<FName, FPhysicsControlNameArray> UPhysicsControlComponent::MakeBodyModifier
 	FPhysicsControlNameArray&                    AllBodyModifiers,
 	const TMap<FName, FPhysicsControlLimbBones>& LimbBones,
 	EPhysicsMovementType                         MovementType,
-	float                                        GravityMultiplier)
+	float                                        GravityMultiplier,
+	bool                                         bUseSkeletalAnimation)
 {
 	TMap<FName, FPhysicsControlNameArray> Result;
 	Result.Reserve(LimbBones.Num());
@@ -919,7 +1237,7 @@ TMap<FName, FPhysicsControlNameArray> UPhysicsControlComponent::MakeBodyModifier
 		for (FName BoneName : BonesInLimb.BoneNames)
 		{
 			FName BodyModifierName = MakeBodyModifier(
-				BonesInLimb.SkeletalMeshComponent, BoneName, MovementType, GravityMultiplier);
+				BonesInLimb.SkeletalMeshComponent, BoneName, MovementType, GravityMultiplier, bUseSkeletalAnimation);
 			if (!BodyModifierName.IsNone())
 			{
 				LimbResult.Names.Add(BodyModifierName);
@@ -940,10 +1258,11 @@ bool UPhysicsControlComponent::DestroyBodyModifier(const FName Name)
 	FPhysicsBodyModifier* PhysicsBodyModifier = Implementation->FindBodyModifier(Name);
 	if (PhysicsBodyModifier)
 	{
-		if (PhysicsBodyModifier->MeshComponent &&
-			PhysicsBodyModifier->MeshComponent->IsA<USkeletalMeshComponent>())
+		if (USkeletalMeshComponent* SkeletalMeshComponent = 
+			Cast<USkeletalMeshComponent>(PhysicsBodyModifier->MeshComponent))
 		{
-			Implementation->RemoveSkeletalMeshReference(Cast<USkeletalMeshComponent>(PhysicsBodyModifier->MeshComponent));
+			Implementation->RemoveSkeletalMeshReferenceForCaching(SkeletalMeshComponent);
+			Implementation->RemoveSkeletalMeshReferenceForModifier(SkeletalMeshComponent);
 		}
 
 		PhysicsBodyModifier->bPendingDestroy = true;
@@ -953,19 +1272,111 @@ bool UPhysicsControlComponent::DestroyBodyModifier(const FName Name)
 }
 
 //======================================================================================================================
-bool UPhysicsControlComponent::SetBodyModifier(
-	const FName Name, EPhysicsMovementType MovementType, float GravityMultiplier)
+void UPhysicsControlComponent::DestroyAllBodyModifiers(const TArray<FName>& Names)
+{
+	for (FName Name : Names)
+	{
+		DestroyBodyModifier(Name);
+	}
+}
+
+//======================================================================================================================
+bool UPhysicsControlComponent::SetBodyModifierKinematicTarget(
+	const FName Name,
+	FVector     KinematicTargetPosition,
+	FRotator    KinematicTargetOrienation,
+	bool        bMakeKinematic)
+{
+	FPhysicsBodyModifier* PhysicsBodyModifier = Implementation->FindBodyModifier(Name);
+	if (PhysicsBodyModifier)
+	{
+		PhysicsBodyModifier->KinematicTargetPosition = KinematicTargetPosition;
+		PhysicsBodyModifier->KinematicTargetOrientation = KinematicTargetOrienation.Quaternion();
+		if (bMakeKinematic)
+		{
+			PhysicsBodyModifier->MovementType = EPhysicsMovementType::Kinematic;
+		}
+		return true;
+	}
+	return false;
+}
+
+//======================================================================================================================
+bool UPhysicsControlComponent::SetBodyModifierMovementType(
+		const FName          Name,
+		EPhysicsMovementType MovementType)
 {
 	FPhysicsBodyModifier* PhysicsBodyModifier = Implementation->FindBodyModifier(Name);
 	if (PhysicsBodyModifier)
 	{
 		PhysicsBodyModifier->MovementType = MovementType;
+		return true;
+	}
+	return false;
+}
+
+//======================================================================================================================
+void UPhysicsControlComponent::SetAllBodyModifierMovementType(
+	const TArray<FName>& Names,
+	EPhysicsMovementType MovementType)
+{
+	for (FName Name : Names)
+	{
+		SetBodyModifierMovementType(Name, MovementType);
+	}
+}
+
+
+//======================================================================================================================
+bool UPhysicsControlComponent::SetBodyModifierGravityMultiplier(
+	const FName Name,
+	float       GravityMultiplier)
+{
+	FPhysicsBodyModifier* PhysicsBodyModifier = Implementation->FindBodyModifier(Name);
+	if (PhysicsBodyModifier)
+	{
 		PhysicsBodyModifier->GravityMultiplier = GravityMultiplier;
 		return true;
 	}
 	return false;
 }
 
+//======================================================================================================================
+void UPhysicsControlComponent::SetAllBodyModifierGravityMultipliers(
+	const TArray<FName>& Names,
+	float                GravityMultiplier)
+{
+	for (FName Name : Names)
+	{
+		SetBodyModifierGravityMultiplier(Name, GravityMultiplier);
+	}
+}
+
+//======================================================================================================================
+bool UPhysicsControlComponent::SetBodyModifierUseSkeletalAnimation(
+	const FName Name,
+	bool        bUseSkeletalAnimation)
+{
+	FPhysicsBodyModifier* PhysicsBodyModifier = Implementation->FindBodyModifier(Name);
+	if (PhysicsBodyModifier)
+	{
+		PhysicsBodyModifier->bUseSkeletalAnimation = bUseSkeletalAnimation;
+		return true;
+	}
+	return false;
+}
+
+//======================================================================================================================
+void UPhysicsControlComponent::SetAllBodyModifiersUseSkeletalAnimation(
+	const TArray<FName>& Names,
+	bool                 bUseSkeletalAnimation)
+{
+	for (FName Name : Names)
+	{
+		SetBodyModifierUseSkeletalAnimation(Name, bUseSkeletalAnimation);
+	}
+}
+
 #if WITH_EDITOR
 //======================================================================================================================
 void UPhysicsControlComponent::DebugDraw(FPrimitiveDrawInterface* PDI) const
@@ -998,7 +1409,7 @@ void UPhysicsControlComponent::DebugDrawControl(
 
 	if (Record.PhysicsControlState.bEnabled && ConstraintInstance)
 	{
-		FBodyInstance* ChildBodyInstance = GetBodyInstance(
+		FBodyInstance* ChildBodyInstance = UE::PhysicsControlComponent::GetBodyInstance(
 			Record.PhysicsControl.ChildMeshComponent, Record.PhysicsControl.ChildBoneName);
 		if (!ChildBodyInstance)
 		{
@@ -1006,7 +1417,7 @@ void UPhysicsControlComponent::DebugDrawControl(
 		}
 		FTransform ChildBodyTM = ChildBodyInstance->GetUnrealWorldTransform();
 
-		FBodyInstance* ParentBodyInstance = GetBodyInstance(
+		FBodyInstance* ParentBodyInstance = UE::PhysicsControlComponent::GetBodyInstance(
 			Record.PhysicsControl.ParentMeshComponent, Record.PhysicsControl.ParentBoneName);
 		FTransform ParentBodyTM = ParentBodyInstance ? ParentBodyInstance->GetUnrealWorldTransform() : FTransform();
 
diff --git a/Engine/Plugins/Experimental/PhysicsControl/Source/PhysicsControl/Private/PhysicsControlComponentHelpers.cpp b/Engine/Plugins/Experimental/PhysicsControl/Source/PhysicsControl/Private/PhysicsControlComponentHelpers.cpp
--- a/Engine/Plugins/Experimental/PhysicsControl/Source/PhysicsControl/Private/PhysicsControlComponentHelpers.cpp
+++ b/Engine/Plugins/Experimental/PhysicsControl/Source/PhysicsControl/Private/PhysicsControlComponentHelpers.cpp
@@ -1,15 +1,23 @@
 // Copyright Epic Games, Inc. All Rights Reserved.
 
 #include "PhysicsControlComponentHelpers.h"
+#include "PhysicsControlData.h"
+#include "PhysicsEngine/ConstraintInstance.h"
+
 #include "Components/StaticMeshComponent.h"
 #include "Components/SkeletalMeshComponent.h"
 
+namespace UE
+{
+namespace PhysicsControlComponent
+{
+
 //======================================================================================================================
-void ConvertSpringParams(
+void ConvertStrengthToSpringParams(
 	double& OutSpring, double& OutDamping, 
 	double InStrength, double InDampingRatio, double InExtraDamping)
 {
-	double AngularFrequency = InStrength * 2.0 * PI;
+	double AngularFrequency = InStrength * UE_DOUBLE_TWO_PI;
 	double Stiffness = AngularFrequency * AngularFrequency;
 
 	OutSpring = Stiffness;
@@ -18,13 +26,127 @@ void ConvertSpringParams(
 }
 
 //======================================================================================================================
-void ConvertSpringParams(
+void ConvertStrengthToSpringParams(
 	FVector& OutSpring, FVector& OutDamping, 
 	const FVector& InStrength, float InDampingRatio, const FVector& InExtraDamping)
 {
-	ConvertSpringParams(OutSpring.X, OutDamping.X, InStrength.X, InDampingRatio, InExtraDamping.X);
-	ConvertSpringParams(OutSpring.Y, OutDamping.Y, InStrength.Y, InDampingRatio, InExtraDamping.Y);
-	ConvertSpringParams(OutSpring.Z, OutDamping.Z, InStrength.Z, InDampingRatio, InExtraDamping.Z);
+	ConvertStrengthToSpringParams(OutSpring.X, OutDamping.X, InStrength.X, InDampingRatio, InExtraDamping.X);
+	ConvertStrengthToSpringParams(OutSpring.Y, OutDamping.Y, InStrength.Y, InDampingRatio, InExtraDamping.Y);
+	ConvertStrengthToSpringParams(OutSpring.Z, OutDamping.Z, InStrength.Z, InDampingRatio, InExtraDamping.Z);
+}
+
+//======================================================================================================================
+void ConvertSpringToStrengthParams(
+	float& OutStrength, float& OutDampingRatio, float& OutExtraDamping,
+	double InSpring, double InDamping)
+{
+	// Simple calculation to get the strength
+	double AngularFrequency = FMath::Sqrt(InSpring);
+	OutStrength = AngularFrequency / UE_DOUBLE_TWO_PI;
+
+	// For damping, try to put as much into the damping ratio as possible, up to a max DR of 1. Then
+	// the rest goes into extra damping.
+	OutDampingRatio = 1.0;
+	double ImpliedDamping = 2.0 * OutDampingRatio * AngularFrequency;
+
+	if (ImpliedDamping < InDamping)
+	{
+		OutExtraDamping = InDamping - ImpliedDamping;
+	}
+	else if (AngularFrequency > 0)
+	{
+		OutExtraDamping = 0;
+		OutDampingRatio = InDamping / (2.0 * AngularFrequency);
+	}
+	else
+	{
+		OutDampingRatio = 1.0;
+		OutExtraDamping = InDamping;
+	}
+}
+
+inline double GetLinearDriveStiffness(const FLinearDriveConstraint& InDrive)
+{
+	return (
+		(InDrive.XDrive.bEnablePositionDrive ? InDrive.XDrive.Stiffness : 0.0) + 
+		(InDrive.YDrive.bEnablePositionDrive ? InDrive.YDrive.Stiffness : 0.0) +
+		(InDrive.ZDrive.bEnablePositionDrive ? InDrive.ZDrive.Stiffness : 0.0) 
+		) / 3.0;
+}
+inline double GetLinearDriveDamping(const FLinearDriveConstraint& InDrive)
+{
+	return (
+		(InDrive.XDrive.bEnableVelocityDrive ? InDrive.XDrive.Damping : 0.0) +
+		(InDrive.YDrive.bEnableVelocityDrive ? InDrive.YDrive.Damping : 0.0) + 
+		(InDrive.ZDrive.bEnableVelocityDrive ? InDrive.ZDrive.Damping : 0.0) 
+		) / 3.0;
+}
+inline double GetLinearDriveMaxForce(const FLinearDriveConstraint& InDrive)
+{
+	return FVector(InDrive.XDrive.MaxForce, InDrive.YDrive.MaxForce, InDrive.ZDrive.MaxForce).Length();
+}
+
+inline double GetAngularDriveStiffness(const FAngularDriveConstraint& InDrive)
+{
+	if (InDrive.AngularDriveMode == EAngularDriveMode::SLERP)
+	{
+		return InDrive.SlerpDrive.bEnablePositionDrive ? InDrive.SlerpDrive.Stiffness : 0.0;
+	}
+	else
+	{
+		return (
+			(InDrive.TwistDrive.bEnablePositionDrive ? InDrive.TwistDrive.Stiffness : 0.0) + 
+			(InDrive.SwingDrive.bEnablePositionDrive ? InDrive.SwingDrive.Stiffness : 0.0)
+			) / 2.0;
+	}
+}
+inline double GetAngularDriveDamping(const FAngularDriveConstraint& InDrive)
+{
+	if (InDrive.AngularDriveMode == EAngularDriveMode::SLERP)
+	{
+		return InDrive.SlerpDrive.bEnableVelocityDrive ? InDrive.SlerpDrive.Damping : 0.0;
+	}
+	else
+	{
+		return (
+			(InDrive.TwistDrive.bEnableVelocityDrive ? InDrive.TwistDrive.Damping : 0.0) +
+			(InDrive.SwingDrive.bEnableVelocityDrive ? InDrive.SwingDrive.Damping : 0.0)
+			) / 2.0;
+	}
+}
+inline double GetAngularDriveMaxTorque(const FAngularDriveConstraint& InDrive)
+{
+	if (InDrive.AngularDriveMode == EAngularDriveMode::SLERP)
+	{
+		return InDrive.SlerpDrive.MaxForce;
+	}
+	else
+	{
+		return FVector2D(InDrive.TwistDrive.MaxForce, InDrive.SwingDrive.MaxForce).Length();
+	}
+}
+
+//======================================================================================================================
+void ConvertConstraintProfileToControlData(
+	FPhysicsControlData& OutControlData, const FConstraintProfileProperties& InProfileProperties)
+{
+	ConvertSpringToStrengthParams(
+		OutControlData.LinearStrength, 
+		OutControlData.LinearDampingRatio, 
+		OutControlData.LinearExtraDamping,
+		GetLinearDriveStiffness(InProfileProperties.LinearDrive), 
+		GetLinearDriveDamping(InProfileProperties.LinearDrive));
+
+	OutControlData.MaxForce = GetLinearDriveMaxForce(InProfileProperties.LinearDrive);
+
+	ConvertSpringToStrengthParams(
+		OutControlData.AngularStrength,
+		OutControlData.AngularDampingRatio,
+		OutControlData.AngularExtraDamping,
+		GetAngularDriveStiffness(InProfileProperties.AngularDrive),
+		GetAngularDriveDamping(InProfileProperties.AngularDrive));
+
+	OutControlData.MaxTorque = GetAngularDriveMaxTorque(InProfileProperties.AngularDrive);
 }
 
 //======================================================================================================================
@@ -63,3 +185,6 @@ FName GetPhysicalParentBone(USkeletalMeshComponent* SkeletalMeshComponent, FName
 	return FName();
 }
 
+} // namespace PhysicsControlComponent
+} // namespace UE
+
diff --git a/Engine/Plugins/Experimental/PhysicsControl/Source/PhysicsControl/Private/PhysicsControlComponentHelpers.h b/Engine/Plugins/Experimental/PhysicsControl/Source/PhysicsControl/Private/PhysicsControlComponentHelpers.h
--- a/Engine/Plugins/Experimental/PhysicsControl/Source/PhysicsControl/Private/PhysicsControlComponentHelpers.h
+++ b/Engine/Plugins/Experimental/PhysicsControl/Source/PhysicsControl/Private/PhysicsControlComponentHelpers.h
@@ -5,23 +5,47 @@
 #include "UObject/NameTypes.h"
 
 struct FBodyInstance;
+struct FPhysicsControlData;
+struct FConstraintProfileProperties;
+
 class UMeshComponent;
 class USkeletalMeshComponent;
 
+namespace UE
+{
+namespace PhysicsControlComponent
+{
+
 /** 
  * Converts strength/damping ratio/extra damping into spring stiffness/damping.
  */
-void ConvertSpringParams(
+void ConvertStrengthToSpringParams(
 	double& OutSpring, double& OutDamping, 
 	double InStrength, double InDampingRatio, double InExtraDamping);
 
 /** 
  * Converts strength/damping ratio/extra damping into spring stiffness/damping 
  */
-void ConvertSpringParams(
+void ConvertStrengthToSpringParams(
 	FVector& OutSpring, FVector& OutDamping, 
 	const FVector& InStrength, float InDampingRatio, const FVector& InExtraDamping);
 
+/**
+ * Converts spring/damping values into strength/damping ratio/extra damping. This tries to get
+ * as much damping into the damping ratio term as possible, without letting it go above 1.
+ */
+void ConvertSpringToStrengthParams(
+	double& OutStrength, double& OutDampingRatio, double& OutExtraDamping,
+	double InSpring, double InDamping);
+
+/**
+ * Converts the drive settings from the constraint profile into the control data strength/damping etc. 
+ * Approximations will be made if (a) the linear drive has different values for the x/y/z axes, or (b)
+ * the constraint profile is set to use twist/swing instead of slerp for the angular drive.
+ */
+void ConvertConstraintProfileToControlData(
+	FPhysicsControlData& OutControlData, const FConstraintProfileProperties& InProfileProperties);
+
 /** 
  * Attempts to find a BodyInstance from the mesh. If it is a static mesh the single body instance
  * will be returned. If it is a skeletal mesh then if BoneName can be found, the body instance corresponding
@@ -35,3 +59,6 @@ FBodyInstance* GetBodyInstance(UMeshComponent* MeshComponent, const FName BoneNa
  * the root without finding a physical bone.
  */
 FName GetPhysicalParentBone(USkeletalMeshComponent* SkeletalMeshComponent, FName BoneName);
+
+} // namespace PhysicsControlComponent
+} // namespace UE
diff --git a/Engine/Plugins/Experimental/PhysicsControl/Source/PhysicsControl/Private/PhysicsControlComponentImpl.cpp b/Engine/Plugins/Experimental/PhysicsControl/Source/PhysicsControl/Private/PhysicsControlComponentImpl.cpp
--- a/Engine/Plugins/Experimental/PhysicsControl/Source/PhysicsControl/Private/PhysicsControlComponentImpl.cpp
+++ b/Engine/Plugins/Experimental/PhysicsControl/Source/PhysicsControl/Private/PhysicsControlComponentImpl.cpp
@@ -152,106 +152,42 @@ void FPhysicsControlComponentImpl::ResetControls(bool bKeepControlRecords)
 }
 
 //======================================================================================================================
-// Currently this looks for world-space targets from the controls, and forms a strength-weighted
-// average of them if there are multiple targets. However, it would probably be better to replace
-// this with an explicit kinematic target on each body modifier, as it is a little unintuitive to
-// make zero strength physical controls. UE-159655
 void FPhysicsControlComponentImpl::ApplyKinematicTarget(const FPhysicsBodyModifier& BodyModifier) const
 {
-	FBodyInstance* BodyInstance = GetBodyInstance(BodyModifier.MeshComponent, BodyModifier.BoneName);
+	FBodyInstance* BodyInstance = UE::PhysicsControlComponent::GetBodyInstance(
+		BodyModifier.MeshComponent, BodyModifier.BoneName);
 	if (!BodyInstance)
 	{
 		return;
 	}
 
-	// First find any controls that are (a) acting in world space (b) driving the modified object
-	float PositionWeight = 0.0f;
-	float OrientationWeight = 0.0f;
-	FVector WeightedPosition(ForceInitToZero);
-	FVector4 WeightedOrientation(ForceInitToZero);
-	for (const TPair<FName, FPhysicsControlRecord>& PhysicsControlRecordPair : PhysicsControlRecords)
+	// Seems like static and skeletal meshes need to be handled differently
+	FTransform TM = BodyInstance->GetUnrealWorldTransform();
+	FTransform KinematicTarget(BodyModifier.KinematicTargetOrientation, BodyModifier.KinematicTargetPosition);
+	if (USkeletalMeshComponent* SkeletalMeshComponent = Cast<USkeletalMeshComponent>(BodyModifier.MeshComponent.Get()))
 	{
-		const FPhysicsControlRecord& Record = PhysicsControlRecordPair.Value;
-		FBodyInstance* ParentBodyInstance = GetBodyInstance(
-			Record.PhysicsControl.ParentMeshComponent, Record.PhysicsControl.ParentBoneName);
-
-		if (!ParentBodyInstance)
+		if (BodyModifier.bUseSkeletalAnimation)
 		{
-			FBodyInstance* ChildBodyInstance = GetBodyInstance(
-				Record.PhysicsControl.ChildMeshComponent, Record.PhysicsControl.ChildBoneName);
-			if (BodyInstance == ChildBodyInstance)
+			FCachedSkeletalMeshData::FBoneData BoneData;
+			if (GetBoneData(BoneData, SkeletalMeshComponent, BodyModifier.BoneName))
 			{
-				FTransform TargetTM;
-				FVector TargetVelocity;
-				FVector TargetAngularVelocity;
-				CalculateControlTargetData(TargetTM, TargetVelocity, TargetAngularVelocity, Record, false);
-
-				// TargetTM is actually the target for the control point, but we will be setting the body TM
-				// so we need to remove the offset
-				TargetTM.AddToTranslation(TargetTM.GetRotation() * -Record.PhysicsControl.ControlSettings.ControlPoint);
-
-				// TODO note that this isn't using the multipliers, or the force limits etc. Using an explicit
-				// kinematic target will solve this
-				float LinearWeight = Record.PhysicsControl.ControlData.LinearStrength + UE_SMALL_NUMBER;
-				float AngularWeight = Record.PhysicsControl.ControlData.AngularStrength + UE_SMALL_NUMBER;
-
-				WeightedPosition += TargetTM.GetTranslation() * LinearWeight;
-				PositionWeight += LinearWeight;
-
-				FQuat Q = TargetTM.GetRotation();
-				Q.EnforceShortestArcWith(FQuat::Identity);
-				WeightedOrientation += FVector4(Q.X, Q.Y, Q.Z, Q.W) * AngularWeight;
-				OrientationWeight += AngularWeight;
+				FTransform BoneTM = BoneData.GetTM();
+				KinematicTarget = KinematicTarget * BoneTM;
 			}
 		}
-	}
-
-	if (PositionWeight <= 0 && OrientationWeight <= 0)
-	{
-		return;
-	}
-
-	// Seems like static and skeletal meshes need to be handled differently
-	if (BodyModifier.MeshComponent.IsA<USkeletalMeshComponent>())
-	{
-		FTransform NewTM = BodyInstance->GetUnrealWorldTransform();
-		if (PositionWeight)
-		{
-			FVector NewPosition = WeightedPosition / PositionWeight;
-			NewTM.SetLocation(NewPosition);
-		}
-		if (OrientationWeight)
-		{
-			WeightedOrientation /= OrientationWeight;
-			FQuat TargetOrientation = FQuat(
-				WeightedOrientation.X, WeightedOrientation.Y, WeightedOrientation.Z, WeightedOrientation.W);
-			TargetOrientation.Normalize();
-			NewTM.SetRotation(TargetOrientation);
-		}
-
-		BodyInstance->SetBodyTransform(NewTM, ETeleportType::None);
+		ETeleportType TT = DetectTeleport(TM, KinematicTarget) ? ETeleportType::ResetPhysics : ETeleportType::None;
+		BodyInstance->SetBodyTransform(KinematicTarget, TT);
 	}
 	else
 	{
-		if (PositionWeight)
-		{
-			FVector TargetPosition = WeightedPosition / PositionWeight;
-			BodyModifier.MeshComponent->SetWorldLocation(TargetPosition, false, nullptr, ETeleportType::None);
-		}
-
-		if (OrientationWeight)
-		{
-			WeightedOrientation /= OrientationWeight;
-			FQuat TargetOrientation = FQuat(
-				WeightedOrientation.X, WeightedOrientation.Y, WeightedOrientation.Z, WeightedOrientation.W);
-			TargetOrientation.Normalize();
-			BodyModifier.MeshComponent->SetWorldRotation(TargetOrientation, false, nullptr, ETeleportType::None);
-		}
+		ETeleportType TT = DetectTeleport(TM, KinematicTarget) ? ETeleportType::ResetPhysics : ETeleportType::None;
+		BodyModifier.MeshComponent->SetWorldTransform(KinematicTarget, false, nullptr, TT);
 	}
 }
 
 //======================================================================================================================
-void FPhysicsControlComponentImpl::AddSkeletalMeshReference(USkeletalMeshComponent* InSkeletalMeshComponent)
+void FPhysicsControlComponentImpl::AddSkeletalMeshReferenceForCaching(
+	USkeletalMeshComponent* InSkeletalMeshComponent)
 {
 	check(InSkeletalMeshComponent);
 	for (TPair<TObjectPtr<USkeletalMeshComponent>, FCachedSkeletalMeshData>& CachedSkeletalMeshDataPair :
@@ -271,31 +207,89 @@ void FPhysicsControlComponentImpl::AddSkeletalMeshReference(USkeletalMeshCompone
 }
 
 //======================================================================================================================
-void FPhysicsControlComponentImpl::RemoveSkeletalMeshReference(USkeletalMeshComponent* InSkeletalMeshComponent)
+bool FPhysicsControlComponentImpl::RemoveSkeletalMeshReferenceForCaching(
+	USkeletalMeshComponent* InSkeletalMeshComponent)
 {
 	check(InSkeletalMeshComponent);
 	if (!InSkeletalMeshComponent)
 	{
 		UE_LOG(LogPhysicsControlComponent, Warning, TEXT("Invalid skeletal mesh component"));
-		return;
+		return false;
 	}
 
 	for (auto It = CachedSkeletalMeshDatas.CreateIterator(); It; ++It)
 	{
 		TObjectPtr<USkeletalMeshComponent> SkeletalMeshComponent = It.Key();
-		FCachedSkeletalMeshData& CachedSkeletalMeshData = It.Value();
+		FCachedSkeletalMeshData& Data = It.Value();
 		if (SkeletalMeshComponent == InSkeletalMeshComponent)
 		{
-			if (--CachedSkeletalMeshData.ReferenceCount == 0)
+			if (--Data.ReferenceCount == 0)
 			{
 				Owner->PrimaryComponentTick.RemovePrerequisite(
 					InSkeletalMeshComponent, InSkeletalMeshComponent->PrimaryComponentTick);
 				It.RemoveCurrent();
+				return true;
 			}
+			return false;
+		}
+	}
+	UE_LOG(LogPhysicsControlComponent, Warning, TEXT("Failed to remove skeletal mesh component reference for caching"));
+	return false;
+}
+
+//======================================================================================================================
+void FPhysicsControlComponentImpl::AddSkeletalMeshReferenceForModifier(
+	USkeletalMeshComponent* InSkeletalMeshComponent)
+{
+	check(InSkeletalMeshComponent);
+	for (TPair<TObjectPtr<USkeletalMeshComponent>, FModifiedSkeletalMeshData>& ModifiedSkeletalMeshDataPair :
+		ModifiedSkeletalMeshDatas)
+	{
+		TObjectPtr<USkeletalMeshComponent> SkeletalMeshComponent = ModifiedSkeletalMeshDataPair.Key;
+		if (SkeletalMeshComponent == InSkeletalMeshComponent)
+		{
+			FModifiedSkeletalMeshData& ModifiedSkeletalMeshData = ModifiedSkeletalMeshDataPair.Value;
+			++ModifiedSkeletalMeshData.ReferenceCount;
 			return;
 		}
 	}
-	UE_LOG(LogPhysicsControlComponent, Warning, TEXT("Failed to remove skeletal mesh component dependency"));
+	FModifiedSkeletalMeshData& Data = ModifiedSkeletalMeshDatas.Add(InSkeletalMeshComponent);
+	Data.ReferenceCount = 1;
+	Data.bOriginalUpdateMeshWhenKinematic = InSkeletalMeshComponent->bUpdateMeshWhenKinematic;
+	Data.OriginalKinematicBonesUpdateType = InSkeletalMeshComponent->KinematicBonesUpdateType;
+	InSkeletalMeshComponent->bUpdateMeshWhenKinematic = true;
+	InSkeletalMeshComponent->KinematicBonesUpdateType = EKinematicBonesUpdateToPhysics::SkipAllBones;
+}
+
+//======================================================================================================================
+bool FPhysicsControlComponentImpl::RemoveSkeletalMeshReferenceForModifier(
+	USkeletalMeshComponent* InSkeletalMeshComponent)
+{
+	check(InSkeletalMeshComponent);
+	if (!InSkeletalMeshComponent)
+	{
+		UE_LOG(LogPhysicsControlComponent, Warning, TEXT("Invalid skeletal mesh component"));
+		return false;
+	}
+
+	for (auto It = ModifiedSkeletalMeshDatas.CreateIterator(); It; ++It)
+	{
+		TObjectPtr<USkeletalMeshComponent> SkeletalMeshComponent = It.Key();
+		FModifiedSkeletalMeshData& Data = It.Value();
+		if (SkeletalMeshComponent == InSkeletalMeshComponent)
+		{
+			if (--Data.ReferenceCount == 0)
+			{
+				InSkeletalMeshComponent->bUpdateMeshWhenKinematic = Data.bOriginalUpdateMeshWhenKinematic;
+				InSkeletalMeshComponent->KinematicBonesUpdateType = Data.OriginalKinematicBonesUpdateType;
+				It.RemoveCurrent();
+				return true;
+			}
+			return false;
+		}
+	}
+	UE_LOG(LogPhysicsControlComponent, Warning, TEXT("Failed to remove skeletal mesh component reference for modifier"));
+	return false;
 }
 
 //======================================================================================================================
@@ -524,12 +518,12 @@ bool FPhysicsControlComponentImpl::ApplyControlStrengths(
 	FVector LinearDamping;
 	FVector MaxForce = Data.MaxForce * Multipliers.MaxForceMultiplier;
 
-	ConvertSpringParams(
+	UE::PhysicsControlComponent::ConvertStrengthToSpringParams(
 		AngularSpring, AngularDamping,
 		Data.AngularStrength * Multipliers.AngularStrengthMultiplier,
 		Data.AngularDampingRatio,
 		Data.AngularExtraDamping * Multipliers.AngularExtraDampingMultiplier);
-	ConvertSpringParams(
+	UE::PhysicsControlComponent::ConvertStrengthToSpringParams(
 		LinearSpring, LinearDamping,
 		Data.LinearStrength * Multipliers.LinearStrengthMultiplier,
 		Data.LinearDampingRatio,
@@ -578,10 +572,10 @@ void FPhysicsControlComponentImpl::ApplyControl(FPhysicsControlRecord& Record)
 		return;
 	}
 
-	FBodyInstance* ParentBodyInstance = GetBodyInstance(
+	FBodyInstance* ParentBodyInstance = UE::PhysicsControlComponent::GetBodyInstance(
 		Record.PhysicsControl.ParentMeshComponent, Record.PhysicsControl.ParentBoneName);
 
-	FBodyInstance* ChildBodyInstance = GetBodyInstance(
+	FBodyInstance* ChildBodyInstance = UE::PhysicsControlComponent::GetBodyInstance(
 		Record.PhysicsControl.ChildMeshComponent, Record.PhysicsControl.ChildBoneName);
 
 	if (!ParentBodyInstance && !ChildBodyInstance)
diff --git a/Engine/Plugins/Experimental/PhysicsControl/Source/PhysicsControl/Private/PhysicsControlComponentImpl.h b/Engine/Plugins/Experimental/PhysicsControl/Source/PhysicsControl/Private/PhysicsControlComponentImpl.h
--- a/Engine/Plugins/Experimental/PhysicsControl/Source/PhysicsControl/Private/PhysicsControlComponentImpl.h
+++ b/Engine/Plugins/Experimental/PhysicsControl/Source/PhysicsControl/Private/PhysicsControlComponentImpl.h
@@ -59,12 +59,24 @@ public:
 	/**
 	 * Starts caching skeletal mesh poses, and registers for a tick pre-requisite
 	 */
-	void AddSkeletalMeshReference(USkeletalMeshComponent* SkeletalMeshComponent);
+	void AddSkeletalMeshReferenceForCaching(USkeletalMeshComponent* SkeletalMeshComponent);
 
 	/**
-	 * Stops caching skeletal mesh poses (if this is the last one), and deregisters for a tick pre-requisite
+	 * Stops caching skeletal mesh poses (if this is the last one), and deregisters for a tick pre-requisite.
+	 * Returns true if this was the last reference, false otherwise.
 	 */
-	void RemoveSkeletalMeshReference(USkeletalMeshComponent* SkeletalMeshComponent);
+	bool RemoveSkeletalMeshReferenceForCaching(USkeletalMeshComponent* SkeletalMeshComponent);
+
+	/**
+	 * Records that a modifier is working with the skeletal mesh and stores original data if necessary
+	 */
+	void AddSkeletalMeshReferenceForModifier(USkeletalMeshComponent* SkeletalMeshComponent);
+
+	/**
+	 * Records that a modifier has stopped working with the skeletal mesh and restores original data if necessary.
+	 * Returns true if this was the last reference, false otherwise.
+	 */
+	bool RemoveSkeletalMeshReferenceForModifier(USkeletalMeshComponent* SkeletalMeshComponent);
 
 	/** Update the constraint based on the record. */
 	void ApplyControl(FPhysicsControlRecord& Record);
@@ -107,8 +119,11 @@ public:
 	// beginning of each tick.
 	TMap<TObjectPtr<USkeletalMeshComponent>, FCachedSkeletalMeshData> CachedSkeletalMeshDatas;
 
-	TMap<FName, FPhysicsControlRecord> PhysicsControlRecords;
+	// Track which skeletons have been affected by a body modifier - some settings get overridden
+	// and then need to be restored when the last body modifier is destroyed.
+	TMap<TObjectPtr<USkeletalMeshComponent>, FModifiedSkeletalMeshData> ModifiedSkeletalMeshDatas;
 
+	TMap<FName, FPhysicsControlRecord> PhysicsControlRecords;
 	TMap<FName, FPhysicsBodyModifier> PhysicsBodyModifiers;
 
 private:
diff --git a/Engine/Plugins/Experimental/PhysicsControl/Source/PhysicsControl/Private/PhysicsControlRecord.cpp b/Engine/Plugins/Experimental/PhysicsControl/Source/PhysicsControl/Private/PhysicsControlRecord.cpp
--- a/Engine/Plugins/Experimental/PhysicsControl/Source/PhysicsControl/Private/PhysicsControlRecord.cpp
+++ b/Engine/Plugins/Experimental/PhysicsControl/Source/PhysicsControl/Private/PhysicsControlRecord.cpp
@@ -28,9 +28,9 @@ FConstraintInstance* FPhysicsControlRecord::CreateConstraint(UObject* Constraint
 	}
 	FConstraintInstance* ConstraintInstance = PhysicsControlState.ConstraintInstance.Get();
 
-	FBodyInstance* ParentBody = GetBodyInstance(
+	FBodyInstance* ParentBody = UE::PhysicsControlComponent::GetBodyInstance(
 		PhysicsControl.ParentMeshComponent, PhysicsControl.ParentBoneName);
-	FBodyInstance* ChildBody = GetBodyInstance(
+	FBodyInstance* ChildBody = UE::PhysicsControlComponent::GetBodyInstance(
 		PhysicsControl.ChildMeshComponent, PhysicsControl.ChildBoneName);
 
 	if (PhysicsControl.ParentMeshComponent && !PhysicsControl.ParentBoneName.IsNone() && !ParentBody)
@@ -83,7 +83,7 @@ void FPhysicsControlRecord::UpdateConstraintControlPoint()
 //======================================================================================================================
 void FPhysicsControlRecord::ResetControlPoint()
 {
-	FBodyInstance* ChildBodyInstance = GetBodyInstance(
+	FBodyInstance* ChildBodyInstance = UE::PhysicsControlComponent::GetBodyInstance(
 		PhysicsControl.ChildMeshComponent, PhysicsControl.ChildBoneName);
 
 	PhysicsControl.ControlSettings.ControlPoint = ChildBodyInstance
diff --git a/Engine/Plugins/Experimental/PhysicsControl/Source/PhysicsControl/Private/PhysicsControlRecord.h b/Engine/Plugins/Experimental/PhysicsControl/Source/PhysicsControl/Private/PhysicsControlRecord.h
--- a/Engine/Plugins/Experimental/PhysicsControl/Source/PhysicsControl/Private/PhysicsControlRecord.h
+++ b/Engine/Plugins/Experimental/PhysicsControl/Source/PhysicsControl/Private/PhysicsControlRecord.h
@@ -2,8 +2,8 @@
 
 #pragma once
 
-#include "UObject/ObjectMacros.h"
 #include "PhysicsControlData.h"
+#include "Components/SkeletalMeshComponent.h"
 
 struct FConstraintInstance;
 struct FBodyInstance;
@@ -31,7 +31,7 @@ struct FPhysicsControlState
 /**
  * There will be a PhysicsControlRecord created at runtime for every Control that has been created
  */
-struct FPhysicsControlRecord : public TSharedFromThis<FPhysicsControlRecord>
+struct FPhysicsControlRecord
 {
 	FPhysicsControlRecord(const FPhysicsControl& InControl)
 		: PhysicsControl(InControl)
@@ -60,28 +60,79 @@ struct FPhysicsControlRecord : public TSharedFromThis<FPhysicsControlRecord>
 /**
  * There will be a PhysicsBodyModifier created at runtime for every BodyInstance involved in the component
  */
-struct FPhysicsBodyModifier : public TSharedFromThis<FPhysicsBodyModifier>
+struct FPhysicsBodyModifier
 {
 	FPhysicsBodyModifier(
 		TObjectPtr<UMeshComponent> InMeshComponent, 
 		const FName&               InBoneName, 
 		EPhysicsMovementType       InMovementType, 
-		float                      InGravityMultiplier)
-	: MeshComponent(InMeshComponent), BoneName(InBoneName)
-	, MovementType(InMovementType), GravityMultiplier(InGravityMultiplier)
-	, bPendingDestroy(false)
+		float                      InGravityMultiplier,
+		bool                       InUseSkeletalAnimation)
+		: MeshComponent(InMeshComponent)
+		, BoneName(InBoneName)
+		, MovementType(InMovementType)
+		, GravityMultiplier(InGravityMultiplier)
+		, KinematicTargetPosition(FVector::ZeroVector)
+		, KinematicTargetOrientation(FQuat::Identity)
+		, bUseSkeletalAnimation(InUseSkeletalAnimation)
+		, bPendingDestroy(false)
 	{}
 
 	/**  The mesh that will be modified. */
 	TObjectPtr<UMeshComponent> MeshComponent;
 
 	/** The name of the skeletal mesh bone or the name of the static mesh body that will be modified. */
-	FName                      BoneName;
+	FName BoneName;
 
-	EPhysicsMovementType       MovementType;
-	float                      GravityMultiplier;
+	/** How the associated body should move. */
+	EPhysicsMovementType MovementType;
 
-	bool bPendingDestroy;
+	/**
+	 * Multiplier for gravity applied to the body. Note that if the body itself has gravity disabled, then
+	 * setting this to 1 will not enable gravity.
+	 */
+	float GravityMultiplier;
+
+	/** 
+	 * The target position when kinematic. Note that this is applied on top of any animation 
+	 * target if bUseSkeletalAnimation is set. 
+	 */
+	FVector KinematicTargetPosition;
+
+	/**
+	 * The target orientation when kinematic. Note that this is applied on top of any animation
+	 * target if bUseSkeletalAnimation is set.
+	 */
+	FQuat KinematicTargetOrientation;
+
+	/** If true then the target will be applied on top of the skeletal animation (if there is any) */
+	uint8 bUseSkeletalAnimation:1;
+
+	/** Indicates if this should be deleted on the next tick */
+	uint8 bPendingDestroy:1;
+};
+
+/**
+ * Used internally/only at runtime to track when a SkeletalMeshComponent is being controlled through
+ * a modifier, and to restore settings when that stops.
+ */
+struct FModifiedSkeletalMeshData
+{
+public: 
+	FModifiedSkeletalMeshData() : ReferenceCount(0) {}
+
+public:
+
+	/** The original setting for restoration when we're deleted */
+	uint8 bOriginalUpdateMeshWhenKinematic : 1;
+
+	/** The original setting for restoration when we're deleted */
+	EKinematicBonesUpdateToPhysics::Type OriginalKinematicBonesUpdateType;
+
+	/**
+	 * Track when skeletal meshes are going to be used so this entry can be removed
+	 */
+	int32 ReferenceCount;
 };
 
 /**
diff --git a/Engine/Plugins/Experimental/PhysicsControl/Source/PhysicsControl/Public/PhysicsControlComponent.h b/Engine/Plugins/Experimental/PhysicsControl/Source/PhysicsControl/Public/PhysicsControlComponent.h
--- a/Engine/Plugins/Experimental/PhysicsControl/Source/PhysicsControl/Public/PhysicsControlComponent.h
+++ b/Engine/Plugins/Experimental/PhysicsControl/Source/PhysicsControl/Public/PhysicsControlComponent.h
@@ -69,9 +69,9 @@ public:
 	UFUNCTION(BlueprintCallable, Category = PhysicsControl)
 	FName MakeControl(
 		UMeshComponent*         ParentMeshComponent,
-		const FName             ParentBoneName,
+		FName                   ParentBoneName,
 		UMeshComponent*         ChildMeshComponent,
-		const FName             ChildBoneName,
+		FName                   ChildBoneName,
 		FPhysicsControlData     ControlData,
 		FPhysicsControlTarget   ControlTarget, 
 		FPhysicsControlSettings ControlSettings,
@@ -93,11 +93,11 @@ public:
 	 */
 	UFUNCTION(BlueprintCallable, Category = PhysicsControl)
 	bool MakeNamedControl(
-		const FName             Name,
+		FName                   Name,
 		UMeshComponent*         ParentMeshComponent,
-		const FName             ParentBoneName,
+		FName                   ParentBoneName,
 		UMeshComponent*         ChildMeshComponent,
-		const FName             ChildBoneName,
+		FName                   ChildBoneName,
 		FPhysicsControlData     ControlData, 
 		FPhysicsControlTarget   ControlTarget, 
 		FPhysicsControlSettings ControlSettings, 
@@ -133,6 +133,40 @@ public:
 		FPhysicsControlSettings ControlSettings,
 		bool                    bEnabled = true);
 
+	/**
+	 * Makes a collection of ParentSpace controls controlling a skeletal mesh, initializing
+	 * them with a constraint profile
+	 * 
+	 * @param SkeletalMeshComponent The skeletal mesh which will have controls
+	 * 
+	 * @param BoneName The name of the bone below which controls should be created. Each bone will be the child in a control
+	 * 
+	 * @param bIncludeSelf Whether or not to include BoneName when creating controls
+	 * 
+	 * @param ConstraintProfile The constraint profile to use for initializing the control strength and
+	 *                          damping (etc) parameters. Note that the controls will all be created in "
+	 *                          parent space" - i.e. with each part controlled relative to its parent. The
+	 *                          strength and damping will be taken from the values that the relevant joint
+	 *                          in the physics asset would have given the constraint profile (or the default
+	 *                          profile if it can't be found) - though they will not match exactly if the linear
+	 *                          drive and different x/y/z values, or if the angular drive was using twist/swing
+	 *                          instead of slerp. Note also that the joint constraints do not use the animation
+	 *                          velocity as a target, so when creating controls in this way the control settings
+	 *                          will set the skeletal animation velocity multiplier to zero.
+	 *
+	 * @param bEnabled      If true then the control will be enabled immediately. If false you will need to call
+	 *                      SetControlEnabled(true) in order to enable it.
+	 *
+	 * @return An array of the controls that have been created
+	 */
+	UFUNCTION(BlueprintCallable, Category = PhysicsControl)
+	TArray<FName> MakeControlsFromSkeletalMeshAndConstraintProfileBelow(
+		USkeletalMeshComponent* SkeletalMeshComponent,
+		FName                   BoneName,
+		bool                    bIncludeSelf,
+		FName                   ConstraintProfile,
+		bool                    bEnabled = true);
+
 	/**
 	 * Makes a collection of controls controlling a skeletal mesh
 	 *
@@ -160,6 +194,37 @@ public:
 		FPhysicsControlSettings ControlSettings,
 		bool                    bEnabled = true);
 
+	/**
+	 * Makes a collection of ParentSpace controls controlling a skeletal mesh, initializing them 
+	 * with a constraint profile
+	 *
+	 * @param SkeletalMeshComponent The skeletal mesh which will have controls
+	 *
+	 * @param BoneNames The names of bones for which controls should be created. Each bone will be the child in a control
+	 *
+	 * @param ConstraintProfile The constraint profile to use for initializing the control strength and
+	 *                          damping (etc) parameters. Note that the controls will all be created in "
+	 *                          parent space" - i.e. with each part controlled relative to its parent. The
+	 *                          strength and damping will be taken from the values that the relevant joint
+	 *                          in the physics asset would have given the constraint profile (or the default
+	 *                          profile if it can't be found) - though they will not match exactly if the linear
+	 *                          drive and different x/y/z values, or if the angular drive was using twist/swing
+	 *                          instead of slerp. Note also that the joint constraints do not use the animation
+	 *                          velocity as a target, so when creating controls in this way the control settings
+	 *                          will set the skeletal animation velocity multiplier to zero.
+	 *
+	 * @param bEnabled      If true then the control will be enabled immediately. If false you will need to call
+	 *                      SetControlEnabled(true) in order to enable it.
+	 *
+	 * @return An array of the controls that have been created
+	 */
+	UFUNCTION(BlueprintCallable, Category = PhysicsControl)
+	TArray<FName> MakeControlsFromSkeletalMeshAndConstraintProfile(
+		USkeletalMeshComponent* SkeletalMeshComponent,
+		const TArray<FName>&    BoneNames,
+		FName                   ConstraintProfile,
+		bool                    bEnabled = true);
+
 	/**
 	 * Calculates which bones belong to which limb in a skeletal mesh
 	 * 
@@ -206,6 +271,38 @@ public:
 		FPhysicsControlSettings                      ControlSettings,
 		bool                                         bEnabled = true);
 
+	/**
+	 * Makes a collection of ParentSpace controls controlling a skeletal mesh, grouped together in limbs, initializing
+	 * them with a constraint profile
+	 *
+	 * @param AllControls A single container for all the controls that have been created
+	 *
+	 * @param LimbBones A map relating the limbs and the bones that they contain. Typically create this
+	 *                  using GetLimbBonesFromSkeletalMesh
+	 *
+	 * @param ConstraintProfile The constraint profile to use for initializing the control strength and 
+	 *                          damping (etc) parameters. Note that the controls will all be created in "
+	 *                          parent space" - i.e. with each part controlled relative to its parent. The 
+	 *                          strength and damping will be taken from the values that the relevant joint 
+	 *                          in the physics asset would have given the constraint profile (or the default 
+	 *                          profile if it can't be found) - though they will not match exactly if the linear
+	 *                          drive and different x/y/z values, or if the angular drive was using twist/swing
+	 *                          instead of slerp. Note also that the joint constraints do not use the animation 
+	 *                          velocity as a target, so when creating controls in this way the control settings
+	 *                          will set the skeletal animation velocity multiplier to zero.
+	 *
+	 * @param bEnabled If true then the control will be enabled immediately. If false you will need to call
+	 *                 SetControlEnabled(true) in order to enable it.
+	 *
+	 * @return A map containing the controls for each limb
+	 */
+	UFUNCTION(BlueprintCallable, Category = PhysicsControl)
+	TMap<FName, FPhysicsControlNameArray> MakeControlsFromLimbBonesAndConstraintProfile(
+		FPhysicsControlNameArray&                    AllControls,
+		const TMap<FName, FPhysicsControlLimbBones>& LimbBones,
+		FName                                        ConstraintProfile,
+		bool                                         bEnabled = true);
+
 	/**
 	 * Destroys a control
 	 *
@@ -215,7 +312,16 @@ public:
 	 * @return     Returns true if the control was found and destroyed, false if not
 	 */
 	UFUNCTION(BlueprintCallable, Category = PhysicsControl)
-	bool DestroyControl(const FName Name);
+	bool DestroyControl(FName Name);
+
+	/**
+	 * Destroys all controls
+	 *
+	 * @param Names The names of the controls to destroy. Note that if you have these in a FPhysicsControlNameArray
+	 *              then it can be split.
+	 */
+	UFUNCTION(BlueprintCallable, Category = PhysicsControl)
+	void DestroyAllControls(const TArray<FName>& Names);
 
 	/**
 	 * Modifies an existing control data - i.e. the strengths etc of the control driving towards the target
@@ -230,7 +336,20 @@ public:
 	 * @return Returns true if the control was found and modified, false if not
 	 */
 	UFUNCTION(BlueprintCallable, Category = PhysicsControl)
-	bool SetControlData(const FName Name, FPhysicsControlData ControlData, bool bEnableControl = true);
+	bool SetControlData(FName Name, FPhysicsControlData ControlData, bool bEnableControl = true);
+
+	/**
+	 * Modifies existing control data - i.e. the strengths etc of the controls driving towards the targets
+	 *
+	 * @param Names The names of the controls to modify. Note that if you have these in a FPhysicsControlNameArray
+	 *              then it can be split.
+	 *
+	 * @param ControlData The new control data
+	 *
+	 * @param bEnableControl Enables/disables the controls
+	 */
+	UFUNCTION(BlueprintCallable, Category = PhysicsControl)
+	void SetAllControlDatas(const TArray<FName>& Names, FPhysicsControlData ControlData, bool bEnableControl = true);
 
 	/**
 	 * Modifies an existing control data using the multipliers
@@ -245,7 +364,23 @@ public:
 	 * @return Returns true if the control was found and modified, false if not
 	 */
 	UFUNCTION(BlueprintCallable, Category = PhysicsControl)
-	bool SetControlMultipliers(const FName Name, FPhysicsControlMultipliers ControlMultipliers, bool bEnableControl = true);
+	bool SetControlMultipliers(FName Name, FPhysicsControlMultipliers ControlMultipliers, bool bEnableControl = true);
+
+	/**
+	 * Modifies existing control data using the multipliers
+	 *
+	 * @param Names The names of the controls to modify. Note that if you have these in a FPhysicsControlNameArray
+	 *              then it can be split.
+	 *
+	 * @param ControlMultipliers The new control multipliers
+	 *
+	 * @param bEnableControl Enables/disables the control
+	 */
+	UFUNCTION(BlueprintCallable, Category = PhysicsControl)
+	void SetAllControlMultipliers(
+		const TArray<FName>&       Names, 
+		FPhysicsControlMultipliers ControlMultipliers, 
+		bool                       bEnableControl = true);
 
 	/**
 	 * Modifies an existing control's linear data - i.e. the strengths etc of the control driving towards the target
@@ -267,7 +402,7 @@ public:
 	 */
 	UFUNCTION(BlueprintCallable, Category = PhysicsControl)
 	bool SetControlLinearData(
-		const FName Name,
+		FName Name,
 		float Strength = 1.0f, 
 		float DampingRatio = 1.0f, 
 		float ExtraDamping = 0.0f, 
@@ -294,7 +429,7 @@ public:
 	 */
 	UFUNCTION(BlueprintCallable, Category = PhysicsControl)
 	bool SetControlAngularData(
-		const FName Name,
+		FName Name,
 		float Strength = 1.0f, 
 		float DampingRatio = 1.0f, 
 		float ExtraDamping = 0.0f, 
@@ -311,7 +446,7 @@ public:
 	 *                 object is in use and is being simulated)
 	 */
 	UFUNCTION(BlueprintCallable, Category = PhysicsControl)
-	bool SetControlPoint(const FName Name, const FVector Position);
+	bool SetControlPoint(FName Name, const FVector Position);
 
 	/**
 	 * Resets the control point to the center of mass of the mesh
@@ -320,7 +455,7 @@ public:
 	 *                control, whatever its name.
 	 */
 	UFUNCTION(BlueprintCallable, Category = PhysicsControl)
-	bool ResetControlPoint(const FName Name);
+	bool ResetControlPoint(FName Name);
 
 	/**
 	 * Modifies an existing control target - i.e. what it is driving towards, relative to the parent object
@@ -335,7 +470,24 @@ public:
 	 * @return Returns true if the control was found and modified, false if not
 	 */
 	UFUNCTION(BlueprintCallable, Category = PhysicsControl)
-	bool SetControlTarget(const FName Name, FPhysicsControlTarget ControlTarget, bool bEnableControl = true);
+	bool SetControlTarget(FName Name, FPhysicsControlTarget ControlTarget, bool bEnableControl = true);
+
+	/**
+	 * Modifies existing control targets - i.e. what they are driving towards, relative to the parent objects
+	 *
+	 * @param Names The names of the controls to modify. Note that if you have these in a FPhysicsControlNameArray
+	 *              then it can be split.
+	 *
+	 * @param ControlTarget The new target for the controls
+	 * 
+	 * @param bEnableControl Enables/disables the controls
+	 */
+	UFUNCTION(BlueprintCallable, Category = PhysicsControl)
+	void SetAllControlTargets(
+		const TArray<FName>&  Names, 
+		FPhysicsControlTarget ControlTarget, 
+		bool                  bEnableControl = true);
+
 
 	/**
 	 * Modifies an existing control target - i.e. what it is driving towards, relative to the parent object
@@ -362,7 +514,7 @@ public:
 	 */
 	UFUNCTION(BlueprintCallable, Category = PhysicsControl)
 	bool SetControlTargetTransform(
-		const FName      Name,
+		FName            Name,
 		const FTransform Transform, 
 		float            VelocityDeltaTime, 
 		bool             bEnableControl = true,
@@ -393,11 +545,11 @@ public:
 	 */
 	UFUNCTION(BlueprintCallable, Category = PhysicsControl)
 	bool SetControlTargetPosition(
-		const FName Name,
-		const FVector  Position, 
-		float          VelocityDeltaTime, 
-		bool           bEnableControl = true,
-		bool           bApplyControlPointToTarget = true);
+		FName         Name,
+		const FVector Position, 
+		float         VelocityDeltaTime, 
+		bool          bEnableControl = true,
+		bool          bApplyControlPointToTarget = true);
 
 	/**
 	 * Modifies an existing control target - i.e. what it is driving towards, relative to the parent object
@@ -424,7 +576,7 @@ public:
 	 */
 	UFUNCTION(BlueprintCallable, Category = PhysicsControl)
 	bool SetControlTargetOrientation(
-		const FName Name,
+		FName          Name,
 		const FRotator Orientation, 
 		float          AngularVelocityDeltaTime, 
 		bool           bEnableControl = true,
@@ -454,29 +606,52 @@ public:
 	 */
 	UFUNCTION(BlueprintCallable, Category = PhysicsControl)
 	bool SetControlTargetPoses(
-		const FName Name,
-		const FVector ParentPosition, const FRotator ParentOrientation,
-		const FVector ChildPosition, const FRotator ChildOrientation,
-		float VelocityDeltaTime, bool bEnableControl = true);
+		FName          Name,
+		const FVector  ParentPosition, 
+		const FRotator ParentOrientation,
+		const FVector  ChildPosition, 
+		const FRotator ChildOrientation,
+		float          VelocityDeltaTime, 
+		bool           bEnableControl = true);
 
 
 	/**
 	 * Sets whether or not the control should use skeletal animation for the targets
 	 *
-	 * @param Name              The name of the control to modify. If blank, then this will access the first
-	 *                          control, whatever its name.
+	 * @param Name The name of the control to modify. If blank, then this will access the first
+	 *             control, whatever its name.
 	 * 
 	 * @param bUseSkeletalAnimation If true then the targets will be a combination of the skeletal animation (if
 	 *                              there is any) and the control target that has been set
 	 * 
 	 * @param SkeletalAnimationVelocityMultiplier If skeletal animation is being used, then this determines the amount of 
-	 *                              velocity extracted from the animation that is used as targets for the controls
+	 *                                            velocity extracted from the animation that is used as targets for the controls
+	 * 
+	 * @return Returns true if the control was found and modified, false if not
 	 */
 	UFUNCTION(BlueprintCallable, Category = PhysicsControl)
 	bool SetControlUseSkeletalAnimation(
-		const FName Name,
-		bool        bUseSkeletalAnimation = true,
-		float       SkeletalAnimationVelocityMultiplier = 1.0f);
+		FName Name,
+		bool  bUseSkeletalAnimation = true,
+		float SkeletalAnimationVelocityMultiplier = 1.0f);
+
+	/**
+	 * Sets whether or not the controls should use skeletal animation for the targets
+	 *
+	 * @param Names The names of the controls to modify. Note that if you have these in a FPhysicsControlNameArray 
+	 *              then it can be split.
+	 * 
+	 * @param bUseSkeletalAnimation If true then the targets will be a combination of the skeletal animation (if
+	 *                              there is any) and the control target that has been set
+	 * 
+	 * @param SkeletalAnimationVelocityMultiplier If skeletal animation is being used, then this determines the amount of 
+	 *                              velocity extracted from the animation that is used as targets for the controls
+	 */
+	UFUNCTION(BlueprintCallable, Category = PhysicsControl)
+	void SetAllControlsUseSkeletalAnimation(
+		const TArray<FName>& Names,
+		bool                 bUseSkeletalAnimation = true,
+		float                SkeletalAnimationVelocityMultiplier = 1.0f);
 
 	/**
 	 * Activates or deactivates a control
@@ -489,7 +664,18 @@ public:
 	 * @return         Returns true if the control was found and modified, false if not
 	 */
 	UFUNCTION(BlueprintCallable, Category = PhysicsControl)
-	bool SetControlEnabled(const FName Name, bool bEnable);
+	bool SetControlEnabled(FName Name, bool bEnable = true);
+
+	/**
+	 * Activates or deactivates controls
+	 *
+	 * @param Names The names of the controls to modify. Note that if you have these in a FPhysicsControlNameArray
+	 *              then it can be split.
+	 *
+	 * @param bEnable  The controls to enable/disable
+	 */
+	UFUNCTION(BlueprintCallable, Category = PhysicsControl)
+	void SetAllControlsEnabled(const TArray<FName>& Names, bool bEnable = true);
 
 	/**
 	 * @param Name The name of the control to modify. If blank, then this will access the first 
@@ -500,7 +686,7 @@ public:
 	 * @return     Returns true if the control was found and modified, false if not
 	 */
 	UFUNCTION(BlueprintCallable, Category = PhysicsControl)
-	bool SetControlAutoDisable(const FName Name, bool bAutoDisable);
+	bool SetControlAutoDisable(FName Name, bool bAutoDisable);
 
 	/**
 	 * @param Name     The name of the control to access. If blank, then this will access the first 
@@ -511,7 +697,7 @@ public:
 	 * @return         Returns true if the control was found, false if not
 	 */
 	//UFUNCTION(BlueprintCallable, Category = PhysicsControl)
-	bool GetControl(const FName Name, FPhysicsControl& Control) const;
+	bool GetControl(FName Name, FPhysicsControl& Control) const;
 
 	/**
 	 * @param Name     The name of the control to access. If blank, then this will access the first
@@ -522,7 +708,7 @@ public:
 	 * @return         Returns true if the control was found, false if not
 	 */
 	UFUNCTION(BlueprintCallable, Category = PhysicsControl)
-	bool GetControlData(const FName Name, FPhysicsControlData& ControlData) const;
+	bool GetControlData(FName Name, FPhysicsControlData& ControlData) const;
 
 	/**
 	 * @param Name     The name of the control to access. If blank, then this will access the first
@@ -533,7 +719,7 @@ public:
 	 * @return         Returns true if the control was found, false if not
 	 */
 	UFUNCTION(BlueprintCallable, Category = PhysicsControl)
-	bool GetControlMultipliers(const FName Name, FPhysicsControlMultipliers& ControlMultipliers) const;
+	bool GetControlMultipliers(FName Name, FPhysicsControlMultipliers& ControlMultipliers) const;
 
 	/**
 	 * @param Name     The name of the control to access. If blank, then this will access the first
@@ -544,7 +730,7 @@ public:
 	 * @return         Returns true if the control was found, false if not
 	 */
 	UFUNCTION(BlueprintCallable, Category = PhysicsControl)
-	bool GetControlTarget(const FName Name, FPhysicsControlTarget& ControlTarget) const;
+	bool GetControlTarget(FName Name, FPhysicsControlTarget& ControlTarget) const;
 
 	/**
 	 * @param Name        The name of the control to access. If blank, then this will access the first
@@ -553,7 +739,7 @@ public:
 	 * @return            Returns true if the control is marked to automatically disable after each tick
 	 */
 	UFUNCTION(BlueprintCallable, Category = PhysicsControl)
-	bool GetControlAutoDisable(const FName Name) const;
+	bool GetControlAutoDisable(FName Name) const;
 
 	/**
 	 * @param Name        The name of the control to access. If blank, then this will access the first
@@ -562,28 +748,30 @@ public:
 	 * @return            Returns true if the control is enabled
 	 */
 	UFUNCTION(BlueprintCallable, Category = PhysicsControl)
-	bool GetControlEnabled(const FName Name) const;
+	bool GetControlEnabled(FName Name) const;
 
 	/**
 	 * Makes a new body modifier for mesh components
 	 */
 	UFUNCTION(BlueprintCallable, Category = PhysicsControl)
 	FName MakeBodyModifier(
-		UMeshComponent*       MeshComponent,
-		const FName           BoneName,
-		EPhysicsMovementType  MovementType = EPhysicsMovementType::Simulated, 
-		float                 GravityMultiplier = 1.0f);
+		UMeshComponent*      MeshComponent,
+		FName                BoneName,
+		EPhysicsMovementType MovementType = EPhysicsMovementType::Simulated, 
+		float                GravityMultiplier = 1.0f,
+		bool                 bUseSkeletalAnimation = true);
 
 	/**
 	 * Makes a new body modifier for mesh components
 	 */
 	UFUNCTION(BlueprintCallable, Category = PhysicsControl)
 	bool MakeNamedBodyModifier(
-		const FName           Name,
-		UMeshComponent*       MeshComponent,
-		const FName           BoneName,
-		EPhysicsMovementType  MovementType = EPhysicsMovementType::Simulated,
-		float                 GravityMultiplier = 1.0f);
+		FName                Name,
+		UMeshComponent*      MeshComponent,
+		FName                BoneName,
+		EPhysicsMovementType MovementType = EPhysicsMovementType::Simulated,
+		float                GravityMultiplier = 1.0f,
+		bool                 bUseSkeletalAnimation = true);
 
 	/**
 	 * Makes new body modifiers for skeletal mesh components
@@ -594,7 +782,8 @@ public:
 		FName                   BoneName,
 		bool                    bIncludeSelf,
 		EPhysicsMovementType    MovementType = EPhysicsMovementType::Simulated,
-		float                   GravityMultiplier = 1.0f);
+		float                   GravityMultiplier = 1.0f,
+		bool                    bUseSkeletalAnimation = true);
 
 
 	/**
@@ -621,7 +810,8 @@ public:
 		FPhysicsControlNameArray&                    AllBodyModifiers,
 		const TMap<FName, FPhysicsControlLimbBones>& LimbBones,
 		EPhysicsMovementType                         MovementType = EPhysicsMovementType::Simulated,
-		float                                        GravityMultiplier = 1.0f);
+		float                                        GravityMultiplier = 1.0f,
+		bool                                         bUseSkeletalAnimation = true);
 
 	/**
 	 * Destroys a BodyModifier
@@ -632,36 +822,142 @@ public:
 	 * @return            Returns true if the body modifier was found and destroyed, false if not
 	 */
 	UFUNCTION(BlueprintCallable, Category = PhysicsControl)
-	bool DestroyBodyModifier(const FName Name);
+	bool DestroyBodyModifier(FName Name);
+
+	/**
+	 * Destroys BodyModifiers
+	 *
+	 * @param Names The names of the body modifiers to destroy. Note that if you have these in a FPhysicsControlNameArray
+	 *              then it can be split.
+	 */
+	UFUNCTION(BlueprintCallable, Category = PhysicsControl)
+	void DestroyAllBodyModifiers(const TArray<FName>& Names);
+
+	/**
+	 * Sets the kinematic target transform for a body modifier. 
+	 * 
+	 * @param Name The name of the body modifier to access. If blank, then this will access the first
+	 *             modifier, whatever its name.
+	 * 
+	 * @param KinematicTargetPosition The position to use as the kinematic target of the associated body,
+	 *                                if it is kinematic
+	 * 
+	 * @param KinematicTargetOrientation The orientation to use as the kinematic target of the associated body,
+	 *                                   if it is kinematic
+	 * 
+	 * @param bMakeKinematic If set then the body will be made kinematic. If not set, then it won't be changed.
+	 * 
+	 * @return Returns true if the body modifier was found, false if not
+	 */
+	UFUNCTION(BlueprintCallable, Category = PhysicsControl)
+	bool SetBodyModifierKinematicTarget(
+		FName    Name,
+		FVector  KinematicTargetPosition, 
+		FRotator KinematicTargetOrienation,
+		bool     bMakeKinematic);
 
 	/**
-	 * @param Name        The name of the body modifier to access. If blank, then this will access the first
-	 *                    modifier, whatever its name.
+	 * Sets the movement type for a body modifier
+	 *
+	 * @param Name The name of the body modifier to access. If blank, then this will access the first
+	 *             modifier, whatever its name.
 	 *
-	 * @param bSimulate   Whether to enable/disable simulation on the body
+	 * @param MovementType Whether to enable/disable simulation on the body
 	 * 
+	 * @return Returns true if the body modifier was found, false if not
+	 */
+	UFUNCTION(BlueprintCallable, Category = PhysicsControl)
+	bool SetBodyModifierMovementType(
+		FName                Name,
+		EPhysicsMovementType MovementType = EPhysicsMovementType::Simulated);
+
+	/**
+	 * Sets the movement type for body modifiers
+	 *
+	 * @param Names The names of the body modifiers to access. Note that if you have these in a FPhysicsControlNameArray
+	 *              then it can be split.
+	 *
+	 * @param MovementType Whether to enable/disable simulation on the body
+	 */
+	UFUNCTION(BlueprintCallable, Category = PhysicsControl)
+	void SetAllBodyModifierMovementType(
+		const TArray<FName>& Names,
+		EPhysicsMovementType MovementType = EPhysicsMovementType::Simulated);
+
+	/**
+	 * Sets the gravity multiplier for a body modifier
+	 *
+	 * @param Name The name of the body modifier to access. If blank, then this will access the first
+	 *             modifier, whatever its name.
+	 *
 	 * @param GravityMultiplier The amount of gravity to apply when simulating
 	 * 
-	 * @return            Returns true if the body modifier was found and modified
+	 * @return Returns true if the body modifier was found, false if not
 	 */
 	UFUNCTION(BlueprintCallable, Category = PhysicsControl)
-	bool SetBodyModifier(
-		const FName          Name,
-		EPhysicsMovementType MovementType = EPhysicsMovementType::Simulated, 
+	bool SetBodyModifierGravityMultiplier(
+		FName Name,
+		float GravityMultiplier = 1.0f);
+
+	/**
+	 * Sets the gravity multiplier for body modifiers
+	 *
+	 * @param Names The names of the body modifiers to access. Note that if you have these in a FPhysicsControlNameArray
+	 *              then it can be split.
+	 *
+	 * @param GravityMultiplier The amount of gravity to apply when simulating
+	 */
+	UFUNCTION(BlueprintCallable, Category = PhysicsControl)
+	void SetAllBodyModifierGravityMultipliers(
+		const TArray<FName>& Names,
 		float                GravityMultiplier = 1.0f);
 
+	/**
+	 * Sets whether a body modifier should use skeletal animation for its kinematic targets
+	 *
+	 * @param Name The name of the body modifier to access. If blank, then this will access the first
+	 *             modifier, whatever its name.
+	 *
+	 * @param bUseSkeletalAnimation Whether the kinematic target is specified in the frame of the 
+	 *                              skeletal animation, rather than world space. Only relevant if the
+	 *                              body is part of a skeletal mesh.
+	 * 
+	 * @return Returns true if the body modifier was found, false if not
+	 */
+	UFUNCTION(BlueprintCallable, Category = PhysicsControl)
+	bool SetBodyModifierUseSkeletalAnimation(
+		FName Name,
+		bool  bUseSkeletalAnimation);
+
+	/**
+	 * Sets whether body modifiers should use skeletal animation for their kinematic targets
+	 *
+	 * @param Names The names of the body modifiers to access. Note that if you have these in a FPhysicsControlNameArray
+	 *              then it can be split.
+	 *
+	 * @param bUseSkeletalAnimation Whether the kinematic target is specified in the frame of the 
+	 *                              skeletal animation, rather than world space. Only relevant if the
+	 *                              body is part of a skeletal mesh.
+	 */
+	UFUNCTION(BlueprintCallable, Category = PhysicsControl)
+	void SetAllBodyModifiersUseSkeletalAnimation(
+		const TArray<FName>& Names,
+		bool                 bUseSkeletalAnimation);
+
 public:
 
 	/**
 	 * If the component moves by more than this distance then it is treated as a teleport,
-	 * which prevents velocities being used for a frame. Zero or negative disables.
+	 * which prevents velocities being used for a frame. It is also used as the threshold for
+	 * teleporting when moving kinematic objects. Zero or negative disables.
 	 */
 	UPROPERTY(EditAnywhere, BlueprintReadWrite, Category = Teleport)
 	float TeleportDistanceThreshold;
 
 	/**
 	 * If the component rotates by more than this angle (in degrees) then it is treated as a teleport,
-	 * which prevents velocities being used for a frame. Zero or negative disables.
+	 * which prevents velocities being used for a frame. It is also used as the threshold for
+	 * teleporting when moving kinematic objects. Zero or negative disables.
 	 */
 	UPROPERTY(EditAnywhere, BlueprintReadWrite, Category = Teleport)
 	float TeleportRotationThreshold;
diff --git a/Engine/Plugins/Experimental/PhysicsControl/Source/PhysicsControl/Public/PhysicsControlData.h b/Engine/Plugins/Experimental/PhysicsControl/Source/PhysicsControl/Public/PhysicsControlData.h
--- a/Engine/Plugins/Experimental/PhysicsControl/Source/PhysicsControl/Public/PhysicsControlData.h
+++ b/Engine/Plugins/Experimental/PhysicsControl/Source/PhysicsControl/Public/PhysicsControlData.h
@@ -15,10 +15,11 @@ class UMeshComponent;
 UENUM(BlueprintType)
 enum class EPhysicsMovementType : uint8
 {
-	// Default means that the simulation mode and transform won't be changed when set in a Body Modifier
-	Default,
-	// Kinematic means that the object won't be simulated, but will be moved to the average of any
-	// world-space control targets (if there are any), weighted by strength.
+	// Static means that the object won't be simulated, and it won't be moved according to the
+	// kinematic target set in the Body Modifier (though something else might move it)
+	Static,
+	// Kinematic means that the object won't be simulated, but will be moved according to the
+	// kinematic target set in the Body Modifier.
 	Kinematic,
 	// Simulated means that the object will be controlled by the physics solver
 	Simulated
diff --git a/Engine/Plugins/Experimental/PhysicsControl/Source/PhysicsControl/Public/PhysicsControlLimbData.h b/Engine/Plugins/Experimental/PhysicsControl/Source/PhysicsControl/Public/PhysicsControlLimbData.h
--- a/Engine/Plugins/Experimental/PhysicsControl/Source/PhysicsControl/Public/PhysicsControlLimbData.h
+++ b/Engine/Plugins/Experimental/PhysicsControl/Source/PhysicsControl/Public/PhysicsControlLimbData.h
@@ -54,6 +54,9 @@ struct PHYSICSCONTROL_API FPhysicsControlLimbBones
 
 	/** The names of the bones in the limb */
 	TArray<FName> BoneNames;
+
+	/** Indicates if the first bone in this limb was included due to "IncludeParentBone". */
+	bool bFirstBoneIsAdditional = false;
 };
 
 /**
