Add option to enable/disable collision when creating controls

diff --git a/Engine/Plugins/Experimental/PhysicsControl/Source/PhysicsControl/Private/PhysicsControlComponent.cpp b/Engine/Plugins/Experimental/PhysicsControl/Source/PhysicsControl/Private/PhysicsControlComponent.cpp
--- a/Engine/Plugins/Experimental/PhysicsControl/Source/PhysicsControl/Private/PhysicsControlComponent.cpp
+++ b/Engine/Plugins/Experimental/PhysicsControl/Source/PhysicsControl/Private/PhysicsControlComponent.cpp
@@ -43,20 +43,22 @@ void UPhysicsControlComponent::InitializeComponent()
 //======================================================================================================================
 void UPhysicsControlComponent::BeginDestroy()
 {
-	for (TPair<FName, FPhysicsControlRecord>& PhysicsControlRecordPair : Implementation->PhysicsControlRecords)
+	if (Implementation)
 	{
-		DestroyControl(PhysicsControlRecordPair.Key);
-		FPhysicsControlRecord& Record = PhysicsControlRecordPair.Value;
-		Record.PhysicsControlState.Reset();
-	}
-	Implementation->PhysicsControlRecords.Empty();
+		for (TPair<FName, FPhysicsControlRecord>& PhysicsControlRecordPair : Implementation->PhysicsControlRecords)
+		{
+			DestroyControl(PhysicsControlRecordPair.Key);
+			FPhysicsControlRecord& Record = PhysicsControlRecordPair.Value;
+			Record.PhysicsControlState.Reset();
+		}
+		Implementation->PhysicsControlRecords.Empty();
 
-	for (TPair<FName, FPhysicsBodyModifier>& PhysicsBodyModifierPair : Implementation->PhysicsBodyModifiers)
-	{
-		DestroyBodyModifier(PhysicsBodyModifierPair.Key);
+		for (TPair<FName, FPhysicsBodyModifier>& PhysicsBodyModifierPair : Implementation->PhysicsBodyModifiers)
+		{
+			DestroyBodyModifier(PhysicsBodyModifierPair.Key);
+		}
+		Implementation->PhysicsBodyModifiers.Empty();
 	}
-	Implementation->PhysicsBodyModifiers.Empty();
-
 	Super::BeginDestroy();
 }
 
@@ -120,6 +122,7 @@ void UPhysicsControlComponent::TickComponent(
 			{
 				ConstraintInstance->SetAngularDriveParams(0.0f, 0.0f, 0.0f);
 				ConstraintInstance->SetLinearDriveParams(0.0f, 0.0f, 0.0f);
+				ConstraintInstance->SetDisableCollision(false);
 			}
 			else
 			{
diff --git a/Engine/Plugins/Experimental/PhysicsControl/Source/PhysicsControl/Private/PhysicsControlComponentImpl.cpp b/Engine/Plugins/Experimental/PhysicsControl/Source/PhysicsControl/Private/PhysicsControlComponentImpl.cpp
--- a/Engine/Plugins/Experimental/PhysicsControl/Source/PhysicsControl/Private/PhysicsControlComponentImpl.cpp
+++ b/Engine/Plugins/Experimental/PhysicsControl/Source/PhysicsControl/Private/PhysicsControlComponentImpl.cpp
@@ -583,6 +583,8 @@ void FPhysicsControlComponentImpl::ApplyControl(FPhysicsControlRecord& Record)
 		return;
 	}
 
+	ConstraintInstance->SetDisableCollision(Record.PhysicsControl.ControlSettings.bDisableCollision);
+
 	// Set strengths etc
 	if (ApplyControlStrengths(Record, ConstraintInstance))
 	{
diff --git a/Engine/Plugins/Experimental/PhysicsControl/Source/PhysicsControl/Public/PhysicsControlData.h b/Engine/Plugins/Experimental/PhysicsControl/Source/PhysicsControl/Public/PhysicsControlData.h
--- a/Engine/Plugins/Experimental/PhysicsControl/Source/PhysicsControl/Public/PhysicsControlData.h
+++ b/Engine/Plugins/Experimental/PhysicsControl/Source/PhysicsControl/Public/PhysicsControlData.h
@@ -201,6 +201,7 @@ struct PHYSICSCONTROL_API FPhysicsControlSettings
 		: ControlPoint(ForceInitToZero)
 		, bUseSkeletalAnimation(true)
 		, SkeletalAnimationVelocityMultiplier(1.0f)
+		, bDisableCollision(false)
 		, bAutoDisable(false)
 	{
 	}
@@ -217,6 +218,13 @@ struct PHYSICSCONTROL_API FPhysicsControlSettings
 	UPROPERTY(EditAnywhere, BlueprintReadWrite, Category = PhysicsControl)
 	float SkeletalAnimationVelocityMultiplier;
 
+	/**
+	 * Whether or not this control should disable collision between the parent and child bodies (only
+	 * has an effect if there is a parent body)
+	 */
+	UPROPERTY(EditAnywhere, BlueprintReadWrite, Category = PhysicsControl)
+	bool bDisableCollision;
+
 	/**
 	 * Whether or not this control should automatically disable itself at the end of each tick. This
 	 * can be useful when it is more convenient to have a branch that handles some condition (e.g. character 
