Don't show the ControlOffset in the BP struct

diff --git a/Engine/Plugins/Experimental/PhysicsControl/Source/PhysicsControl/Public/PhysicsControlData.h b/Engine/Plugins/Experimental/PhysicsControl/Source/PhysicsControl/Public/PhysicsControlData.h
--- a/Engine/Plugins/Experimental/PhysicsControl/Source/PhysicsControl/Public/PhysicsControlData.h
+++ b/Engine/Plugins/Experimental/PhysicsControl/Source/PhysicsControl/Public/PhysicsControlData.h
@@ -206,8 +206,11 @@ struct PHYSICSCONTROL_API FPhysicsControlSettings
 	{
 	}
 
-	/** The position of the control point relative to the child mesh. */
-	UPROPERTY(EditAnywhere, BlueprintReadWrite, Category = PhysicsControl)
+	/**
+	 * The position of the control point relative to the child mesh. Note that this can't be authored
+	 * directly here/on creation - it needs to be set after creation in UPhysicsControlComponent::SetControlPoint
+	 */
+	UPROPERTY()
 	FVector ControlPoint;
 
 	/** If true then the target will be applied on top of the skeletal animation (if there is any) */
