Minor updates/fixes to PhysicsControlComponent

diff --git a/Engine/Plugins/Experimental/PhysicsControl/Source/PhysicsControl/Private/PhysicsControlComponent.cpp b/Engine/Plugins/Experimental/PhysicsControl/Source/PhysicsControl/Private/PhysicsControlComponent.cpp
--- a/Engine/Plugins/Experimental/PhysicsControl/Source/PhysicsControl/Private/PhysicsControlComponent.cpp
+++ b/Engine/Plugins/Experimental/PhysicsControl/Source/PhysicsControl/Private/PhysicsControlComponent.cpp
@@ -148,7 +148,7 @@ void UPhysicsControlComponent::TickComponent(
 			Implementation->ApplyKinematicTarget(BodyModifier);
 			break;
 		case EPhysicsMovementType::Simulated:
-			BodyInstance->SetInstanceSimulatePhysics(true);
+			BodyInstance->SetInstanceSimulatePhysics(true, false, true);
 			break;
 		default:
 			UE_LOG(LogPhysicsControlComponent, Warning, TEXT("Invalid movement type %d"), BodyModifier.MovementType);
@@ -762,7 +762,10 @@ bool UPhysicsControlComponent::SetControlData(
 	if (Record)
 	{
 		Record->PhysicsControl.ControlData = ControlData;
-		Record->PhysicsControlState.bEnabled = bEnableControl;
+		if (bEnableControl)
+		{
+			Record->PhysicsControlState.bEnabled = true;
+		}
 		return true;
 	}
 	return false;
@@ -792,7 +795,10 @@ bool UPhysicsControlComponent::SetControlMultipliers(
 	if (Record)
 	{
 		Record->PhysicsControl.ControlMultipliers = ControlMultipliers;
-		Record->PhysicsControlState.bEnabled = bEnableControl;
+		if (bEnableControl)
+		{
+			Record->PhysicsControlState.bEnabled = true;
+		}
 		return true;
 	}
 	return false;
@@ -822,7 +828,10 @@ bool UPhysicsControlComponent::SetControlLinearData(
 		Record->PhysicsControl.ControlData.LinearDampingRatio = DampingRatio;
 		Record->PhysicsControl.ControlData.LinearExtraDamping = ExtraDamping;
 		Record->PhysicsControl.ControlData.MaxForce = MaxForce;
-		Record->PhysicsControlState.bEnabled = bEnableControl;
+		if (bEnableControl)
+		{
+			Record->PhysicsControlState.bEnabled = true;
+		}
 		return true;
 	}
 	return false;
@@ -839,7 +848,10 @@ bool UPhysicsControlComponent::SetControlAngularData(
 		Record->PhysicsControl.ControlData.AngularDampingRatio = DampingRatio;
 		Record->PhysicsControl.ControlData.AngularExtraDamping = ExtraDamping;
 		Record->PhysicsControl.ControlData.MaxTorque = MaxTorque;
-		Record->PhysicsControlState.bEnabled = bEnableControl;
+		if (bEnableControl)
+		{
+			Record->PhysicsControlState.bEnabled = true;
+		}
 		return true;
 	}
 	return false;
@@ -880,7 +892,10 @@ bool UPhysicsControlComponent::SetControlTarget(
 	if (Record)
 	{
 		Record->PhysicsControl.ControlTarget = ControlTarget;
-		Record->PhysicsControlState.bEnabled = bEnableControl;
+		if (bEnableControl)
+		{
+			Record->PhysicsControlState.bEnabled = true;
+		}
 		return true;
 	}
 	return false;
@@ -899,17 +914,17 @@ void UPhysicsControlComponent::SetAllControlTargets(
 }
 
 //======================================================================================================================
-bool UPhysicsControlComponent::SetControlTargetTransform(
-	const FName Name, const FTransform Transform, float VelocityDeltaTime, 
+bool UPhysicsControlComponent::SetControlTargetPositionAndOrientation(
+	const FName Name, const FVector Position, const FRotator Orientation, float VelocityDeltaTime, 
 	bool bEnableControl, bool bApplyControlPointToTarget)
 {
 	FPhysicsControlRecord* Record = Implementation->FindControlRecord(Name);
 	if (Record)
 	{
 		SetControlTargetPosition(
-			Name, Transform.GetTranslation(), VelocityDeltaTime, bEnableControl, bApplyControlPointToTarget);
+			Name, Position, VelocityDeltaTime, bEnableControl, bApplyControlPointToTarget);
 		SetControlTargetOrientation(
-			Name, Transform.GetRotation().Rotator(), VelocityDeltaTime, bEnableControl, bApplyControlPointToTarget);
+			Name, Orientation, VelocityDeltaTime, bEnableControl, bApplyControlPointToTarget);
 	}
 	return false;
 
@@ -934,7 +949,10 @@ bool UPhysicsControlComponent::SetControlTargetPosition(
 		}
 		Record->PhysicsControl.ControlTarget.TargetPosition = Position;
 		Record->PhysicsControl.ControlTarget.bApplyControlPointToTarget = bApplyControlPointToTarget;
-		Record->PhysicsControlState.bEnabled = bEnableControl;
+		if (bEnableControl)
+		{
+			Record->PhysicsControlState.bEnabled = true;
+		}
 		return true;
 	}
 	return false;
@@ -964,7 +982,10 @@ bool UPhysicsControlComponent::SetControlTargetOrientation(
 		}
 		Record->PhysicsControl.ControlTarget.TargetOrientation = Orientation;
 		Record->PhysicsControl.ControlTarget.bApplyControlPointToTarget = bApplyControlPointToTarget;
-		Record->PhysicsControlState.bEnabled = bEnableControl;
+		if (bEnableControl)
+		{
+			Record->PhysicsControlState.bEnabled = true;
+		}
 		return true;
 	}
 	return false;
@@ -1007,7 +1028,10 @@ bool UPhysicsControlComponent::SetControlTargetPoses(
 		Record->PhysicsControl.ControlTarget.TargetOrientation = OrientationQ.Rotator();
 		Record->PhysicsControl.ControlTarget.TargetPosition = Position;
 		Record->PhysicsControl.ControlTarget.bApplyControlPointToTarget = true;
-		Record->PhysicsControlState.bEnabled = bEnableControl;
+		if (bEnableControl)
+		{
+			Record->PhysicsControlState.bEnabled = true;
+		}
 		return true;
 	}
 	return false;
diff --git a/Engine/Plugins/Experimental/PhysicsControl/Source/PhysicsControl/Private/PhysicsControlComponentImpl.cpp b/Engine/Plugins/Experimental/PhysicsControl/Source/PhysicsControl/Private/PhysicsControlComponentImpl.cpp
--- a/Engine/Plugins/Experimental/PhysicsControl/Source/PhysicsControl/Private/PhysicsControlComponentImpl.cpp
+++ b/Engine/Plugins/Experimental/PhysicsControl/Source/PhysicsControl/Private/PhysicsControlComponentImpl.cpp
@@ -70,11 +70,12 @@ FPhysicsControl* FPhysicsControlComponentImpl::FindControl(const FName Name)
 
 //======================================================================================================================
 bool FPhysicsControlComponentImpl::DetectTeleport(
-	const FTransform& OldComponentTM, const FTransform& NewComponentTM) const
+	const FVector& OldPosition, const FQuat& OldOrientation,
+	const FVector& NewPosition, const FQuat& NewOrientation) const
 {
 	if (Owner->TeleportDistanceThreshold > 0)
 	{
-		double Distance = FVector::Distance(OldComponentTM.GetTranslation(), NewComponentTM.GetTranslation());
+		double Distance = FVector::Distance(OldPosition, NewPosition);
 		if (Distance > Owner->TeleportDistanceThreshold)
 		{
 			return true;
@@ -82,7 +83,7 @@ bool FPhysicsControlComponentImpl::DetectTeleport(
 	}
 	if (Owner->TeleportRotationThreshold > 0)
 	{
-		double Radians = OldComponentTM.GetRotation().AngularDistance(NewComponentTM.GetRotation());
+		double Radians = OldOrientation.AngularDistance(NewOrientation);
 		if (FMath::RadiansToDegrees(Radians) > Owner->TeleportRotationThreshold)
 		{
 			return true;
@@ -91,6 +92,11 @@ bool FPhysicsControlComponentImpl::DetectTeleport(
 	return false;
 }
 
+//======================================================================================================================
+bool FPhysicsControlComponentImpl::DetectTeleport(const FTransform& OldTM, const FTransform& NewTM) const
+{
+	return DetectTeleport(OldTM.GetTranslation(), OldTM.GetRotation(), NewTM.GetTranslation(), NewTM.GetRotation());
+}
 
 //======================================================================================================================
 void FPhysicsControlComponentImpl::UpdateCachedSkeletalBoneData(float Dt)
@@ -154,18 +160,20 @@ void FPhysicsControlComponentImpl::ResetControls(bool bKeepControlRecords)
 //======================================================================================================================
 void FPhysicsControlComponentImpl::ApplyKinematicTarget(const FPhysicsBodyModifier& BodyModifier) const
 {
-	FBodyInstance* BodyInstance = UE::PhysicsControlComponent::GetBodyInstance(
-		BodyModifier.MeshComponent, BodyModifier.BoneName);
-	if (!BodyInstance)
-	{
-		return;
-	}
-
 	// Seems like static and skeletal meshes need to be handled differently
-	FTransform TM = BodyInstance->GetUnrealWorldTransform();
-	FTransform KinematicTarget(BodyModifier.KinematicTargetOrientation, BodyModifier.KinematicTargetPosition);
 	if (USkeletalMeshComponent* SkeletalMeshComponent = Cast<USkeletalMeshComponent>(BodyModifier.MeshComponent.Get()))
 	{
+		FBodyInstance* BodyInstance = UE::PhysicsControlComponent::GetBodyInstance(
+			BodyModifier.MeshComponent, BodyModifier.BoneName);
+		if (!BodyInstance)
+		{
+			return;
+		}
+
+		FTransform TM = BodyInstance->GetUnrealWorldTransform();
+		FTransform KinematicTarget = TM;
+		KinematicTarget.SetRotation(BodyModifier.KinematicTargetOrientation);
+		KinematicTarget.SetTranslation(BodyModifier.KinematicTargetPosition);
 		if (BodyModifier.bUseSkeletalAnimation)
 		{
 			FCachedSkeletalMeshData::FBoneData BoneData;
@@ -180,8 +188,14 @@ void FPhysicsControlComponentImpl::ApplyKinematicTarget(const FPhysicsBodyModifi
 	}
 	else
 	{
-		ETeleportType TT = DetectTeleport(TM, KinematicTarget) ? ETeleportType::ResetPhysics : ETeleportType::None;
-		BodyModifier.MeshComponent->SetWorldTransform(KinematicTarget, false, nullptr, TT);
+		FTransform TM = BodyModifier.MeshComponent->GetComponentToWorld();
+		ETeleportType TT = DetectTeleport(
+			TM.GetTranslation(), TM.GetRotation(), 
+			BodyModifier.KinematicTargetPosition, BodyModifier.KinematicTargetOrientation) 
+			? ETeleportType::ResetPhysics : ETeleportType::None;
+		// Note that calling BodyInstance->SetBodyTransform moves the physics, but not the mesh
+		BodyModifier.MeshComponent->SetWorldLocationAndRotation(
+			BodyModifier.KinematicTargetPosition, BodyModifier.KinematicTargetOrientation, false, nullptr, TT);
 	}
 }
 
diff --git a/Engine/Plugins/Experimental/PhysicsControl/Source/PhysicsControl/Private/PhysicsControlComponentImpl.h b/Engine/Plugins/Experimental/PhysicsControl/Source/PhysicsControl/Private/PhysicsControlComponentImpl.h
--- a/Engine/Plugins/Experimental/PhysicsControl/Source/PhysicsControl/Private/PhysicsControlComponentImpl.h
+++ b/Engine/Plugins/Experimental/PhysicsControl/Source/PhysicsControl/Private/PhysicsControlComponentImpl.h
@@ -48,7 +48,15 @@ public:
 	 * @return true if the difference between the old and new TMs exceeds the teleport
 	 * translation/rotation thresholds
 	 */
-	bool DetectTeleport(const FTransform& OldComponentTM, const FTransform& NewComponentTM) const;
+	bool DetectTeleport(const FTransform& OldTM, const FTransform& NewTM) const;
+
+	/**
+	 * @return true if the difference between the old and new TMs exceeds the teleport
+	 * translation/rotation thresholds
+	 */
+	bool DetectTeleport(
+		const FVector& OldPosition, const FQuat& OldOrientation,
+		const FVector& NewPosition, const FQuat& NewOrientation) const;
 
 	/**
 	 * Terminates the underlying physical constraints, resets our internal stored state for each control,
diff --git a/Engine/Plugins/Experimental/PhysicsControl/Source/PhysicsControl/Public/PhysicsControlComponent.h b/Engine/Plugins/Experimental/PhysicsControl/Source/PhysicsControl/Public/PhysicsControlComponent.h
--- a/Engine/Plugins/Experimental/PhysicsControl/Source/PhysicsControl/Public/PhysicsControlComponent.h
+++ b/Engine/Plugins/Experimental/PhysicsControl/Source/PhysicsControl/Public/PhysicsControlComponent.h
@@ -331,7 +331,7 @@ public:
 	 * 
 	 * @param ControlData The new control data
 	 * 
-	 * @param bEnableControl Enables/disables the control
+	 * @param bEnableControl Enables the control if it is currently disabled
 	 * 
 	 * @return Returns true if the control was found and modified, false if not
 	 */
@@ -346,7 +346,7 @@ public:
 	 *
 	 * @param ControlData The new control data
 	 *
-	 * @param bEnableControl Enables/disables the controls
+	 * @param bEnableControl Enables the control if it is currently disabled
 	 */
 	UFUNCTION(BlueprintCallable, Category = PhysicsControl)
 	void SetAllControlDatas(const TArray<FName>& Names, FPhysicsControlData ControlData, bool bEnableControl = true);
@@ -359,7 +359,7 @@ public:
 	 *
 	 * @param ControlMultipliers The new control multipliers
 	 *
-	 * @param bEnableControl Enables/disables the control
+	 * @param bEnableControl Enables the control if it is currently disabled
 	 *
 	 * @return Returns true if the control was found and modified, false if not
 	 */
@@ -374,7 +374,7 @@ public:
 	 *
 	 * @param ControlMultipliers The new control multipliers
 	 *
-	 * @param bEnableControl Enables/disables the control
+	 * @param bEnableControl Enables the control if it is currently disabled
 	 */
 	UFUNCTION(BlueprintCallable, Category = PhysicsControl)
 	void SetAllControlMultipliers(
@@ -396,8 +396,8 @@ public:
 	 *
 	 * @param MaxForce The maximum force used to drive the linear motion. Zero indicates no limit.
 	 *
-	 * @param bEnableControl Enables/disables the control
-	 * 
+	 * @param bEnableControl Enables the control if it is currently disabled
+	 *
 	 * @return Returns true if the control was found and modified, false if not
 	 */
 	UFUNCTION(BlueprintCallable, Category = PhysicsControl)
@@ -423,7 +423,7 @@ public:
 	 *
 	 * @param MaxTorque The maximum torque used to drive the angular motion. Zero indicates no limit.
 	 * 
-	 * @param bEnableControl Enables/disables the control
+	 * @param bEnableControl Enables the control if it is currently disabled
 	 *
 	 * @return Returns true if the control was found and modified, false if not
 	 */
@@ -465,8 +465,8 @@ public:
 	 * 
 	 * @param ControlTarget The new target for the control
 	 * 
-	 * @param bEnableControl Enables/disables the control
-	 * 
+	 * @param bEnableControl Enables the control if it is currently disabled
+	 *
 	 * @return Returns true if the control was found and modified, false if not
 	 */
 	UFUNCTION(BlueprintCallable, Category = PhysicsControl)
@@ -480,7 +480,7 @@ public:
 	 *
 	 * @param ControlTarget The new target for the controls
 	 * 
-	 * @param bEnableControl Enables/disables the controls
+	 * @param bEnableControl Enables the control if it is currently disabled
 	 */
 	UFUNCTION(BlueprintCallable, Category = PhysicsControl)
 	void SetAllControlTargets(
@@ -493,30 +493,32 @@ public:
 	 * Modifies an existing control target - i.e. what it is driving towards, relative to the parent object
 	 *
 	 * @param Name The name of the control to modify. If blank, then this will access the first
-	 *             control, whatever its name.
+	 *        control, whatever its name.
 	 *
-	 * @param Transform The new transform target for the control
+	 * @param Position The new position target for the control
 	 * 
+	 * @param Orientation The new orientation target for the control
+	 *
 	 * @param VelocityDeltaTime If non-zero, the target velocity will be calculated using the current target 
-	 *                          position. If zero, the target velocity will be set to zero.    
+	 *        position. If zero, the target velocity will be set to zero.    
+	 *
+	 * @param bEnableControl Enables the control if it is currently disabled
 	 *
-	 * @param bEnableControl Enables/disables the control
-	 * 
 	 * @param bApplyControlPointToTarget If true, then the target position/orientation is treated as
-	 *                                   a "virtual" object, where the system attempts to move the object
-	 *                                   to match the pose of this "virtual" object that has been placed at
-	 *                                   the target transform. Use this when you want to specify the target
-	 *                                   transform for the object as a whole. If false, then the target transform
-	 *                                   is used as is, and the system drives the control point towards this
-	 *                                   transform.
+	 *        a "virtual" object, where the system attempts to move the object to match the pose of 
+	 *        this "virtual" object that has been placed at the target transform. Use this when you 
+	 *        want to specify the target transform for the object as a whole. If false, then the 
+	 *        target transform is used as is, and the system drives the control point towards this
+	 *        transform.
 	 *
 	 * @return                  Returns true if the control was found and modified, false if not
 	 */
 	UFUNCTION(BlueprintCallable, Category = PhysicsControl)
-	bool SetControlTargetTransform(
+	bool SetControlTargetPositionAndOrientation(
 		FName            Name,
-		const FTransform Transform, 
-		float            VelocityDeltaTime, 
+		const FVector    Position,
+		const FRotator   Orientation,
+		float            VelocityDeltaTime,
 		bool             bEnableControl = true,
 		bool             bApplyControlPointToTarget = true);
 
@@ -531,8 +533,8 @@ public:
 	 * @param VelocityDeltaTime If non-zero, the target velocity will be calculated using the current target 
 	 *                          position. If zero, the target velocity will be set to zero.    
 	 *
-	 * @param bEnableControl    Enables/disables the control
-	 * 
+	 * @param bEnableControl Enables the control if it is currently disabled
+	 *
 	 * @param bApplyControlPointToTarget If true, then the target position/orientation is treated as
 	 *                                   a "virtual" object, where the system attempts to move the object
 	 *                                   to match the pose of this "virtual" object that has been placed at
@@ -562,7 +564,7 @@ public:
 	 * @param AngularVelocityDeltaTime  If non-zero, the target angular velocity will be calculated using the current 
 	 *                                  target position. If zero, the target velocity will be set to zero.
 	 * 
-	 * @param bEnableControl Enables/disables the control
+	 * @param bEnableControl Enables the control if it is currently disabled
 	 *
 	 * @param bApplyControlPointToTarget If true, then the target position/orientation is treated as
 	 *                                   a "virtual" object, where the system attempts to move the object
@@ -600,8 +602,8 @@ public:
 	 * @param VelocityDeltaTime If non-zero, the target velocity will be calculated using the current target
 	 *                          position. If zero, the target velocity will be set to zero.    
 	 *
-	 * @param bEnableControl Enables/disables the control
-	 * 
+	 * @param bEnableControl    Enables the control if it is currently disabled
+	 *
 	 * @return                  Returns true if the control was found and modified, false if not
 	 */
 	UFUNCTION(BlueprintCallable, Category = PhysicsControl)
