Remove the ability to use blank control and body modifier names

diff --git a/Engine/Plugins/Experimental/PhysicsControl/Source/PhysicsControl/Private/PhysicsControlComponentImpl.cpp b/Engine/Plugins/Experimental/PhysicsControl/Source/PhysicsControl/Private/PhysicsControlComponentImpl.cpp
--- a/Engine/Plugins/Experimental/PhysicsControl/Source/PhysicsControl/Private/PhysicsControlComponentImpl.cpp
+++ b/Engine/Plugins/Experimental/PhysicsControl/Source/PhysicsControl/Private/PhysicsControlComponentImpl.cpp
@@ -43,14 +43,6 @@ bool FPhysicsControlComponentImpl::GetBoneData(
 //======================================================================================================================
 FPhysicsControlRecord* FPhysicsControlComponentImpl::FindControlRecord(const FName Name)
 {
-	if (PhysicsControlRecords.IsEmpty())
-	{
-		return nullptr;
-	}
-	if (Name.IsNone())
-	{
-		return &PhysicsControlRecords.CreateIterator().Value();
-	}
 	if (FPhysicsControlRecord* Record = PhysicsControlRecords.Find(Name))
 	{
 		return Record;
@@ -626,18 +618,6 @@ void FPhysicsControlComponentImpl::ApplyControl(FPhysicsControlRecord& Record)
 //======================================================================================================================
 FPhysicsBodyModifier* FPhysicsControlComponentImpl::FindBodyModifier(const FName Name)
 {
-	if (PhysicsBodyModifiers.IsEmpty())
-	{
-		return nullptr;
-	}
-	if (Name.IsNone())
-	{
-		return &PhysicsBodyModifiers.CreateIterator().Value();
-	}
-	if (FPhysicsBodyModifier* Record = PhysicsBodyModifiers.Find(Name))
-	{
-		return Record;
-	}
-	return nullptr;
+	return PhysicsBodyModifiers.Find(Name);
 }
 
diff --git a/Engine/Plugins/Experimental/PhysicsControl/Source/PhysicsControl/Public/PhysicsControlComponent.h b/Engine/Plugins/Experimental/PhysicsControl/Source/PhysicsControl/Public/PhysicsControlComponent.h
--- a/Engine/Plugins/Experimental/PhysicsControl/Source/PhysicsControl/Public/PhysicsControlComponent.h
+++ b/Engine/Plugins/Experimental/PhysicsControl/Source/PhysicsControl/Public/PhysicsControlComponent.h
@@ -56,14 +56,10 @@ public:
 	 * Makes a new control for mesh components
 	 * 
 	 * @param ControlData   Describes the initial strength etc of the new control
-	 * 
 	 * @param ControlTarget Describes the initial target for the new control
-	 * 
 	 * @param ControlSettings General settings for the control
-	 * 
 	 * @param bEnabled      If true then the control will be enabled immediately. If false you will need to call
 	 *                      SetControlEnabled(true) in order to activate it.
-	 *
 	 * @return The name of the new control
 	 */
 	UFUNCTION(BlueprintCallable, Category = PhysicsControl)
@@ -81,14 +77,10 @@ public:
 	 * Makes a new control for mesh components
 	 * 
 	 * @param ControlData   Describes the initial strength etc of the new control
-	 * 
 	 * @param ControlTarget Describes the initial target for the new control
-	 * 
 	 * @param ControlSettings General settings for the control
-	 * 
 	 * @param bEnabled      If true then the control will be enabled immediately. If false you will need to call
 	 *                      SetControlEnabled(true) in order to enable it.
-	 * 
 	 * @return True if a new control was created, false if a control of the specified name already exists
 	 */
 	UFUNCTION(BlueprintCallable, Category = PhysicsControl)
@@ -107,20 +99,13 @@ public:
 	 * Makes a collection of controls controlling a skeletal mesh
 	 * 
 	 * @param SkeletalMeshComponent The skeletal mesh which will have controls
-	 * 
 	 * @param BoneName The name of the bone below which controls should be created. Each bone will be the child in a control
-	 * 
 	 * @param bIncludeSelf Whether or not to include BoneName when creating controls
-	 * 
 	 * @param ControlType What type of control to create. This determines what the parent will be for each control
-	 *
 	 * @param ControlData   Describes the initial strength etc of the new control
-	 * 
 	 * @param ControlSettings General settings for the control
-	 *
 	 * @param bEnabled      If true then the control will be enabled immediately. If false you will need to call
 	 *                      SetControlEnabled(true) in order to enable it.
-	 *
 	 * @return An array of the controls that have been created
 	 */
 	UFUNCTION(BlueprintCallable, Category = PhysicsControl)
@@ -138,11 +123,8 @@ public:
 	 * them with a constraint profile
 	 * 
 	 * @param SkeletalMeshComponent The skeletal mesh which will have controls
-	 * 
 	 * @param BoneName The name of the bone below which controls should be created. Each bone will be the child in a control
-	 * 
 	 * @param bIncludeSelf Whether or not to include BoneName when creating controls
-	 * 
 	 * @param ConstraintProfile The constraint profile to use for initializing the control strength and
 	 *                          damping (etc) parameters. Note that the controls will all be created in "
 	 *                          parent space" - i.e. with each part controlled relative to its parent. The
@@ -153,10 +135,8 @@ public:
 	 *                          instead of slerp. Note also that the joint constraints do not use the animation
 	 *                          velocity as a target, so when creating controls in this way the control settings
 	 *                          will set the skeletal animation velocity multiplier to zero.
-	 *
 	 * @param bEnabled      If true then the control will be enabled immediately. If false you will need to call
 	 *                      SetControlEnabled(true) in order to enable it.
-	 *
 	 * @return An array of the controls that have been created
 	 */
 	UFUNCTION(BlueprintCallable, Category = PhysicsControl)
@@ -171,18 +151,12 @@ public:
 	 * Makes a collection of controls controlling a skeletal mesh
 	 *
 	 * @param SkeletalMeshComponent The skeletal mesh which will have controls
-	 *
 	 * @param BoneNames The names of bones for which controls should be created. Each bone will be the child in a control
-	 *
 	 * @param ControlType What type of control to create. This determines what the parent will be for each control
-	 *
 	 * @param ControlData   Describes the initial strength etc of the new control
-	 *
 	 * @param ControlSettings General settings for the control
-	 *
 	 * @param bEnabled      If true then the control will be enabled immediately. If false you will need to call
 	 *                      SetControlEnabled(true) in order to enable it.
-	 *
 	 * @return An array of the controls that have been created
 	 */
 	UFUNCTION(BlueprintCallable, Category = PhysicsControl)
@@ -199,9 +173,7 @@ public:
 	 * with a constraint profile
 	 *
 	 * @param SkeletalMeshComponent The skeletal mesh which will have controls
-	 *
 	 * @param BoneNames The names of bones for which controls should be created. Each bone will be the child in a control
-	 *
 	 * @param ConstraintProfile The constraint profile to use for initializing the control strength and
 	 *                          damping (etc) parameters. Note that the controls will all be created in "
 	 *                          parent space" - i.e. with each part controlled relative to its parent. The
@@ -212,10 +184,8 @@ public:
 	 *                          instead of slerp. Note also that the joint constraints do not use the animation
 	 *                          velocity as a target, so when creating controls in this way the control settings
 	 *                          will set the skeletal animation velocity multiplier to zero.
-	 *
 	 * @param bEnabled      If true then the control will be enabled immediately. If false you will need to call
 	 *                      SetControlEnabled(true) in order to enable it.
-	 *
 	 * @return An array of the controls that have been created
 	 */
 	UFUNCTION(BlueprintCallable, Category = PhysicsControl)
@@ -229,13 +199,11 @@ public:
 	 * Calculates which bones belong to which limb in a skeletal mesh
 	 * 
 	 * @param SkeletalMeshComponent The skeletal mesh which will be analyzed
-	 * 
 	 * @param LimbSetupData This needs to be filled in with the list of limbs to "discover". Note that the 
 	 *                      limbs should be listed starting at the "leaf" (i.e. outer) parts of the skeleton 
 	 *                      first, typically finishing with the spine. In addition, the spine limb is typically 
 	 *                      specified using the first spine bone, but flagging it to include its parent 
 	 *                      (normally the pelvis). 
-	 * 
 	 * @return A map of limb names to bones
 	 */
 	UFUNCTION(BlueprintCallable, Category = PhysicsControl)
@@ -247,19 +215,13 @@ public:
 	 * Makes a collection of controls controlling a skeletal mesh, grouped together in limbs
 	 *
 	 * @param AllControls A single container for all the controls that have been created
-	 *
 	 * @param LimbBones A map relating the limbs and the bones that they contain. Typically create this 
 	 *                  using GetLimbBonesFromSkeletalMesh
-	 *
 	 * @param ControlType What type of control to create. This determines what the parent will be for each control
-	 *
 	 * @param ControlData Describes the initial strength etc of the new control
-	 *
 	 * @param ControlSettings General settings for the control
-	 *
 	 * @param bEnabled If true then the control will be enabled immediately. If false you will need to call
 	 *                 SetControlEnabled(true) in order to enable it.
-	 *
 	 * @return A map containing the controls for each limb
 	 */
 	UFUNCTION(BlueprintCallable, Category = PhysicsControl)
@@ -276,10 +238,8 @@ public:
 	 * them with a constraint profile
 	 *
 	 * @param AllControls A single container for all the controls that have been created
-	 *
 	 * @param LimbBones A map relating the limbs and the bones that they contain. Typically create this
 	 *                  using GetLimbBonesFromSkeletalMesh
-	 *
 	 * @param ConstraintProfile The constraint profile to use for initializing the control strength and 
 	 *                          damping (etc) parameters. Note that the controls will all be created in "
 	 *                          parent space" - i.e. with each part controlled relative to its parent. The 
@@ -290,10 +250,8 @@ public:
 	 *                          instead of slerp. Note also that the joint constraints do not use the animation 
 	 *                          velocity as a target, so when creating controls in this way the control settings
 	 *                          will set the skeletal animation velocity multiplier to zero.
-	 *
 	 * @param bEnabled If true then the control will be enabled immediately. If false you will need to call
 	 *                 SetControlEnabled(true) in order to enable it.
-	 *
 	 * @return A map containing the controls for each limb
 	 */
 	UFUNCTION(BlueprintCallable, Category = PhysicsControl)
@@ -306,9 +264,7 @@ public:
 	/**
 	 * Destroys a control
 	 *
-	 * @param Name The name of the control to destroy. If blank, then this will destroy the first
-	 *             control, whatever its name.
-	 *
+	 * @param Name The name of the control to destroy. 
 	 * @return     Returns true if the control was found and destroyed, false if not
 	 */
 	UFUNCTION(BlueprintCallable, Category = PhysicsControl)
@@ -326,13 +282,9 @@ public:
 	/**
 	 * Modifies an existing control data - i.e. the strengths etc of the control driving towards the target
 	 *
-	 * @param Name The name of the control to modify. If blank, then this will access the first 
-	 *             control, whatever its name.
-	 * 
+	 * @param Name The name of the control to modify. 
 	 * @param ControlData The new control data
-	 * 
 	 * @param bEnableControl Enables the control if it is currently disabled
-	 * 
 	 * @return Returns true if the control was found and modified, false if not
 	 */
 	UFUNCTION(BlueprintCallable, Category = PhysicsControl)
@@ -343,9 +295,7 @@ public:
 	 *
 	 * @param Names The names of the controls to modify. Note that if you have these in a FPhysicsControlNameArray
 	 *              then it can be split.
-	 *
 	 * @param ControlData The new control data
-	 *
 	 * @param bEnableControl Enables the control if it is currently disabled
 	 */
 	UFUNCTION(BlueprintCallable, Category = PhysicsControl)
@@ -354,13 +304,9 @@ public:
 	/**
 	 * Modifies an existing control data using the multipliers
 	 *
-	 * @param Name The name of the control to modify. If blank, then this will access the first
-	 *             control, whatever its name.
-	 *
+	 * @param Name The name of the control to modify. 
 	 * @param ControlMultipliers The new control multipliers
-	 *
 	 * @param bEnableControl Enables the control if it is currently disabled
-	 *
 	 * @return Returns true if the control was found and modified, false if not
 	 */
 	UFUNCTION(BlueprintCallable, Category = PhysicsControl)
@@ -371,9 +317,7 @@ public:
 	 *
 	 * @param Names The names of the controls to modify. Note that if you have these in a FPhysicsControlNameArray
 	 *              then it can be split.
-	 *
 	 * @param ControlMultipliers The new control multipliers
-	 *
 	 * @param bEnableControl Enables the control if it is currently disabled
 	 */
 	UFUNCTION(BlueprintCallable, Category = PhysicsControl)
@@ -385,19 +329,12 @@ public:
 	/**
 	 * Modifies an existing control's linear data - i.e. the strengths etc of the control driving towards the target
 	 *
-	 * @param Name The name of the control to modify. If blank, then this will access the first
-	 *             control, whatever its name.
-	 *
+	 * @param Name The name of the control to modify. 
 	 * @param Strength The strength used to drive linear motion
-	 *
 	 * @param DampingRatio The amount of damping associated with the linear strength. 1 Results in critically damped motion
-	 *
 	 * @param ExtraDamping The amount of additional linear damping
-	 *
 	 * @param MaxForce The maximum force used to drive the linear motion. Zero indicates no limit.
-	 *
 	 * @param bEnableControl Enables the control if it is currently disabled
-	 *
 	 * @return Returns true if the control was found and modified, false if not
 	 */
 	UFUNCTION(BlueprintCallable, Category = PhysicsControl)
@@ -412,19 +349,12 @@ public:
 	/**
 	 * Modifies an existing control's angular data - i.e. the strengths etc of the control driving towards the target
 	 *
-	 * @param Name The name of the control to modify. If blank, then this will access the first
-	 *             control, whatever its name.
-	 *
+	 * @param Name The name of the control to modify. 
 	 * @param Strength The strength used to drive angular motion
-	 *
 	 * @param DampingRatio The amount of damping associated with the angular strength. 1 Results in critically damped motion
-	 *
 	 * @param ExtraDamping The amount of additional angular damping
-	 *
 	 * @param MaxTorque The maximum torque used to drive the angular motion. Zero indicates no limit.
-	 * 
 	 * @param bEnableControl Enables the control if it is currently disabled
-	 *
 	 * @return Returns true if the control was found and modified, false if not
 	 */
 	UFUNCTION(BlueprintCallable, Category = PhysicsControl)
@@ -439,11 +369,9 @@ public:
 	/**
 	 * Sets the point at which controls will "push" the child object.
 	 * 
-	 * @param Name The name of the control to modify. If blank, then this will access the first
-	 *             control, whatever its name.
-	 * 
+	 * @param Name The name of the control to modify. 
 	 * @param Position The position of the control point on the child mesh object (only relevant if that 
-	 *                 object is in use and is being simulated)
+	 *        object is in use and is being simulated)
 	 */
 	UFUNCTION(BlueprintCallable, Category = PhysicsControl)
 	bool SetControlPoint(FName Name, const FVector Position);
@@ -451,8 +379,7 @@ public:
 	/**
 	 * Resets the control point to the center of mass of the mesh
 	 *
-	 * @param Name    The name of the control to modify. If blank, then this will access the first
-	 *                control, whatever its name.
+	 * @param Name The name of the control to modify. 
 	 */
 	UFUNCTION(BlueprintCallable, Category = PhysicsControl)
 	bool ResetControlPoint(FName Name);
@@ -460,13 +387,9 @@ public:
 	/**
 	 * Modifies an existing control target - i.e. what it is driving towards, relative to the parent object
 	 *
-	 * @param Name The name of the control to modify. If blank, then this will access the first 
-	 *             control, whatever its name.
-	 * 
+	 * @param Name The name of the control to modify. 
 	 * @param ControlTarget The new target for the control
-	 * 
 	 * @param bEnableControl Enables the control if it is currently disabled
-	 *
 	 * @return Returns true if the control was found and modified, false if not
 	 */
 	UFUNCTION(BlueprintCallable, Category = PhysicsControl)
@@ -477,9 +400,7 @@ public:
 	 *
 	 * @param Names The names of the controls to modify. Note that if you have these in a FPhysicsControlNameArray
 	 *              then it can be split.
-	 *
 	 * @param ControlTarget The new target for the controls
-	 * 
 	 * @param bEnableControl Enables the control if it is currently disabled
 	 */
 	UFUNCTION(BlueprintCallable, Category = PhysicsControl)
@@ -492,25 +413,18 @@ public:
 	/**
 	 * Modifies an existing control target - i.e. what it is driving towards, relative to the parent object
 	 *
-	 * @param Name The name of the control to modify. If blank, then this will access the first
-	 *        control, whatever its name.
-	 *
+	 * @param Name The name of the control to modify. 
 	 * @param Position The new position target for the control
-	 * 
 	 * @param Orientation The new orientation target for the control
-	 *
 	 * @param VelocityDeltaTime If non-zero, the target velocity will be calculated using the current target 
 	 *        position. If zero, the target velocity will be set to zero.    
-	 *
 	 * @param bEnableControl Enables the control if it is currently disabled
-	 *
 	 * @param bApplyControlPointToTarget If true, then the target position/orientation is treated as
 	 *        a "virtual" object, where the system attempts to move the object to match the pose of 
 	 *        this "virtual" object that has been placed at the target transform. Use this when you 
 	 *        want to specify the target transform for the object as a whole. If false, then the 
 	 *        target transform is used as is, and the system drives the control point towards this
 	 *        transform.
-	 *
 	 * @return                  Returns true if the control was found and modified, false if not
 	 */
 	UFUNCTION(BlueprintCallable, Category = PhysicsControl)
@@ -525,25 +439,19 @@ public:
 	/**
 	 * Modifies an existing control target - i.e. what it is driving towards, relative to the parent object
 	 *
-	 * @param Name              The name of the control to modify. If blank, then this will access the first
-	 *                          control, whatever its name.
-	 *
+	 * @param Name The name of the control to modify. 
 	 * @param Position          The new position target for the control
-	 * 
 	 * @param VelocityDeltaTime If non-zero, the target velocity will be calculated using the current target 
-	 *                          position. If zero, the target velocity will be set to zero.    
-	 *
+	 *        position. If zero, the target velocity will be set to zero.    
 	 * @param bEnableControl Enables the control if it is currently disabled
-	 *
 	 * @param bApplyControlPointToTarget If true, then the target position/orientation is treated as
-	 *                                   a "virtual" object, where the system attempts to move the object
-	 *                                   to match the pose of this "virtual" object that has been placed at
-	 *                                   the target transform. Use this when you want to specify the target
-	 *                                   transform for the object as a whole. If false, then the target transform
-	 *                                   is used as is, and the system drives the control point towards this
-	 *                                   transform.
-	 *
-	 * @return                  Returns true if the control was found and modified, false if not
+	 *        a "virtual" object, where the system attempts to move the object
+	 *        to match the pose of this "virtual" object that has been placed at
+	 *        the target transform. Use this when you want to specify the target
+	 *        transform for the object as a whole. If false, then the target transform
+	 *        is used as is, and the system drives the control point towards this
+	 *        transform.
+	 * @return Returns true if the control was found and modified, false if not
 	 */
 	UFUNCTION(BlueprintCallable, Category = PhysicsControl)
 	bool SetControlTargetPosition(
@@ -556,25 +464,19 @@ public:
 	/**
 	 * Modifies an existing control target - i.e. what it is driving towards, relative to the parent object
 	 *
-	 * @param Name                      The name of the control to modify. If blank, then this will access the first
-	 *                                  control, whatever its name.
-	 *							        
-	 * @param Orientation               The new orientation target for the control
-	 *
-	 * @param AngularVelocityDeltaTime  If non-zero, the target angular velocity will be calculated using the current 
-	 *                                  target position. If zero, the target velocity will be set to zero.
-	 * 
+	 * @param Name The name of the control to modify. 
+	 * @param Orientation The new orientation target for the control
+	 * @param AngularVelocityDeltaTime If non-zero, the target angular velocity will be calculated using the current 
+	 *        target position. If zero, the target velocity will be set to zero.
 	 * @param bEnableControl Enables the control if it is currently disabled
-	 *
 	 * @param bApplyControlPointToTarget If true, then the target position/orientation is treated as
-	 *                                   a "virtual" object, where the system attempts to move the object
-	 *                                   to match the pose of this "virtual" object that has been placed at
-	 *                                   the target transform. Use this when you want to specify the target
-	 *                                   transform for the object as a whole. If false, then the target transform
-	 *                                   is used as is, and the system drives the control point towards this
-	 *                                   transform.
-	 * 
-	 * @return                          Returns true if the control was found and modified, false if not
+	 *        a "virtual" object, where the system attempts to move the object
+	 *        to match the pose of this "virtual" object that has been placed at
+	 *        the target transform. Use this when you want to specify the target
+	 *        transform for the object as a whole. If false, then the target transform
+	 *        is used as is, and the system drives the control point towards this
+	 *        transform.
+	 * @return Returns true if the control was found and modified, false if not
 	 */
 	UFUNCTION(BlueprintCallable, Category = PhysicsControl)
 	bool SetControlTargetOrientation(
@@ -588,22 +490,14 @@ public:
 	 * Calculates and sets an existing control target. This takes the "virtual" position/orientation of the parent 
 	 * and child and calculates the relative control. Note that this will set bApplyControlPointToTarget to true.
 	 * 
-	 * @param Name              The name of the control to modify. If blank, then this will access the first
-	 *                          control, whatever its name.
-	 *
+	 * @param Name              The name of the control to modify. 
 	 * @param ParentPosition    The virtual/target parent position
-	 * 
 	 * @param ParentOrientation The virtual/target parent orientation
-	 *
 	 * @param ChildPosition     The virtual/target child position
-	 *
 	 * @param ChildOrientation  The virtual/target child orientation
-	 *
 	 * @param VelocityDeltaTime If non-zero, the target velocity will be calculated using the current target
 	 *                          position. If zero, the target velocity will be set to zero.    
-	 *
 	 * @param bEnableControl    Enables the control if it is currently disabled
-	 *
 	 * @return                  Returns true if the control was found and modified, false if not
 	 */
 	UFUNCTION(BlueprintCallable, Category = PhysicsControl)
@@ -620,15 +514,11 @@ public:
 	/**
 	 * Sets whether or not the control should use skeletal animation for the targets
 	 *
-	 * @param Name The name of the control to modify. If blank, then this will access the first
-	 *             control, whatever its name.
-	 * 
+	 * @param Name The name of the control to modify. 
 	 * @param bUseSkeletalAnimation If true then the targets will be a combination of the skeletal animation (if
-	 *                              there is any) and the control target that has been set
-	 * 
-	 * @param SkeletalAnimationVelocityMultiplier If skeletal animation is being used, then this determines the amount of 
-	 *                                            velocity extracted from the animation that is used as targets for the controls
-	 * 
+	 *        there is any) and the control target that has been set
+	 * @param SkeletalAnimationVelocityMultiplier If skeletal animation is being used, then this determines 
+	 *        the amount of velocity extracted from the animation that is used as targets for the controls
 	 * @return Returns true if the control was found and modified, false if not
 	 */
 	UFUNCTION(BlueprintCallable, Category = PhysicsControl)
@@ -642,10 +532,8 @@ public:
 	 *
 	 * @param Names The names of the controls to modify. Note that if you have these in a FPhysicsControlNameArray 
 	 *              then it can be split.
-	 * 
 	 * @param bUseSkeletalAnimation If true then the targets will be a combination of the skeletal animation (if
 	 *                              there is any) and the control target that has been set
-	 * 
 	 * @param SkeletalAnimationVelocityMultiplier If skeletal animation is being used, then this determines the amount of 
 	 *                              velocity extracted from the animation that is used as targets for the controls
 	 */
@@ -658,11 +546,8 @@ public:
 	/**
 	 * Activates or deactivates a control
 	 *
-	 * @param Name     The name of the control to modify. If blank, then this will access the first
-	 *                 control, whatever its name.
-	 *
+	 * @param Name     The name of the control to modify. 
 	 * @param bEnable  The control to enable/disable
-	 *
 	 * @return         Returns true if the control was found and modified, false if not
 	 */
 	UFUNCTION(BlueprintCallable, Category = PhysicsControl)
@@ -673,80 +558,60 @@ public:
 	 *
 	 * @param Names The names of the controls to modify. Note that if you have these in a FPhysicsControlNameArray
 	 *              then it can be split.
-	 *
 	 * @param bEnable  The controls to enable/disable
 	 */
 	UFUNCTION(BlueprintCallable, Category = PhysicsControl)
 	void SetAllControlsEnabled(const TArray<FName>& Names, bool bEnable = true);
 
 	/**
-	 * @param Name The name of the control to modify. If blank, then this will access the first 
-	 *             control, whatever its name.
-	 * 
+	 * @param Name The name of the control to modify. 
 	 * @param bAutoDisable If set then the control will automatically deactivate after each tick.
-	 * 
 	 * @return     Returns true if the control was found and modified, false if not
 	 */
 	UFUNCTION(BlueprintCallable, Category = PhysicsControl)
 	bool SetControlAutoDisable(FName Name, bool bAutoDisable);
 
 	/**
-	 * @param Name     The name of the control to access. If blank, then this will access the first 
-	 *                 control, whatever its name.
-	 *
+	 * @param Name     The name of the control to access. 
 	 * @param Control  The control that will be filled in, if found
-	 *
 	 * @return         Returns true if the control was found, false if not
 	 */
 	//UFUNCTION(BlueprintCallable, Category = PhysicsControl)
 	bool GetControl(FName Name, FPhysicsControl& Control) const;
 
 	/**
-	 * @param Name     The name of the control to access. If blank, then this will access the first
-	 *                 control, whatever its name.
-	 *
+	 * @param Name     The name of the control to access. 
 	 * @param Control  The control data that will be filled in if found
-	 *
 	 * @return         Returns true if the control was found, false if not
 	 */
 	UFUNCTION(BlueprintCallable, Category = PhysicsControl)
 	bool GetControlData(FName Name, FPhysicsControlData& ControlData) const;
 
 	/**
-	 * @param Name     The name of the control to access. If blank, then this will access the first
-	 *                 control, whatever its name.
-	 *
+	 * @param Name     The name of the control to access. 
 	 * @param Control  The control multipliers that will be filled in if found
-	 *
 	 * @return         Returns true if the control was found, false if not
 	 */
 	UFUNCTION(BlueprintCallable, Category = PhysicsControl)
 	bool GetControlMultipliers(FName Name, FPhysicsControlMultipliers& ControlMultipliers) const;
 
 	/**
-	 * @param Name     The name of the control to access. If blank, then this will access the first
-	 *                 control, whatever its name.
-	 *
+	 * @param Name     The name of the control to access. 
 	 * @param Control  The control target, if found
-	 *
 	 * @return         Returns true if the control was found, false if not
 	 */
 	UFUNCTION(BlueprintCallable, Category = PhysicsControl)
 	bool GetControlTarget(FName Name, FPhysicsControlTarget& ControlTarget) const;
 
 	/**
-	 * @param Name        The name of the control to access. If blank, then this will access the first
-	 *                    control, whatever its name.
-	 *
+	 * @param Name        The name of the control to access. 
 	 * @return            Returns true if the control is marked to automatically disable after each tick
 	 */
 	UFUNCTION(BlueprintCallable, Category = PhysicsControl)
 	bool GetControlAutoDisable(FName Name) const;
 
 	/**
-	 * @param Name        The name of the control to access. If blank, then this will access the first
-	 *                    control, whatever its name.
-	 *
+	 * @param Name        The name of the control to access. 
 	 * @return            Returns true if the control is enabled
 	 */
 	UFUNCTION(BlueprintCallable, Category = PhysicsControl)
@@ -792,19 +657,13 @@ public:
 	 * Makes a collection of controls controlling a skeletal mesh, grouped together in limbs
 	 *
 	 * @param AllControls A single container for all the controls that have been created
-	 *
 	 * @param LimbBones A map relating the limbs and the bones that they contain. Typically create this 
 	 *                  using GetLimbBonesFromSkeletalMesh
-	 *
 	 * @param ControlType What type of control to create. This determines what the parent will be for each control
-	 *
 	 * @param ControlData   Describes the initial strength etc of the new control
-	 *
 	 * @param ControlSettings General settings for the control
-	 *
 	 * @param bEnabled      If true then the control will be enabled immediately. If false you will need to call
 	 *                      SetControlEnabled(true) in order to enable it.
-	 *
 	 * @return A map containing the controls for each limb
 	 */
 	UFUNCTION(BlueprintCallable, Category = PhysicsControl)
@@ -818,9 +677,7 @@ public:
 	/**
 	 * Destroys a BodyModifier
 	 *
-	 * @param Name        The name of the body modifier to destroy. If blank, then this will destroy the first
-	 *                    body modifier, whatever its name.
-	 *
+	 * @param Name        The name of the body modifier to destroy. 
 	 * @return            Returns true if the body modifier was found and destroyed, false if not
 	 */
 	UFUNCTION(BlueprintCallable, Category = PhysicsControl)
@@ -838,17 +695,12 @@ public:
 	/**
 	 * Sets the kinematic target transform for a body modifier. 
 	 * 
-	 * @param Name The name of the body modifier to access. If blank, then this will access the first
-	 *             modifier, whatever its name.
-	 * 
+	 * @param Name The name of the body modifier to access. 
 	 * @param KinematicTargetPosition The position to use as the kinematic target of the associated body,
 	 *                                if it is kinematic
-	 * 
 	 * @param KinematicTargetOrientation The orientation to use as the kinematic target of the associated body,
 	 *                                   if it is kinematic
-	 * 
 	 * @param bMakeKinematic If set then the body will be made kinematic. If not set, then it won't be changed.
-	 * 
 	 * @return Returns true if the body modifier was found, false if not
 	 */
 	UFUNCTION(BlueprintCallable, Category = PhysicsControl)
@@ -861,11 +713,8 @@ public:
 	/**
 	 * Sets the movement type for a body modifier
 	 *
-	 * @param Name The name of the body modifier to access. If blank, then this will access the first
-	 *             modifier, whatever its name.
-	 *
+	 * @param Name The name of the body modifier to access. 
 	 * @param MovementType Whether to enable/disable simulation on the body
-	 * 
 	 * @return Returns true if the body modifier was found, false if not
 	 */
 	UFUNCTION(BlueprintCallable, Category = PhysicsControl)
@@ -878,7 +727,6 @@ public:
 	 *
 	 * @param Names The names of the body modifiers to access. Note that if you have these in a FPhysicsControlNameArray
 	 *              then it can be split.
-	 *
 	 * @param MovementType Whether to enable/disable simulation on the body
 	 */
 	UFUNCTION(BlueprintCallable, Category = PhysicsControl)
@@ -889,11 +737,8 @@ public:
 	/**
 	 * Sets the gravity multiplier for a body modifier
 	 *
-	 * @param Name The name of the body modifier to access. If blank, then this will access the first
-	 *             modifier, whatever its name.
-	 *
+	 * @param Name The name of the body modifier to access. 
 	 * @param GravityMultiplier The amount of gravity to apply when simulating
-	 * 
 	 * @return Returns true if the body modifier was found, false if not
 	 */
 	UFUNCTION(BlueprintCallable, Category = PhysicsControl)
@@ -906,7 +751,6 @@ public:
 	 *
 	 * @param Names The names of the body modifiers to access. Note that if you have these in a FPhysicsControlNameArray
 	 *              then it can be split.
-	 *
 	 * @param GravityMultiplier The amount of gravity to apply when simulating
 	 */
 	UFUNCTION(BlueprintCallable, Category = PhysicsControl)
@@ -917,13 +761,10 @@ public:
 	/**
 	 * Sets whether a body modifier should use skeletal animation for its kinematic targets
 	 *
-	 * @param Name The name of the body modifier to access. If blank, then this will access the first
-	 *             modifier, whatever its name.
-	 *
+	 * @param Name The name of the body modifier to access. 
 	 * @param bUseSkeletalAnimation Whether the kinematic target is specified in the frame of the 
 	 *                              skeletal animation, rather than world space. Only relevant if the
 	 *                              body is part of a skeletal mesh.
-	 * 
 	 * @return Returns true if the body modifier was found, false if not
 	 */
 	UFUNCTION(BlueprintCallable, Category = PhysicsControl)
@@ -936,7 +777,6 @@ public:
 	 *
 	 * @param Names The names of the body modifiers to access. Note that if you have these in a FPhysicsControlNameArray
 	 *              then it can be split.
-	 *
 	 * @param bUseSkeletalAnimation Whether the kinematic target is specified in the frame of the 
 	 *                              skeletal animation, rather than world space. Only relevant if the
 	 *                              body is part of a skeletal mesh.
