Simple ability to set the collision type on Body Modifiers

diff --git a/Engine/Plugins/Experimental/PhysicsControl/Source/PhysicsControl/Private/PhysicsControlComponent.cpp b/Engine/Plugins/Experimental/PhysicsControl/Source/PhysicsControl/Private/PhysicsControlComponent.cpp
--- a/Engine/Plugins/Experimental/PhysicsControl/Source/PhysicsControl/Private/PhysicsControlComponent.cpp
+++ b/Engine/Plugins/Experimental/PhysicsControl/Source/PhysicsControl/Private/PhysicsControlComponent.cpp
@@ -155,6 +155,16 @@ void UPhysicsControlComponent::TickComponent(
 			break;
 		}
 
+		UBodySetup* BodySetup = BodyInstance->GetBodySetup();
+		if (BodySetup)
+		{
+			int32 NumShapes = BodySetup->AggGeom.GetElementCount();
+			for (int32 ShapeIndex = 0 ; ShapeIndex != NumShapes ; ++ShapeIndex)
+			{
+				BodyInstance->SetShapeCollisionEnabled(ShapeIndex, BodyModifier.CollisionType);
+			}
+		}
+
 		if (BodyInstance->IsInstanceSimulatingPhysics())
 		{
 			float GravityZ = BodyInstance->GetPhysicsScene()->GetOwningWorld()->GetGravityZ();
@@ -1149,14 +1159,16 @@ bool UPhysicsControlComponent::GetControlEnabled(const FName Name) const
 
 //======================================================================================================================
 FName UPhysicsControlComponent::MakeBodyModifier(
-	UMeshComponent*      MeshComponent,
-	const FName          BoneName,
-	EPhysicsMovementType MovementType, 
-	float                GravityMultiplier,
-	bool                 bUseSkeletalAnimation)
+	UMeshComponent*         MeshComponent,
+	const FName             BoneName,
+	EPhysicsMovementType    MovementType, 
+	ECollisionEnabled::Type CollisionType,
+	float                   GravityMultiplier,
+	bool                    bUseSkeletalAnimation)
 {
 	FName Name = Implementation->GetUniqueBodyModifierName(BoneName);
-	if (MakeNamedBodyModifier(Name, MeshComponent, BoneName, MovementType, GravityMultiplier, bUseSkeletalAnimation))
+	if (MakeNamedBodyModifier(
+		Name, MeshComponent, BoneName, MovementType, CollisionType, GravityMultiplier, bUseSkeletalAnimation))
 	{
 		return Name;
 	}
@@ -1165,12 +1177,13 @@ FName UPhysicsControlComponent::MakeBodyModifier(
 
 //======================================================================================================================
 bool UPhysicsControlComponent::MakeNamedBodyModifier(
-	const FName          Name,
-	UMeshComponent*      MeshComponent,
-	const FName          BoneName,
-	EPhysicsMovementType MovementType,
-	float                GravityMultiplier,
-	bool                 bUseSkeletalAnimation)
+	const FName             Name,
+	UMeshComponent*         MeshComponent,
+	const FName             BoneName,
+	EPhysicsMovementType    MovementType,
+	ECollisionEnabled::Type CollisionType,
+	float                   GravityMultiplier,
+	bool                    bUseSkeletalAnimation)
 {
 	FPhysicsBodyModifier* BodyModifier = Implementation->FindBodyModifier(Name);
 	if (BodyModifier)
@@ -1186,7 +1199,8 @@ bool UPhysicsControlComponent::MakeNamedBodyModifier(
 	}
 
 	FPhysicsBodyModifier& Modifier = Implementation->PhysicsBodyModifiers.Add(
-		Name, FPhysicsBodyModifier(MeshComponent, BoneName, MovementType, GravityMultiplier, bUseSkeletalAnimation));
+		Name, FPhysicsBodyModifier(
+			MeshComponent, BoneName, MovementType, CollisionType, GravityMultiplier, bUseSkeletalAnimation));
 
 	USkeletalMeshComponent* SkeletalMeshComponent = Cast<USkeletalMeshComponent>(MeshComponent);
 	if (SkeletalMeshComponent)
@@ -1201,11 +1215,12 @@ bool UPhysicsControlComponent::MakeNamedBodyModifier(
 //======================================================================================================================
 TArray<FName> UPhysicsControlComponent::MakeBodyModifiersFromSkeletalMeshBelow(
 	USkeletalMeshComponent* SkeletalMeshComponent,
-	FName                BoneName,
-	bool                 bIncludeSelf,
-	EPhysicsMovementType MovementType,
-	float                GravityMultiplier,
-	bool                 bUseSkeletalAnimation)
+	FName                   BoneName,
+	bool                    bIncludeSelf,
+	EPhysicsMovementType    MovementType,
+	ECollisionEnabled::Type CollisionType,
+	float                   GravityMultiplier,
+	bool                    bUseSkeletalAnimation)
 {
 	TArray<FName> Result;
 	UPhysicsAsset* PhysicsAsset = SkeletalMeshComponent ? SkeletalMeshComponent->GetPhysicsAsset() : nullptr;
@@ -1217,7 +1232,7 @@ TArray<FName> UPhysicsControlComponent::MakeBodyModifiersFromSkeletalMeshBelow(
 	SkeletalMeshComponent->ForEachBodyBelow(
 		BoneName, bIncludeSelf, /*bSkipCustomType=*/false,
 		[
-			this, PhysicsAsset, SkeletalMeshComponent, MovementType, 
+			this, PhysicsAsset, SkeletalMeshComponent, MovementType, CollisionType,
 			GravityMultiplier, bUseSkeletalAnimation, &Result
 		](const FBodyInstance* BI)
 		{
@@ -1225,7 +1240,7 @@ TArray<FName> UPhysicsControlComponent::MakeBodyModifiersFromSkeletalMeshBelow(
 			{
 				const FName BoneName = PhysicsAsset->SkeletalBodySetups[BI->InstanceBodyIndex]->BoneName;
 				FName BodyModifierName = MakeBodyModifier(
-					SkeletalMeshComponent, BoneName, MovementType, GravityMultiplier, bUseSkeletalAnimation);
+					SkeletalMeshComponent, BoneName, MovementType, CollisionType, GravityMultiplier, bUseSkeletalAnimation);
 				Result.Add(BodyModifierName);
 			}
 		});
@@ -1238,6 +1253,7 @@ TMap<FName, FPhysicsControlNameArray> UPhysicsControlComponent::MakeBodyModifier
 	FPhysicsControlNameArray&                    AllBodyModifiers,
 	const TMap<FName, FPhysicsControlLimbBones>& LimbBones,
 	EPhysicsMovementType                         MovementType,
+	ECollisionEnabled::Type                      CollisionType,
 	float                                        GravityMultiplier,
 	bool                                         bUseSkeletalAnimation)
 {
@@ -1264,7 +1280,7 @@ TMap<FName, FPhysicsControlNameArray> UPhysicsControlComponent::MakeBodyModifier
 		for (FName BoneName : BonesInLimb.BoneNames)
 		{
 			FName BodyModifierName = MakeBodyModifier(
-				BonesInLimb.SkeletalMeshComponent, BoneName, MovementType, GravityMultiplier, bUseSkeletalAnimation);
+				BonesInLimb.SkeletalMeshComponent, BoneName, MovementType, CollisionType, GravityMultiplier, bUseSkeletalAnimation);
 			if (!BodyModifierName.IsNone())
 			{
 				LimbResult.Names.Add(BodyModifierName);
@@ -1330,7 +1346,7 @@ bool UPhysicsControlComponent::SetBodyModifierKinematicTarget(
 
 //======================================================================================================================
 bool UPhysicsControlComponent::SetBodyModifierMovementType(
-		const FName          Name,
+		FName                Name,
 		EPhysicsMovementType MovementType)
 {
 	FPhysicsBodyModifier* PhysicsBodyModifier = Implementation->FindBodyModifier(Name);
@@ -1353,6 +1369,30 @@ void UPhysicsControlComponent::SetAllBodyModifierMovementType(
 	}
 }
 
+//======================================================================================================================
+bool UPhysicsControlComponent::SetBodyModifierCollisionType(
+	FName                   Name,
+	ECollisionEnabled::Type CollisionType)
+{
+	FPhysicsBodyModifier* PhysicsBodyModifier = Implementation->FindBodyModifier(Name);
+	if (PhysicsBodyModifier)
+	{
+		PhysicsBodyModifier->CollisionType = CollisionType;
+		return true;
+	}
+	return false;
+}
+
+//======================================================================================================================
+void UPhysicsControlComponent::SetAllBodyModifierCollisionType(
+	const TArray<FName>&    Names,
+	ECollisionEnabled::Type CollisionType)
+{
+	for (FName Name : Names)
+	{
+		SetBodyModifierCollisionType(Name, CollisionType);
+	}
+}
 
 //======================================================================================================================
 bool UPhysicsControlComponent::SetBodyModifierGravityMultiplier(
diff --git a/Engine/Plugins/Experimental/PhysicsControl/Source/PhysicsControl/Private/PhysicsControlRecord.h b/Engine/Plugins/Experimental/PhysicsControl/Source/PhysicsControl/Private/PhysicsControlRecord.h
--- a/Engine/Plugins/Experimental/PhysicsControl/Source/PhysicsControl/Private/PhysicsControlRecord.h
+++ b/Engine/Plugins/Experimental/PhysicsControl/Source/PhysicsControl/Private/PhysicsControlRecord.h
@@ -4,6 +4,7 @@
 
 #include "PhysicsControlData.h"
 #include "Components/SkeletalMeshComponent.h"
+#include "Engine/EngineTypes.h"
 
 struct FConstraintInstance;
 struct FBodyInstance;
@@ -65,12 +66,14 @@ struct FPhysicsBodyModifier
 	FPhysicsBodyModifier(
 		TObjectPtr<UMeshComponent> InMeshComponent, 
 		const FName&               InBoneName, 
-		EPhysicsMovementType       InMovementType, 
+		EPhysicsMovementType       InMovementType,
+		ECollisionEnabled::Type    InCollisionType,
 		float                      InGravityMultiplier,
 		bool                       InUseSkeletalAnimation)
 		: MeshComponent(InMeshComponent)
 		, BoneName(InBoneName)
 		, MovementType(InMovementType)
+		, CollisionType(InCollisionType)
 		, GravityMultiplier(InGravityMultiplier)
 		, KinematicTargetPosition(FVector::ZeroVector)
 		, KinematicTargetOrientation(FQuat::Identity)
@@ -87,6 +90,9 @@ struct FPhysicsBodyModifier
 	/** How the associated body should move. */
 	EPhysicsMovementType MovementType;
 
+	/** How the associated body should collide/interact */
+	ECollisionEnabled::Type CollisionType;
+
 	/**
 	 * Multiplier for gravity applied to the body. Note that if the body itself has gravity disabled, then
 	 * setting this to 1 will not enable gravity.
diff --git a/Engine/Plugins/Experimental/PhysicsControl/Source/PhysicsControl/Public/PhysicsControlComponent.h b/Engine/Plugins/Experimental/PhysicsControl/Source/PhysicsControl/Public/PhysicsControlComponent.h
--- a/Engine/Plugins/Experimental/PhysicsControl/Source/PhysicsControl/Public/PhysicsControlComponent.h
+++ b/Engine/Plugins/Experimental/PhysicsControl/Source/PhysicsControl/Public/PhysicsControlComponent.h
@@ -622,23 +622,25 @@ public:
 	 */
 	UFUNCTION(BlueprintCallable, Category = PhysicsControl)
 	FName MakeBodyModifier(
-		UMeshComponent*      MeshComponent,
-		FName                BoneName,
-		EPhysicsMovementType MovementType = EPhysicsMovementType::Simulated, 
-		float                GravityMultiplier = 1.0f,
-		bool                 bUseSkeletalAnimation = true);
+		UMeshComponent*         MeshComponent,
+		FName                   BoneName,
+		EPhysicsMovementType    MovementType = EPhysicsMovementType::Simulated, 
+		ECollisionEnabled::Type CollisionType = ECollisionEnabled::QueryAndPhysics,
+		float                   GravityMultiplier = 1.0f,
+		bool                    bUseSkeletalAnimation = true);
 
 	/**
 	 * Makes a new body modifier for mesh components
 	 */
 	UFUNCTION(BlueprintCallable, Category = PhysicsControl)
 	bool MakeNamedBodyModifier(
-		FName                Name,
-		UMeshComponent*      MeshComponent,
-		FName                BoneName,
-		EPhysicsMovementType MovementType = EPhysicsMovementType::Simulated,
-		float                GravityMultiplier = 1.0f,
-		bool                 bUseSkeletalAnimation = true);
+		FName                   Name,
+		UMeshComponent*         MeshComponent,
+		FName                   BoneName,
+		EPhysicsMovementType    MovementType = EPhysicsMovementType::Simulated,
+		ECollisionEnabled::Type CollisionType = ECollisionEnabled::QueryAndPhysics,
+		float                   GravityMultiplier = 1.0f,
+		bool                    bUseSkeletalAnimation = true);
 
 	/**
 	 * Makes new body modifiers for skeletal mesh components
@@ -649,6 +651,7 @@ public:
 		FName                   BoneName,
 		bool                    bIncludeSelf,
 		EPhysicsMovementType    MovementType = EPhysicsMovementType::Simulated,
+		ECollisionEnabled::Type CollisionType = ECollisionEnabled::QueryAndPhysics,
 		float                   GravityMultiplier = 1.0f,
 		bool                    bUseSkeletalAnimation = true);
 
@@ -671,6 +674,7 @@ public:
 		FPhysicsControlNameArray&                    AllBodyModifiers,
 		const TMap<FName, FPhysicsControlLimbBones>& LimbBones,
 		EPhysicsMovementType                         MovementType = EPhysicsMovementType::Simulated,
+		ECollisionEnabled::Type                      CollisionType = ECollisionEnabled::QueryAndPhysics,
 		float                                        GravityMultiplier = 1.0f,
 		bool                                         bUseSkeletalAnimation = true);
 
@@ -726,14 +730,38 @@ public:
 	 * Sets the movement type for body modifiers
 	 *
 	 * @param Names The names of the body modifiers to access. Note that if you have these in a FPhysicsControlNameArray
-	 *              then it can be split.
-	 * @param MovementType Whether to enable/disable simulation on the body
+	 *        then it can be split.
+	 * @param MovementType Whether to enable/disable simulation on the bodies
 	 */
 	UFUNCTION(BlueprintCallable, Category = PhysicsControl)
 	void SetAllBodyModifierMovementType(
 		const TArray<FName>& Names,
 		EPhysicsMovementType MovementType = EPhysicsMovementType::Simulated);
 
+	/**
+	 * Sets the collision type for a body modifier
+	 *
+	 * @param Name The name of the body modifier to access.
+	 * @param CollisionType Collision type to set on the body
+	 * @return Returns true if the body modifier was found, false if not
+	 */
+	UFUNCTION(BlueprintCallable, Category = PhysicsControl)
+	bool SetBodyModifierCollisionType(
+		FName                   Name,
+		ECollisionEnabled::Type CollisionType = ECollisionEnabled::QueryAndPhysics);
+
+	/**
+	 * Sets the collision type for body modifiers
+	 *
+	 * @param Names The names of the body modifiers to access. Note that if you have these in a FPhysicsControlNameArray
+	 *        then it can be split.
+	 * @param CollisionType Collision type to set on the bodies
+	 */
+	UFUNCTION(BlueprintCallable, Category = PhysicsControl)
+	void SetAllBodyModifierCollisionType(
+		const TArray<FName>&    Names,
+		ECollisionEnabled::Type CollisionType = ECollisionEnabled::QueryAndPhysics);
+
 	/**
 	 * Sets the gravity multiplier for a body modifier
 	 *
