Make bApplyControlPointToTarget in Physics Control Component targets default to false.
Also fixes uninitialised memory warning in PhysicsControlLimbSetupData

diff --git a/Engine/Plugins/Experimental/PhysicsControl/Source/PhysicsControl/Private/PhysicsControlComponentImpl.cpp b/Engine/Plugins/Experimental/PhysicsControl/Source/PhysicsControl/Private/PhysicsControlComponentImpl.cpp
--- a/Engine/Plugins/Experimental/PhysicsControl/Source/PhysicsControl/Private/PhysicsControlComponentImpl.cpp
+++ b/Engine/Plugins/Experimental/PhysicsControl/Source/PhysicsControl/Private/PhysicsControlComponentImpl.cpp
@@ -372,37 +372,14 @@ void FPhysicsControlComponentImpl::CalculateControlTargetData(
 	const FPhysicsControlRecord& Record,
 	bool                         bCalculateVelocity) const
 {
-	const FPhysicsControlTarget& Target = Record.PhysicsControl.ControlTarget;
+	OutTargetTM = FTransform();
+	OutTargetVelocity.Set(0, 0, 0);
+	OutTargetAngularVelocity.Set(0, 0, 0);
 
-	// Calculate the authored target position/orientation - i.e. not using the skeletal animation
-	FQuat TargetOrientationQ = Target.TargetOrientation.Quaternion();
+	bool bUsedSkeletalAnimation = false;
 
-	// Incorporate the offset from the control point
-	FVector ExtraTargetPosition =
-		Record.PhysicsControl.ControlTarget.bApplyControlPointToTarget
-		? Record.PhysicsControl.ControlSettings.ControlPoint
-		: FVector::ZeroVector;
-	FVector ExtraTargetPositionWorld = TargetOrientationQ * ExtraTargetPosition;
-	FVector TargetPosition = Target.TargetPosition + ExtraTargetPositionWorld;
-
-	if (bCalculateVelocity)
-	{
-		// Note that Target.TargetAngularVelocity is in revs per second (as it's user-facing)
-		OutTargetAngularVelocity = Target.TargetAngularVelocity * UE_TWO_PI;
-		FVector ExtraVelocity = OutTargetAngularVelocity.Cross(ExtraTargetPositionWorld);
-		OutTargetVelocity = Target.TargetVelocity + ExtraVelocity;
-	}
-	else
-	{
-		OutTargetVelocity.Set(0, 0, 0);
-		OutTargetAngularVelocity.Set(0, 0, 0);
-	}
-
-	// OutTargetTM is the target transform of the constraint's child frame relative to the
-	// constraint's parent frame
-	OutTargetTM = FTransform(TargetOrientationQ, TargetPosition);
-
-	// Adjust based on any skeletal action
+	// Set the target TM and velocities based on any skeletal action. Note that the targets from animation 
+	// should always account for the control point
 	if (Record.PhysicsControl.ControlSettings.bUseSkeletalAnimation)
 	{
 		FCachedSkeletalMeshData::FBoneData ChildBoneData, ParentBoneData;
@@ -431,22 +408,23 @@ void FPhysicsControlComponentImpl::CalculateControlTargetData(
 		// bodies often use for velocity).
 		if (bHaveChildBoneData)
 		{
+			bUsedSkeletalAnimation = true;
 			FTransform ChildBoneTM = ChildBoneData.GetTM();
 			if (bHaveParentBoneData)
 			{
 				FTransform ParentBoneTM = ParentBoneData.GetTM();
 				FTransform SkeletalDeltaTM = ChildBoneTM * ParentBoneTM.Inverse();
 				// This puts TargetTM in the space of the ParentBone
-				OutTargetTM = OutTargetTM * SkeletalDeltaTM;
+				OutTargetTM = SkeletalDeltaTM;
+
+				// Add on the control point offset
+				OutTargetTM.AddToTranslation(OutTargetTM.GetRotation() * Record.PhysicsControl.ControlSettings.ControlPoint);
 
 				FQuat ParentBoneQ = ParentBoneTM.GetRotation();
 				FQuat ParentBoneInvQ = ParentBoneQ.Inverse();
 
 				if (bCalculateVelocity)
 				{
-					OutTargetVelocity = SkeletalDeltaTM.GetRotation() * OutTargetVelocity;
-					OutTargetAngularVelocity = SkeletalDeltaTM.GetRotation() * OutTargetAngularVelocity;
-
 					if (Record.PhysicsControl.ControlSettings.SkeletalAnimationVelocityMultiplier)
 					{
 						// Offset of the control point from the target child bone TM, in world space.
@@ -480,7 +458,10 @@ void FPhysicsControlComponentImpl::CalculateControlTargetData(
 			}
 			else
 			{
-				OutTargetTM = OutTargetTM * ChildBoneTM;
+				OutTargetTM = ChildBoneTM;
+
+				// Add on the control point offset
+				OutTargetTM.AddToTranslation(OutTargetTM.GetRotation()* Record.PhysicsControl.ControlSettings.ControlPoint);
 
 				if (bCalculateVelocity)
 				{
@@ -507,6 +488,37 @@ void FPhysicsControlComponentImpl::CalculateControlTargetData(
 			}
 		}
 	}
+
+	// Now apply the explicit target specified in the record. It operates in the space of the target
+	// transform we (may have) just calculated.
+	{
+		const FPhysicsControlTarget& Target = Record.PhysicsControl.ControlTarget;
+
+		// Calculate the authored target position/orientation - i.e. not using the skeletal animation
+		FQuat TargetOrientationQ = Target.TargetOrientation.Quaternion();
+		FVector TargetPosition = Target.TargetPosition;
+
+		FVector ExtraTargetPosition(0);
+		// Incorporate the offset from the control point. If we used animation, then we don't need
+		// to do this.
+		if (!bUsedSkeletalAnimation && 
+			Record.PhysicsControl.ControlTarget.bApplyControlPointToTarget)
+		{
+			ExtraTargetPosition = TargetOrientationQ * Record.PhysicsControl.ControlSettings.ControlPoint;
+		}
+
+		// The record's target is specified in the space of the previously calculated/set OutTargetTM
+		OutTargetTM = FTransform(TargetOrientationQ, TargetPosition + ExtraTargetPosition) * OutTargetTM;
+
+		if (bCalculateVelocity)
+		{
+			// Note that Target.TargetAngularVelocity is in revs per second (as it's user-facing)
+			FVector TargetAngularVelocity = Target.TargetAngularVelocity * UE_TWO_PI;
+			OutTargetAngularVelocity += TargetAngularVelocity;
+			FVector ExtraVelocity = TargetAngularVelocity.Cross(ExtraTargetPosition);
+			OutTargetVelocity += Target.TargetVelocity + ExtraVelocity;
+		}
+	}
 }
 
 //======================================================================================================================
diff --git a/Engine/Plugins/Experimental/PhysicsControl/Source/PhysicsControl/Public/PhysicsControlComponent.h b/Engine/Plugins/Experimental/PhysicsControl/Source/PhysicsControl/Public/PhysicsControlComponent.h
--- a/Engine/Plugins/Experimental/PhysicsControl/Source/PhysicsControl/Public/PhysicsControlComponent.h
+++ b/Engine/Plugins/Experimental/PhysicsControl/Source/PhysicsControl/Public/PhysicsControlComponent.h
@@ -434,7 +434,7 @@ public:
 		const FRotator   Orientation,
 		float            VelocityDeltaTime,
 		bool             bEnableControl = true,
-		bool             bApplyControlPointToTarget = true);
+		bool             bApplyControlPointToTarget = false);
 
 	/**
 	 * Modifies an existing control target - i.e. what it is driving towards, relative to the parent object
@@ -459,7 +459,7 @@ public:
 		const FVector Position, 
 		float         VelocityDeltaTime, 
 		bool          bEnableControl = true,
-		bool          bApplyControlPointToTarget = true);
+		bool          bApplyControlPointToTarget = false);
 
 	/**
 	 * Modifies an existing control target - i.e. what it is driving towards, relative to the parent object
@@ -484,7 +484,7 @@ public:
 		const FRotator Orientation, 
 		float          AngularVelocityDeltaTime, 
 		bool           bEnableControl = true,
-		bool           bApplyControlPointToTarget = true);
+		bool           bApplyControlPointToTarget = false);
 
 	/**
 	 * Calculates and sets an existing control target. This takes the "virtual" position/orientation of the parent 
diff --git a/Engine/Plugins/Experimental/PhysicsControl/Source/PhysicsControl/Public/PhysicsControlData.h b/Engine/Plugins/Experimental/PhysicsControl/Source/PhysicsControl/Public/PhysicsControlData.h
--- a/Engine/Plugins/Experimental/PhysicsControl/Source/PhysicsControl/Public/PhysicsControlData.h
+++ b/Engine/Plugins/Experimental/PhysicsControl/Source/PhysicsControl/Public/PhysicsControlData.h
@@ -158,7 +158,7 @@ struct PHYSICSCONTROL_API FPhysicsControlTarget
 		, TargetVelocity(ForceInitToZero)
 		, TargetOrientation(ForceInitToZero)
 		, TargetAngularVelocity(ForceInitToZero)
-		, bApplyControlPointToTarget(true)
+		, bApplyControlPointToTarget(false)
 	{
 	}
 
diff --git a/Engine/Plugins/Experimental/PhysicsControl/Source/PhysicsControl/Public/PhysicsControlLimbData.h b/Engine/Plugins/Experimental/PhysicsControl/Source/PhysicsControl/Public/PhysicsControlLimbData.h
--- a/Engine/Plugins/Experimental/PhysicsControl/Source/PhysicsControl/Public/PhysicsControlLimbData.h
+++ b/Engine/Plugins/Experimental/PhysicsControl/Source/PhysicsControl/Public/PhysicsControlLimbData.h
@@ -38,7 +38,7 @@ struct PHYSICSCONTROL_API FPhysicsControlLimbSetupData
 	 * of spine_01) in the spine limb.
 	 */
 	UPROPERTY(EditAnywhere, BlueprintReadWrite, Category = PhysicsControl)
-	bool bIncludeParentBone;
+	bool bIncludeParentBone = false;
 };
 
 /**
