* Fixed UnrealEditor linux allmodules compile errors due to missing includes (change is adding includes)

diff --git a/Engine/Plugins/Enterprise/DatasmithCADImporter/Source/ParametricSurfaceExtension/Private/ParametricRetessellateAction.cpp b/Engine/Plugins/Enterprise/DatasmithCADImporter/Source/ParametricSurfaceExtension/Private/ParametricRetessellateAction.cpp
--- a/Engine/Plugins/Enterprise/DatasmithCADImporter/Source/ParametricSurfaceExtension/Private/ParametricRetessellateAction.cpp
+++ b/Engine/Plugins/Enterprise/DatasmithCADImporter/Source/ParametricSurfaceExtension/Private/ParametricRetessellateAction.cpp
@@ -8,6 +8,8 @@
 #include "DatasmithStaticMeshImporter.h" // Call to BuildStaticMesh
 #include "DatasmithUtils.h"
 #include "DatasmithTranslator.h"
+#include "Physics/Experimental/PhysScene_Chaos.h"
+#include "Physics/PhysScene.h"
 #include "UI/DatasmithDisplayHelper.h"
 #include "MeshDescriptionHelper.h"
 
diff --git a/Engine/Plugins/Enterprise/DatasmithCADImporter/Source/ParametricSurfaceExtension/Private/ParametricSurfaceBlueprintLibrary.cpp b/Engine/Plugins/Enterprise/DatasmithCADImporter/Source/ParametricSurfaceExtension/Private/ParametricSurfaceBlueprintLibrary.cpp
--- a/Engine/Plugins/Enterprise/DatasmithCADImporter/Source/ParametricSurfaceExtension/Private/ParametricSurfaceBlueprintLibrary.cpp
+++ b/Engine/Plugins/Enterprise/DatasmithCADImporter/Source/ParametricSurfaceExtension/Private/ParametricSurfaceBlueprintLibrary.cpp
@@ -11,6 +11,7 @@
 #include "Chaos/ChaosScene.h"
 #include "Engine/StaticMesh.h"
 #include "IStaticMeshEditor.h"
+#include "Physics/Experimental/PhysScene_Chaos.h"
 #include "Physics/PhysScene.h"
 #include "PhysicsEngine/BodySetup.h"
 #include "StaticMeshAttributes.h"
diff --git a/Engine/Plugins/Experimental/Animation/PoseSearch/Source/Runtime/Private/AnimNode_BlendStack.cpp b/Engine/Plugins/Experimental/Animation/PoseSearch/Source/Runtime/Private/AnimNode_BlendStack.cpp
--- a/Engine/Plugins/Experimental/Animation/PoseSearch/Source/Runtime/Private/AnimNode_BlendStack.cpp
+++ b/Engine/Plugins/Experimental/Animation/PoseSearch/Source/Runtime/Private/AnimNode_BlendStack.cpp
@@ -3,6 +3,7 @@
 #include "PoseSearch/AnimNode_BlendStack.h"
 
 #include "Algo/MaxElement.h"
+#include "Animation/AnimInstanceProxy.h"
 #include "Animation/AnimSequence.h"
 #include "Animation/BlendSpace.h"
 
diff --git a/Engine/Plugins/Experimental/ChaosVehiclesPlugin/Source/ChaosVehicles/Private/ChaosVehicleManager.cpp b/Engine/Plugins/Experimental/ChaosVehiclesPlugin/Source/ChaosVehicles/Private/ChaosVehicleManager.cpp
--- a/Engine/Plugins/Experimental/ChaosVehiclesPlugin/Source/ChaosVehicles/Private/ChaosVehicleManager.cpp
+++ b/Engine/Plugins/Experimental/ChaosVehiclesPlugin/Source/ChaosVehicles/Private/ChaosVehicleManager.cpp
@@ -1,6 +1,7 @@
 // Copyright Epic Games, Inc. All Rights Reserved.
 
 #include "ChaosVehicleManager.h"
+#include "Physics/Experimental/PhysScene_Chaos.h"
 #include "UObject/UObjectIterator.h"
 
 #include "PhysicalMaterials/PhysicalMaterial.h"
diff --git a/Engine/Plugins/Experimental/ChaosVehiclesPlugin/Source/ChaosVehicles/Private/ChaosWheeledVehicleMovementComponent.cpp b/Engine/Plugins/Experimental/ChaosVehiclesPlugin/Source/ChaosVehicles/Private/ChaosWheeledVehicleMovementComponent.cpp
--- a/Engine/Plugins/Experimental/ChaosVehiclesPlugin/Source/ChaosVehicles/Private/ChaosWheeledVehicleMovementComponent.cpp
+++ b/Engine/Plugins/Experimental/ChaosVehiclesPlugin/Source/ChaosVehicles/Private/ChaosWheeledVehicleMovementComponent.cpp
@@ -9,6 +9,7 @@
 #include "DisplayDebugHelpers.h"
 #include "DisplayDebugHelpers.h"
 #include "PhysicalMaterials/PhysicalMaterial.h"
+#include "PhysicsEngine/PhysicsAsset.h"
 #include "VehicleAnimationInstance.h"
 #include "ChaosVehicleManager.h"
 #include "ChaosVehicleWheel.h"
diff --git a/Engine/Plugins/Experimental/PhysicsControl/Source/PhysicsControl/Private/PhysicsControlComponent.cpp b/Engine/Plugins/Experimental/PhysicsControl/Source/PhysicsControl/Private/PhysicsControlComponent.cpp
--- a/Engine/Plugins/Experimental/PhysicsControl/Source/PhysicsControl/Private/PhysicsControlComponent.cpp
+++ b/Engine/Plugins/Experimental/PhysicsControl/Source/PhysicsControl/Private/PhysicsControlComponent.cpp
@@ -1,16 +1,19 @@
 // Copyright Epic Games, Inc. All Rights Reserved.
 
 #include "PhysicsControlComponent.h"
+#include "Physics/Experimental/PhysScene_Chaos.h"
 #include "PhysicsControlComponentLog.h"
 #include "PhysicsControlRecord.h"
 #include "PhysicsControlComponentHelpers.h"
 #include "PhysicsControlComponentImpl.h"
 
+#include "PhysicsEngine/BodySetup.h"
 #include "PhysicsEngine/ConstraintInstance.h"
 #include "PhysicsEngine/BodyInstance.h"
 #include "Physics/PhysicsInterfaceCore.h"
 #include "Components/StaticMeshComponent.h"
 #include "Components/SkeletalMeshComponent.h"
+#include "PhysicsEngine/PhysicsAsset.h"
 
 //======================================================================================================================
 UPhysicsControlComponent::UPhysicsControlComponent(const FObjectInitializer& ObjectInitializer)
diff --git a/Engine/Plugins/Experimental/VirtualCamera/Source/VirtualCamera/Private/VirtualCameraPlayerControllerBase.cpp b/Engine/Plugins/Experimental/VirtualCamera/Source/VirtualCamera/Private/VirtualCameraPlayerControllerBase.cpp
--- a/Engine/Plugins/Experimental/VirtualCamera/Source/VirtualCamera/Private/VirtualCameraPlayerControllerBase.cpp
+++ b/Engine/Plugins/Experimental/VirtualCamera/Source/VirtualCamera/Private/VirtualCameraPlayerControllerBase.cpp
@@ -2,6 +2,7 @@
 
 #include "VirtualCameraPlayerControllerBase.h"
 
+#include "Components/InputComponent.h"
 #include "ConcertVirtualCamera.h"
 #include "CineCameraActor.h"
 #include "Engine/Engine.h"
diff --git a/Engine/Plugins/Importers/USDImporter/Source/USDExporter/Private/SkeletalMeshExporterUSD.cpp b/Engine/Plugins/Importers/USDImporter/Source/USDExporter/Private/SkeletalMeshExporterUSD.cpp
--- a/Engine/Plugins/Importers/USDImporter/Source/USDExporter/Private/SkeletalMeshExporterUSD.cpp
+++ b/Engine/Plugins/Importers/USDImporter/Source/USDExporter/Private/SkeletalMeshExporterUSD.cpp
@@ -2,8 +2,10 @@
 
 #include "SkeletalMeshExporterUSD.h"
 
+#include "Engine/SkinnedAssetCommon.h"
 #include "EngineAnalytics.h"
 #include "MaterialExporterUSD.h"
+#include "Rendering/SkeletalMeshRenderData.h"
 #include "SkeletalMeshExporterUSDOptions.h"
 #include "USDClassesModule.h"
 #include "USDConversionUtils.h"
diff --git a/Engine/Plugins/Importers/USDImporter/Source/USDSchemas/Private/USDSkelRootTranslator.cpp b/Engine/Plugins/Importers/USDImporter/Source/USDSchemas/Private/USDSkelRootTranslator.cpp
--- a/Engine/Plugins/Importers/USDImporter/Source/USDSchemas/Private/USDSkelRootTranslator.cpp
+++ b/Engine/Plugins/Importers/USDImporter/Source/USDSchemas/Private/USDSkelRootTranslator.cpp
@@ -24,12 +24,14 @@
 #include "UsdWrappers/UsdPrim.h"
 
 #include "Animation/AnimBlueprint.h"
+#include "Animation/AnimInstance.h"
 #include "Animation/AnimSequence.h"
 #include "Animation/Skeleton.h"
 #include "AnimationUtils.h"
 #include "Components/SkeletalMeshComponent.h"
 #include "Components/SkinnedMeshComponent.h"
 #include "Engine/SkeletalMesh.h"
+#include "Engine/SkinnedAssetCommon.h"
 #include "GroomComponent.h"
 #include "Materials/Material.h"
 #include "Materials/MaterialInstanceConstant.h"
diff --git a/Engine/Plugins/Importers/USDImporter/Source/USDStageImporter/Private/USDStageImporter.cpp b/Engine/Plugins/Importers/USDImporter/Source/USDStageImporter/Private/USDStageImporter.cpp
--- a/Engine/Plugins/Importers/USDImporter/Source/USDStageImporter/Private/USDStageImporter.cpp
+++ b/Engine/Plugins/Importers/USDImporter/Source/USDStageImporter/Private/USDStageImporter.cpp
@@ -2,6 +2,8 @@
 
 #include "USDStageImporter.h"
 
+#include "Engine/SkinnedAssetCommon.h"
+#include "PhysicsEngine/PhysicsAsset.h"
 #include "USDAssetCache.h"
 #include "USDAssetImportData.h"
 #include "USDClassesModule.h"
diff --git a/Engine/Plugins/Importers/USDImporter/Source/USDUtilities/Private/USDConversionUtils.cpp b/Engine/Plugins/Importers/USDImporter/Source/USDUtilities/Private/USDConversionUtils.cpp
--- a/Engine/Plugins/Importers/USDImporter/Source/USDUtilities/Private/USDConversionUtils.cpp
+++ b/Engine/Plugins/Importers/USDImporter/Source/USDUtilities/Private/USDConversionUtils.cpp
@@ -2,6 +2,7 @@
 
 #include "USDConversionUtils.h"
 
+#include "PhysicsEngine/PhysicsAsset.h"
 #include "USDAssetImportData.h"
 #include "USDDuplicateType.h"
 #include "USDErrorUtils.h"
diff --git a/Engine/Plugins/Importers/USDImporter/Source/USDUtilities/Private/USDPrimConversion.cpp b/Engine/Plugins/Importers/USDImporter/Source/USDUtilities/Private/USDPrimConversion.cpp
--- a/Engine/Plugins/Importers/USDImporter/Source/USDUtilities/Private/USDPrimConversion.cpp
+++ b/Engine/Plugins/Importers/USDImporter/Source/USDUtilities/Private/USDPrimConversion.cpp
@@ -2,6 +2,7 @@
 
 #include "USDPrimConversion.h"
 
+#include "Engine/SkinnedAssetCommon.h"
 #include "UnrealUSDWrapper.h"
 #include "USDAssetImportData.h"
 #include "USDAttributeUtils.h"
diff --git a/Engine/Plugins/Importers/USDImporter/Source/USDUtilities/Private/USDSkeletalDataConversion.cpp b/Engine/Plugins/Importers/USDImporter/Source/USDUtilities/Private/USDSkeletalDataConversion.cpp
--- a/Engine/Plugins/Importers/USDImporter/Source/USDUtilities/Private/USDSkeletalDataConversion.cpp
+++ b/Engine/Plugins/Importers/USDImporter/Source/USDUtilities/Private/USDSkeletalDataConversion.cpp
@@ -2,6 +2,7 @@
 
 #include "USDSkeletalDataConversion.h"
 
+#include "Engine/SkinnedAssetCommon.h"
 #include "UnrealUSDWrapper.h"
 #include "USDAssetImportData.h"
 #include "USDAttributeUtils.h"
diff --git a/Engine/Plugins/Runtime/HairStrands/Source/HairStrandsCore/Private/GroomDeformerBuilder.cpp b/Engine/Plugins/Runtime/HairStrands/Source/HairStrandsCore/Private/GroomDeformerBuilder.cpp
--- a/Engine/Plugins/Runtime/HairStrands/Source/HairStrandsCore/Private/GroomDeformerBuilder.cpp
+++ b/Engine/Plugins/Runtime/HairStrands/Source/HairStrandsCore/Private/GroomDeformerBuilder.cpp
@@ -1,6 +1,7 @@
 // Copyright Epic Games, Inc. All Rights Reserved.
 
 #include "GroomDeformerBuilder.h"
+#include "Engine/SkinnedAssetCommon.h"
 #include "GroomBuilder.h"
 #include "HairStrandsCore.h"
 #include "Animation/Skeleton.h"
diff --git a/Engine/Plugins/Runtime/HairStrands/Source/HairStrandsCore/Private/GroomGeometryCache.cpp b/Engine/Plugins/Runtime/HairStrands/Source/HairStrandsCore/Private/GroomGeometryCache.cpp
--- a/Engine/Plugins/Runtime/HairStrands/Source/HairStrandsCore/Private/GroomGeometryCache.cpp
+++ b/Engine/Plugins/Runtime/HairStrands/Source/HairStrandsCore/Private/GroomGeometryCache.cpp
@@ -11,6 +11,7 @@
 #include "Rendering/SkinWeightVertexBuffer.h"
 #include "CommonRenderResources.h"
 #include "Components/SkeletalMeshComponent.h"
+#include "SkeletalMeshSceneProxy.h"
 #include "SkeletalRenderPublic.h"
 #include "Rendering/SkeletalMeshLODRenderData.h"
 #include "RenderGraphUtils.h"
diff --git a/Engine/Plugins/Runtime/HairStrands/Source/HairStrandsCore/Private/GroomManager.cpp b/Engine/Plugins/Runtime/HairStrands/Source/HairStrandsCore/Private/GroomManager.cpp
--- a/Engine/Plugins/Runtime/HairStrands/Source/HairStrandsCore/Private/GroomManager.cpp
+++ b/Engine/Plugins/Runtime/HairStrands/Source/HairStrandsCore/Private/GroomManager.cpp
@@ -8,6 +8,7 @@
 #include "Rendering/SkinWeightVertexBuffer.h"
 #include "CommonRenderResources.h"
 #include "Components/SkeletalMeshComponent.h"
+#include "SkeletalMeshSceneProxy.h"
 #include "SkeletalRenderPublic.h"
 #include "GroomTextureBuilder.h"
 #include "GroomResources.h"
diff --git a/Engine/Plugins/Runtime/nDisplay/Source/DisplayClusterLightCardEditor/Private/Outliner/DisplayClusterLightCardOutlinerMode.cpp b/Engine/Plugins/Runtime/nDisplay/Source/DisplayClusterLightCardEditor/Private/Outliner/DisplayClusterLightCardOutlinerMode.cpp
--- a/Engine/Plugins/Runtime/nDisplay/Source/DisplayClusterLightCardEditor/Private/Outliner/DisplayClusterLightCardOutlinerMode.cpp
+++ b/Engine/Plugins/Runtime/nDisplay/Source/DisplayClusterLightCardEditor/Private/Outliner/DisplayClusterLightCardOutlinerMode.cpp
@@ -3,6 +3,7 @@
 #include "DisplayClusterLightCardOutlinerMode.h"
 
 #include "DisplayClusterLightCardEditorLog.h"
+#include "Logging/MessageLog.h"
 #include "SDisplayClusterLightCardOutliner.h"
 
 #include "ActorEditorUtils.h"
diff --git a/Engine/Plugins/Tests/EditorTests/Source/EditorTests/Private/UnrealEd/FbxAutomationTests.cpp b/Engine/Plugins/Tests/EditorTests/Source/EditorTests/Private/UnrealEd/FbxAutomationTests.cpp
--- a/Engine/Plugins/Tests/EditorTests/Source/EditorTests/Private/UnrealEd/FbxAutomationTests.cpp
+++ b/Engine/Plugins/Tests/EditorTests/Source/EditorTests/Private/UnrealEd/FbxAutomationTests.cpp
@@ -1,6 +1,7 @@
 // Copyright Epic Games, Inc. All Rights Reserved.
 
 #include "CoreMinimal.h"
+#include "Engine/SkinnedAssetCommon.h"
 #include "HAL/FileManager.h"
 #include "Misc/Paths.h"
 #include "Misc/ConfigCacheIni.h"
diff --git a/Engine/Source/Runtime/Experimental/Chaos/Public/Chaos/Math/Poisson.h b/Engine/Source/Runtime/Experimental/Chaos/Public/Chaos/Math/Poisson.h
--- a/Engine/Source/Runtime/Experimental/Chaos/Public/Chaos/Math/Poisson.h
+++ b/Engine/Source/Runtime/Experimental/Chaos/Public/Chaos/Math/Poisson.h
@@ -8,6 +8,7 @@
 #include "Math/Matrix.h"
 #include "Math/Vector.h"
 #include "Chaos/Utilities.h"
+#include <iostream>
 
 namespace Chaos {
 
