Physics Control Actor and additional functions to manage controls etc in Physics Control Component. Rename from Make to Create

diff --git a/Engine/Plugins/Experimental/PhysicsControl/Config/BasePhysicsControl.ini b/Engine/Plugins/Experimental/PhysicsControl/Config/BasePhysicsControl.ini
new file mode 100644
--- /dev/null
+++ b/Engine/Plugins/Experimental/PhysicsControl/Config/BasePhysicsControl.ini
@@ -0,0 +1,30 @@
+[CoreRedirects]
++StructRedirects=(OldName="/Script/PhysicsControl.PhysicsControlNameArray", NewName="/Script/PhysicsControl.PhysicsControlNames")
++StructRedirects=(OldName="/Script/PhysicsControl.FPhysicsControlMultipliers", NewName="/Script/PhysicsControl.FPhysicsControlMultiplier")
+
++FunctionRedirects=(OldName="PhysicsControlComponent.MakeControl",NewName="PhysicsControlComponent.CreateControl")
++FunctionRedirects=(OldName="PhysicsControlComponent.MakeNamedControl",NewName="PhysicsControlComponent.CreateNamedControl")
++FunctionRedirects=(OldName="PhysicsControlComponent.MakeControlsFromSkeletalMeshBelow",NewName="PhysicsControlComponent.CreateControlsFromSkeletalMeshBelow")
++FunctionRedirects=(OldName="PhysicsControlComponent.MakeControlsFromSkeletalMeshAndConstraintProfileBelow",NewName="PhysicsControlComponent.CreateControlsFromSkeletalMeshAndConstraintProfileBelow")
++FunctionRedirects=(OldName="PhysicsControlComponent.MakeControlsFromSkeletalMesh",NewName="PhysicsControlComponent.CreateControlsFromSkeletalMesh")
++FunctionRedirects=(OldName="PhysicsControlComponent.MakeControlsFromSkeletalMeshAndConstraintProfile",NewName="PhysicsControlComponent.CreateControlsFromSkeletalMeshAndConstraintProfile")
++FunctionRedirects=(OldName="PhysicsControlComponent.MakeControlsFromLimbBones",NewName="PhysicsControlComponent.CreateControlsFromLimbBones")
++FunctionRedirects=(OldName="PhysicsControlComponent.MakeControlsFromLimbBonesAndConstraintProfile",NewName="PhysicsControlComponent.CreateControlsFromLimbBonesAndConstraintProfile")
++FunctionRedirects=(OldName="PhysicsControlComponent.MakeBodyModifier",NewName="PhysicsControlComponent.CreateBodyModifier")
++FunctionRedirects=(OldName="PhysicsControlComponent.MakeNamedBodyModifier",NewName="PhysicsControlComponent.CreateNamedBodyModifier")
++FunctionRedirects=(OldName="PhysicsControlComponent.MakeBodyModifiersFromSkeletalMeshBelow",NewName="PhysicsControlComponent.CreateBodyModifiersFromSkeletalMeshBelow")
++FunctionRedirects=(OldName="PhysicsControlComponent.MakeBodyModifiersFromLimbBones",NewName="PhysicsControlComponent.CreateBodyModifiersFromLimbBones")
++FunctionRedirects=(OldName="PhysicsControlComponent.SetAllBodyModifierMovementType",NewName="PhysicsControlComponent.SetBodyModifiersMovementType")
++FunctionRedirects=(OldName="PhysicsControlComponent.SetAllBodyModifierGravityMultipliers",NewName="PhysicsControlComponent.SetBodyModifiersGravityMultiplier")
++FunctionRedirects=(OldName="PhysicsControlComponent.DestroyAllControls",NewName="PhysicsControlComponent.DestroyControls")
++FunctionRedirects=(OldName="PhysicsControlComponent.SetAllControlTargets",NewName="PhysicsControlComponent.SetControlTargets")
++FunctionRedirects=(OldName="PhysicsControlComponent.SetAllControlDatas",NewName="PhysicsControlComponent.SetControlDatas")
++FunctionRedirects=(OldName="PhysicsControlComponent.SetAllControlsUseSkeletalAnimation",NewName="PhysicsControlComponent.SetControlsUseSkeletalAnimation")
++FunctionRedirects=(OldName="PhysicsControlComponent.SetAllControlsInSetUseSkeletalAnimation",NewName="PhysicsControlComponent.SetControlsInSetUseSkeletalAnimation")
++FunctionRedirects=(OldName="PhysicsControlComponent.SetAllControlsEnabled",NewName="PhysicsControlComponent.SetControlsEnabled")
++FunctionRedirects=(OldName="PhysicsControlComponent.DestroyAllBodyModifiers",NewName="PhysicsControlComponent.DestroyBodyModifiers")
++FunctionRedirects=(OldName="PhysicsControlComponent.SetAllBodyModifiersMovementType",NewName="PhysicsControlComponent.SetBodyModifiersMovementType")
++FunctionRedirects=(OldName="PhysicsControlComponent.SetAllBodyModifiersCollisionType",NewName="PhysicsControlComponent.SetBodyModifiersCollisionType")
++FunctionRedirects=(OldName="PhysicsControlComponent.SetAllBodyModifiersGravityMultiplier",NewName="PhysicsControlComponent.SetBodyModifiersGravityMultiplier")
++FunctionRedirects=(OldName="PhysicsControlComponent.SetAllBodyModifiersUseSkeletalAnimation",NewName="PhysicsControlComponent.SetBodyModifiersUseSkeletalAnimation")
+
diff --git a/Engine/Plugins/Experimental/PhysicsControl/Source/PhysicsControl/Private/PhysicsControlActor.cpp b/Engine/Plugins/Experimental/PhysicsControl/Source/PhysicsControl/Private/PhysicsControlActor.cpp
new file mode 100644
--- /dev/null
+++ b/Engine/Plugins/Experimental/PhysicsControl/Source/PhysicsControl/Private/PhysicsControlActor.cpp
@@ -0,0 +1,259 @@
+// Copyright Epic Games, Inc. All Rights Reserved.
+
+#include "PhysicsControlActor.h"
+#include "PhysicsControlComponent.h"
+#include "PhysicsControlLimbData.h"
+#include "UObject/ConstructorHelpers.h"
+#include "Components/SkeletalMeshComponent.h"
+#include "Components/MeshComponent.h"
+
+#include UE_INLINE_GENERATED_CPP_BY_NAME(PhysicsControlActor)
+
+//======================================================================================================================
+void UPhysicsControlInitializerComponent::CreateInitialCharacterControls(UPhysicsControlComponent* ControlComponent)
+{
+	USkeletalMeshComponent* SkeletalMeshComponent = nullptr;
+
+	if (InitialCharacterControls.CharacterActor)
+	{
+		if (InitialCharacterControls.SkeletalMeshComponentName != NAME_None)
+		{
+			for (UActorComponent* Component : InitialCharacterControls.CharacterActor->GetComponents())
+			{
+				if (Component->GetFName() == InitialCharacterControls.SkeletalMeshComponentName)
+				{
+					SkeletalMeshComponent = Cast<USkeletalMeshComponent>(Component);
+					break;
+				}
+			}
+		}
+		else
+		{
+			SkeletalMeshComponent = InitialCharacterControls.CharacterActor->GetComponentByClass<USkeletalMeshComponent>();
+		}
+	}
+
+	if (SkeletalMeshComponent)
+	{
+		FPhysicsControlNames AllWorldSpaceControls;
+		TMap<FName, FPhysicsControlNames> LimbWorldSpaceControls;
+		FPhysicsControlNames AllParentSpaceControls;
+		TMap<FName, FPhysicsControlNames> LimbParentSpaceControls;
+		TMap<FName, FPhysicsControlNames> LimbBodyModifiers;
+		FPhysicsControlNames AllBodyModifiers;
+
+		ControlComponent->CreateControlsAndBodyModifiersFromLimbBones(
+			AllWorldSpaceControls, LimbWorldSpaceControls,
+			AllParentSpaceControls, LimbParentSpaceControls,
+			AllBodyModifiers, LimbBodyModifiers,
+			SkeletalMeshComponent,
+			InitialCharacterControls.LimbSetupData,
+			InitialCharacterControls.WorldSpaceControlData,
+			InitialCharacterControls.WorldSpaceControlSettings,
+			InitialCharacterControls.bEnableWorldSpaceControls,
+			InitialCharacterControls.ParentSpaceControlData,
+			InitialCharacterControls.ParentSpaceControlSettings,
+			InitialCharacterControls.bEnableParentSpaceControls,
+			InitialCharacterControls.PhysicsMovementType
+		);
+	}
+}
+
+//======================================================================================================================
+void UPhysicsControlInitializerComponent::CreateOrUpdateInitialControls(UPhysicsControlComponent* ControlComponent)
+{
+	// Create any individual controls
+	for (const TPair<FName, FInitialPhysicsControl>& InitialPhysicsControlPair : InitialControls)
+	{
+		const FName Name = InitialPhysicsControlPair.Key;
+		const FInitialPhysicsControl& InitialPhysicsControl = InitialPhysicsControlPair.Value;
+
+		UMeshComponent* ChildMeshComponent = nullptr;
+		if (InitialPhysicsControl.ChildActor)
+		{
+			if (InitialPhysicsControl.ChildMeshComponentName != NAME_None)
+			{
+				for (UActorComponent* Component : InitialPhysicsControl.ChildActor->GetComponents())
+				{
+					if (Component->GetFName() == InitialPhysicsControl.ChildMeshComponentName)
+					{
+						ChildMeshComponent = Cast<UMeshComponent>(Component);
+						break;
+					}
+				}
+			}
+			else
+			{
+				ChildMeshComponent = Cast<UMeshComponent>(InitialPhysicsControl.ChildActor->GetRootComponent());
+			}
+		}
+
+		if (!ChildMeshComponent)
+		{
+			continue;
+		}
+
+		UMeshComponent* ParentMeshComponent = nullptr;
+		if (InitialPhysicsControl.ParentActor)
+		{
+			if (InitialPhysicsControl.ParentMeshComponentName != NAME_None)
+			{
+				for (UActorComponent* Component : InitialPhysicsControl.ParentActor->GetComponents())
+				{
+					if (Component->GetFName() == InitialPhysicsControl.ParentMeshComponentName)
+					{
+						ParentMeshComponent = Cast<UMeshComponent>(Component);
+						break;
+					}
+				}
+			}
+			else
+			{
+				ParentMeshComponent = Cast<UMeshComponent>(InitialPhysicsControl.ParentActor->GetRootComponent());
+			}
+		}
+
+		// If we get here then we will be removing the original and making a new one in its place.
+		TArray<FName> Sets = ControlComponent->GetSetsContainingControl(Name);
+		ControlComponent->DestroyControl(Name);
+
+		if (ControlComponent->CreateNamedControl(
+			Name,
+			ParentMeshComponent,
+			InitialPhysicsControl.ParentBoneName,
+			ChildMeshComponent,
+			InitialPhysicsControl.ChildBoneName,
+			InitialPhysicsControl.ControlData,
+			InitialPhysicsControl.ControlTarget,
+			InitialPhysicsControl.ControlSettings,
+			"All",
+			true))
+		{
+			// If we were replacing a control then add it to the original sets
+			FPhysicsControlNames Names;
+			for (FName Set : Sets)
+			{
+				ControlComponent->AddControlToSet(Names, Name, Set);
+			}
+		}
+	}
+}
+
+//======================================================================================================================
+void UPhysicsControlInitializerComponent::CreateOrUpdateInitialBodyModifiers(UPhysicsControlComponent* ControlComponent)
+{
+	// Create any individual body modifiers
+	for (const TPair<FName, FInitialBodyModifier>& InitialBodyModifierPair : InitialBodyModifiers)
+	{
+		const FName Name = InitialBodyModifierPair.Key;
+		const FInitialBodyModifier& InitialBodyModifier = InitialBodyModifierPair.Value;
+
+		UMeshComponent* MeshComponent = nullptr;
+		if (InitialBodyModifier.Actor)
+		{
+			if (InitialBodyModifier.MeshComponentName != NAME_None)
+			{
+				for (UActorComponent* Component : InitialBodyModifier.Actor->GetComponents())
+				{
+					if (Component->GetFName() == InitialBodyModifier.MeshComponentName)
+					{
+						MeshComponent = Cast<UMeshComponent>(Component);
+						break;
+					}
+				}
+			}
+			else
+			{
+				MeshComponent = Cast<UMeshComponent>(InitialBodyModifier.Actor->GetRootComponent());
+			}
+		}
+
+		if (!MeshComponent)
+		{
+			continue;
+		}
+
+		// If we get here then we will be removing the original and making a new one in its place.
+		TArray<FName> Sets = ControlComponent->GetSetsContainingBodyModifier(Name);
+		ControlComponent->DestroyBodyModifier(Name);
+
+		if (ControlComponent->CreateNamedBodyModifier(
+			Name,
+			MeshComponent,
+			InitialBodyModifier.BoneName,
+			"All",
+			InitialBodyModifier.MovementType,
+			InitialBodyModifier.CollisionType,
+			InitialBodyModifier.GravityMultiplier,
+			InitialBodyModifier.bUseSkeletalAnimation))
+		{
+			// If we were replacing a control then add it to the original sets
+			FPhysicsControlNames Names;
+			for (FName Set : Sets)
+			{
+				ControlComponent->AddBodyModifierToSet(Names, Name, Set);
+			}
+		}
+	}
+}
+
+//======================================================================================================================
+void UPhysicsControlInitializerComponent::CreateControls(UPhysicsControlComponent* PhysicsControlComponent)
+{
+	if (!PhysicsControlComponent)
+	{
+		return;
+	}
+	PhysicsControlComponent->RegisterAllComponentTickFunctions(true);
+
+	CreateInitialCharacterControls(PhysicsControlComponent);
+	CreateOrUpdateInitialControls(PhysicsControlComponent);
+	CreateOrUpdateInitialBodyModifiers(PhysicsControlComponent);
+}
+
+//======================================================================================================================
+void UPhysicsControlInitializerComponent::BeginPlay()
+{
+	Super::BeginPlay();
+	if (bCreateControlsAtBeginPlay)
+	{
+		UPhysicsControlComponent* ControlComponent = GetOwner()->GetComponentByClass<UPhysicsControlComponent>();
+		if (ControlComponent)
+		{
+			CreateControls(ControlComponent);
+		}
+	}
+}
+
+//======================================================================================================================
+UPhysicsControlInitializerComponent::UPhysicsControlInitializerComponent(const FObjectInitializer& ObjectInitializer)
+	: Super(ObjectInitializer)
+{
+}
+
+//======================================================================================================================
+APhysicsControlActor::APhysicsControlActor(const FObjectInitializer& ObjectInitializer)
+	: Super(ObjectInitializer)
+{
+	// Structure to hold one-time initialization
+	struct FConstructorStatics
+	{
+		ConstructorHelpers::FObjectFinderOptional<UTexture2D> KBSJointTexture;
+		FName NAME_Physics;
+		FConstructorStatics()
+			: KBSJointTexture(TEXT("/Engine/EditorResources/S_KBSJoint"))
+			, NAME_Physics(TEXT("Physics"))
+		{
+		}
+	};
+	static FConstructorStatics ConstructorStatics;
+
+	ControlComponent = CreateDefaultSubobject<UPhysicsControlComponent>(TEXT("MyControlComponent"));
+	RootComponent = ControlComponent;
+
+	ControlInitializerComponent = CreateDefaultSubobject<UPhysicsControlInitializerComponent>(TEXT("MyControlInitializerComponent"));
+
+	PrimaryActorTick.bCanEverTick = true;
+
+	SetHidden(true);
+}
diff --git a/Engine/Plugins/Experimental/PhysicsControl/Source/PhysicsControl/Private/PhysicsControlComponent.cpp b/Engine/Plugins/Experimental/PhysicsControl/Source/PhysicsControl/Private/PhysicsControlComponent.cpp
--- a/Engine/Plugins/Experimental/PhysicsControl/Source/PhysicsControl/Private/PhysicsControlComponent.cpp
+++ b/Engine/Plugins/Experimental/PhysicsControl/Source/PhysicsControl/Private/PhysicsControlComponent.cpp
@@ -1,8 +1,6 @@
 // Copyright Epic Games, Inc. All Rights Reserved.
 
 #include "PhysicsControlComponent.h"
-#include "Engine/World.h"
-#include "Physics/Experimental/PhysScene_Chaos.h"
 #include "PhysicsControlComponentLog.h"
 #include "PhysicsControlRecord.h"
 #include "PhysicsControlComponentHelpers.h"
@@ -11,31 +9,52 @@
 #include "PhysicsEngine/BodySetup.h"
 #include "PhysicsEngine/ConstraintInstance.h"
 #include "PhysicsEngine/BodyInstance.h"
+#include "PhysicsEngine/PhysicsAsset.h"
+
+#include "Physics/Experimental/PhysScene_Chaos.h"
 #include "Physics/PhysicsInterfaceCore.h"
+
 #include "Components/StaticMeshComponent.h"
 #include "Components/SkeletalMeshComponent.h"
-#include "PhysicsEngine/PhysicsAsset.h"
+#include "Components/BillboardComponent.h"
+
+#include "Engine/Engine.h"
+#include "Engine/World.h"
+#include "Engine/Texture2D.h"
+
 #include "SceneManagement.h"
 
+//======================================================================================================================
+// This is used, rather than UEnum::GetValueAsString, so that we have more control over the string returned, which 
+// gets used as a prefix for the automatically named sets etc
+static FName GetControlTypeName(EPhysicsControlType ControlType)
+{
+	switch (ControlType)
+	{
+	case EPhysicsControlType::ParentSpace:
+		return "ParentSpace";
+	case EPhysicsControlType::WorldSpace:
+		return "WorldSpace";
+	default:
+		return "None";
+	}
+}
+
 //======================================================================================================================
 UPhysicsControlComponent::UPhysicsControlComponent(const FObjectInitializer& ObjectInitializer)
 	: Super(ObjectInitializer)
 {
 	Implementation = MakePimpl<FPhysicsControlComponentImpl>(this);
 
-	TeleportDistanceThreshold = 300.0f;
-	TeleportRotationThreshold = 0.0f;
-
-	bShowDebugVisualization = true;
-	VisualizationSizeScale = 5.0f;
-	VelocityPredictionTime = 0.2f;
-	MaxNumControlsOrModifiersPerName = 256;
-
 	// ActorComponent setup
 	bWantsInitializeComponent = true;
 	PrimaryComponentTick.bCanEverTick = true;
 	PrimaryComponentTick.bTickEvenWhenPaused = true;
 	PrimaryComponentTick.TickGroup = TG_PrePhysics;
+
+#if WITH_EDITORONLY_DATA
+	bVisualizeComponent = true;
+#endif
 }
 
 //======================================================================================================================
@@ -52,15 +71,15 @@ void UPhysicsControlComponent::BeginDestroy()
 	{
 		for (TPair<FName, FPhysicsControlRecord>& PhysicsControlRecordPair : Implementation->PhysicsControlRecords)
 		{
-			DestroyControl(PhysicsControlRecordPair.Key);
-			FPhysicsControlRecord& Record = PhysicsControlRecordPair.Value;
-			Record.PhysicsControlState.Reset();
+			Implementation->DestroyControl(
+				PhysicsControlRecordPair.Key, FPhysicsControlComponentImpl::EDestroyBehavior::KeepRecord);
 		}
 		Implementation->PhysicsControlRecords.Empty();
 
 		for (TPair<FName, FPhysicsBodyModifier>& PhysicsBodyModifierPair : Implementation->PhysicsBodyModifiers)
 		{
-			DestroyBodyModifier(PhysicsBodyModifierPair.Key);
+			Implementation->DestroyBodyModifier(
+				PhysicsBodyModifierPair.Key, FPhysicsControlComponentImpl::EDestroyBehavior::KeepRecord);
 		}
 		Implementation->PhysicsBodyModifiers.Empty();
 	}
@@ -69,35 +88,11 @@ void UPhysicsControlComponent::BeginDestroy()
 
 //======================================================================================================================
 void UPhysicsControlComponent::TickComponent(
-	float DeltaTime, enum ELevelTick TickType, FActorComponentTickFunction* ThisTickFunction)
+	float                        DeltaTime, 
+	enum ELevelTick              TickType, 
+	FActorComponentTickFunction* ThisTickFunction)
 {
-	// Handle control removal
-	bool bRemovedOneControl = false;
-	for (auto It = Implementation->PhysicsControlRecords.CreateIterator(); It; ++It)
-	{
-		FPhysicsControlRecord& Record = It.Value();
-
-		if (Record.PhysicsControlState.bPendingDestroy)
-		{
-			Record.PhysicsControlState.Reset();
-			It.RemoveCurrent();
-			bRemovedOneControl = true;
-		}
-	}
-	Implementation->PhysicsControlRecords.Compact();
-
-	// Handle body modifier removal
-	bool bRemovedOneBodyModifier = false;
-	for (auto It = Implementation->PhysicsBodyModifiers.CreateIterator(); It; ++It)
-	{
-		FPhysicsBodyModifier& BodyModifier = It.Value();
-		if (BodyModifier.bPendingDestroy)
-		{
-			It.RemoveCurrent();
-			bRemovedOneBodyModifier = true;
-		}
-	}
-	Implementation->PhysicsBodyModifiers.Compact();
+	Super::TickComponent(DeltaTime, TickType, ThisTickFunction);
 
 	// We only want to continue the update if this is a "real" tick that corresponds to updating the
 	// world. We certainly don't want to tick during a pause, because part of the processing involves 
@@ -107,6 +102,10 @@ void UPhysicsControlComponent::TickComponent(
 		return;
 	}
 
+	// Tidy up
+	Implementation->PhysicsControlRecords.Compact();
+	Implementation->PhysicsBodyModifiers.Compact();
+
 	// Update the skeletal mesh caches
 	Implementation->UpdateCachedSkeletalBoneData(DeltaTime);
 
@@ -217,8 +216,7 @@ TMap<FName, FPhysicsControlLimbBones> UPhysicsControlComponent::GetLimbBonesFrom
 		if (LimbSetupData.bIncludeParentBone)
 		{
 			LimbBones.bFirstBoneIsAdditional = true;
-			FName ParentBoneName;
-			ParentBoneName = UE::PhysicsControlComponent::GetPhysicalParentBone(
+			const FName ParentBoneName = UE::PhysicsControlComponent::GetPhysicalParentBone(
 				SkeletalMeshComponent, LimbSetupData.StartBone);
 			if (!ParentBoneName.IsNone())
 			{
@@ -251,20 +249,21 @@ TMap<FName, FPhysicsControlLimbBones> UPhysicsControlComponent::GetLimbBonesFrom
 
 //======================================================================================================================
 // Note - some params passed by value to allow them to be unconnected and enable inline BP editing
-FName UPhysicsControlComponent::MakeControl(
-	UMeshComponent*         ParentMeshComponent,
-	const FName             ParentBoneName,
-	UMeshComponent*         ChildMeshComponent,
-	const FName             ChildBoneName,
-	FPhysicsControlData     ControlData, 
-	FPhysicsControlTarget   ControlTarget, 
-	FPhysicsControlSettings ControlSettings,
-	bool                    bEnabled)
-{
-	FName Name = Implementation->GetUniqueControlName(ParentBoneName, ChildBoneName);
-	if (MakeNamedControl(
+FName UPhysicsControlComponent::CreateControl(
+	UMeshComponent*               ParentMeshComponent,
+	const FName                   ParentBoneName,
+	UMeshComponent*               ChildMeshComponent,
+	const FName                   ChildBoneName,
+	const FPhysicsControlData     ControlData, 
+	const FPhysicsControlTarget   ControlTarget, 
+	const FPhysicsControlSettings ControlSettings,
+	const FName                   Set,
+	const bool                    bEnabled)
+{
+	const FName Name = Implementation->GetUniqueControlName(ParentBoneName, ChildBoneName);
+	if (CreateNamedControl(
 		Name, ParentMeshComponent, ParentBoneName, ChildMeshComponent, ChildBoneName, 
-		ControlData, ControlTarget, ControlSettings, bEnabled))
+		ControlData, ControlTarget, ControlSettings, Set, bEnabled))
 	{
 		return Name;
 	}
@@ -273,19 +272,19 @@ FName UPhysicsControlComponent::MakeControl(
 
 //======================================================================================================================
 // Note - some params passed by value to allow them to be unconnected and enable inline BP editing
-bool UPhysicsControlComponent::MakeNamedControl(
-	const FName             Name, 
-	UMeshComponent*         ParentMeshComponent,
-	const FName             ParentBoneName,
-	UMeshComponent*         ChildMeshComponent,
-	const FName             ChildBoneName,
-	FPhysicsControlData     ControlData, 
-	FPhysicsControlTarget   ControlTarget, 
-	FPhysicsControlSettings ControlSettings,
-	bool                    bEnabled)
-{
-	FPhysicsControlRecord* PhysicsControlRecord = Implementation->FindControlRecord(Name);
-	if (PhysicsControlRecord)
+bool UPhysicsControlComponent::CreateNamedControl(
+	const FName                   Name, 
+	UMeshComponent*               ParentMeshComponent,
+	const FName                   ParentBoneName,
+	UMeshComponent*               ChildMeshComponent,
+	const FName                   ChildBoneName,
+	const FPhysicsControlData     ControlData, 
+	const FPhysicsControlTarget   ControlTarget, 
+	const FPhysicsControlSettings ControlSettings,
+	const FName                   Set,
+	const bool                    bEnabled)
+{
+	if (Implementation->FindControlRecord(Name))
 	{
 		return false;
 	}
@@ -313,19 +312,22 @@ bool UPhysicsControlComponent::MakeNamedControl(
 	NewRecord.PhysicsControlState.bEnabled = bEnabled;
 	NewRecord.ResetControlPoint();
 
+	Implementation->NameRecords.AddControl(Name, Set);
+
 	return true;
 }
 
 //======================================================================================================================
 // Note - some params passed by value to allow them to be unconnected and enable inline BP editing
-TArray<FName> UPhysicsControlComponent::MakeControlsFromSkeletalMeshBelow(
+TArray<FName> UPhysicsControlComponent::CreateControlsFromSkeletalMeshBelow(
 	USkeletalMeshComponent* SkeletalMeshComponent,
-	FName                   BoneName,
-	bool                    bIncludeSelf,
-	EPhysicsControlType     ControlType,
-	FPhysicsControlData     ControlData,
-	FPhysicsControlSettings ControlSettings,
-	bool                    bEnabled)
+	const FName                   BoneName,
+	const bool                    bIncludeSelf,
+	const EPhysicsControlType     ControlType,
+	const FPhysicsControlData     ControlData,
+	const FPhysicsControlSettings ControlSettings,
+	const FName                   Set,
+	const bool                    bEnabled)
 {
 	TArray<FName> Result;
 	UPhysicsAsset* PhysicsAsset = SkeletalMeshComponent ? SkeletalMeshComponent->GetPhysicsAsset() : nullptr;
@@ -343,7 +345,7 @@ TArray<FName> UPhysicsControlComponent::MakeControlsFromSkeletalMeshBelow(
 		BoneName, bIncludeSelf, /*bSkipCustomType=*/false, 
 		[
 			this, PhysicsAsset, ParentMeshComponent, SkeletalMeshComponent, 
-			&ControlData, &ControlSettings, &Result, bEnabled
+			ControlType, &ControlData, &ControlSettings, Set, &Result, bEnabled
 		](const FBodyInstance* BI)
 		{
 			if (USkeletalBodySetup* BodySetup = Cast<USkeletalBodySetup>(BI->BodySetup.Get()))
@@ -360,12 +362,15 @@ TArray<FName> UPhysicsControlComponent::MakeControlsFromSkeletalMeshBelow(
 						return;
 					}
 				}
-				FName ControlName = MakeControl(
+				const FName ControlName = CreateControl(
 					ParentMeshComponent, ParentBoneName, SkeletalMeshComponent, ChildBoneName,
-					ControlData, FPhysicsControlTarget(), ControlSettings, bEnabled);
+					ControlData, FPhysicsControlTarget(), ControlSettings, 
+					FName(GetControlTypeName(ControlType).ToString().Append("_").Append(Set.ToString())), 
+					bEnabled);
 				if (!ControlName.IsNone())
 				{
 					Result.Add(ControlName);
+					Implementation->NameRecords.AddControl(ControlName, GetControlTypeName(ControlType));
 				}
 				else
 				{
@@ -380,12 +385,13 @@ TArray<FName> UPhysicsControlComponent::MakeControlsFromSkeletalMeshBelow(
 
 //======================================================================================================================
 // Note - some params passed by value to allow them to be unconnected and enable inline BP editing
-TArray<FName> UPhysicsControlComponent::MakeControlsFromSkeletalMeshAndConstraintProfileBelow(
+TArray<FName> UPhysicsControlComponent::CreateControlsFromSkeletalMeshAndConstraintProfileBelow(
 	USkeletalMeshComponent* SkeletalMeshComponent,
-	FName                   BoneName,
-	bool                    bIncludeSelf,
-	FName                   ConstraintProfile,
-	bool                    bEnabled)
+	const FName             BoneName,
+	const bool              bIncludeSelf,
+	const FName             ConstraintProfile,
+	const FName             Set,
+	const bool              bEnabled)
 {
 	TArray<FName> Result;
 	UPhysicsAsset* PhysicsAsset = SkeletalMeshComponent ? SkeletalMeshComponent->GetPhysicsAsset() : nullptr;
@@ -404,7 +410,7 @@ TArray<FName> UPhysicsControlComponent::MakeControlsFromSkeletalMeshAndConstrain
 		BoneName, bIncludeSelf, /*bSkipCustomType=*/false,
 		[
 			this, PhysicsAsset, SkeletalMeshComponent,
-			ConstraintProfile, &ControlSettings, &Result, bEnabled
+			ConstraintProfile, Set, &ControlSettings, &Result, bEnabled
 		](const FBodyInstance* BI)
 		{
 			if (USkeletalBodySetup* BodySetup = Cast<USkeletalBodySetup>(BI->BodySetup.Get()))
@@ -424,18 +430,23 @@ TArray<FName> UPhysicsControlComponent::MakeControlsFromSkeletalMeshAndConstrain
 				if (!SkeletalMeshComponent->GetConstraintProfilePropertiesOrDefault(
 					ProfileProperties, ChildBoneName, ConstraintProfile))
 				{
-					UE_LOG(LogPhysicsControlComponent, Warning, TEXT("Failed get constraint profile for %s"), *ChildBoneName.ToString());
+					UE_LOG(LogPhysicsControlComponent, Warning, 
+						TEXT("Failed get constraint profile for %s"), *ChildBoneName.ToString());
 					return;
 				}
 
 				UE::PhysicsControlComponent::ConvertConstraintProfileToControlData(ControlData, ProfileProperties);
 
-				FName ControlName = MakeControl(
+				const FName ControlName = CreateControl(
 					SkeletalMeshComponent, ParentBoneName, SkeletalMeshComponent, ChildBoneName,
-					ControlData, FPhysicsControlTarget(), ControlSettings, bEnabled);
+					ControlData, FPhysicsControlTarget(), ControlSettings, 
+					FName(GetControlTypeName(EPhysicsControlType::ParentSpace).ToString().Append("_").Append(Set.ToString())),
+					bEnabled);
 				if (!ControlName.IsNone())
 				{
 					Result.Add(ControlName);
+					Implementation->NameRecords.AddControl(
+						ControlName, GetControlTypeName(EPhysicsControlType::ParentSpace));
 				}
 				else
 				{
@@ -450,13 +461,14 @@ TArray<FName> UPhysicsControlComponent::MakeControlsFromSkeletalMeshAndConstrain
 
 //======================================================================================================================
 // Note - some params passed by value to allow them to be unconnected and enable inline BP editing
-TArray<FName> UPhysicsControlComponent::MakeControlsFromSkeletalMesh(
-	USkeletalMeshComponent* SkeletalMeshComponent,
-	const TArray<FName>&    BoneNames,
-	EPhysicsControlType     ControlType,
-	FPhysicsControlData     ControlData,
-	FPhysicsControlSettings ControlSettings,
-	bool                    bEnabled)
+TArray<FName> UPhysicsControlComponent::CreateControlsFromSkeletalMesh(
+	USkeletalMeshComponent*       SkeletalMeshComponent,
+	const TArray<FName>&          BoneNames,
+	const EPhysicsControlType     ControlType,
+	const FPhysicsControlData     ControlData,
+	const FPhysicsControlSettings ControlSettings,
+	const FName                   Set,
+	const bool                    bEnabled)
 {
 	TArray<FName> Result;
 	UPhysicsAsset* PhysicsAsset = SkeletalMeshComponent ? SkeletalMeshComponent->GetPhysicsAsset() : nullptr;
@@ -483,12 +495,15 @@ TArray<FName> UPhysicsControlComponent::MakeControlsFromSkeletalMesh(
 				continue;
 			}
 		}
-		FName ControlName = MakeControl(
+		const FName ControlName = CreateControl(
 			ParentMeshComponent, ParentBoneName, SkeletalMeshComponent, ChildBoneName,
-			ControlData, FPhysicsControlTarget(), ControlSettings, bEnabled);
+			ControlData, FPhysicsControlTarget(), ControlSettings, 
+			FName(GetControlTypeName(ControlType).ToString().Append("_").Append(Set.ToString())),
+			bEnabled);
 		if (!ControlName.IsNone())
 		{
 			Result.Add(ControlName);
+			Implementation->NameRecords.AddControl(ControlName, GetControlTypeName(ControlType));
 		}
 		else
 		{
@@ -501,11 +516,12 @@ TArray<FName> UPhysicsControlComponent::MakeControlsFromSkeletalMesh(
 
 //======================================================================================================================
 // Note - some params passed by value to allow them to be unconnected and enable inline BP editing
-TArray<FName> UPhysicsControlComponent::MakeControlsFromSkeletalMeshAndConstraintProfile(
+TArray<FName> UPhysicsControlComponent::CreateControlsFromSkeletalMeshAndConstraintProfile(
 	USkeletalMeshComponent* SkeletalMeshComponent,
 	const TArray<FName>&    BoneNames,
-	FName                   ConstraintProfile,
-	bool                    bEnabled)
+	const FName             ConstraintProfile,
+	const FName             Set,
+	const bool              bEnabled)
 {
 	TArray<FName> Result;
 	UPhysicsAsset* PhysicsAsset = SkeletalMeshComponent ? SkeletalMeshComponent->GetPhysicsAsset() : nullptr;
@@ -524,8 +540,8 @@ TArray<FName> UPhysicsControlComponent::MakeControlsFromSkeletalMeshAndConstrain
 	{
 		FBodyInstance* ChildBone = SkeletalMeshComponent->GetBodyInstance(ChildBoneName);
 
-		FName ParentBoneName;
-		ParentBoneName = UE::PhysicsControlComponent::GetPhysicalParentBone(SkeletalMeshComponent, ChildBoneName);
+		const FName ParentBoneName = 
+			UE::PhysicsControlComponent::GetPhysicalParentBone(SkeletalMeshComponent, ChildBoneName);
 		if (ParentBoneName.IsNone())
 		{
 			continue;
@@ -536,18 +552,22 @@ TArray<FName> UPhysicsControlComponent::MakeControlsFromSkeletalMeshAndConstrain
 		if (!SkeletalMeshComponent->GetConstraintProfilePropertiesOrDefault(
 			ProfileProperties, ChildBoneName, ConstraintProfile))
 		{
-			UE_LOG(LogPhysicsControlComponent, Warning, TEXT("Failed get constraint profile for %s"), *ChildBoneName.ToString());
+			UE_LOG(LogPhysicsControlComponent, Warning, 
+				TEXT("Failed get constraint profile for %s"), *ChildBoneName.ToString());
 			continue;
 		}
 
 		UE::PhysicsControlComponent::ConvertConstraintProfileToControlData(ControlData, ProfileProperties);
 
-		FName ControlName = MakeControl(
+		const FName ControlName = CreateControl(
 			SkeletalMeshComponent, ParentBoneName, SkeletalMeshComponent, ChildBoneName,
-			ControlData, FPhysicsControlTarget(), ControlSettings, bEnabled);
+			ControlData, FPhysicsControlTarget(), ControlSettings, 
+			FName(GetControlTypeName(EPhysicsControlType::ParentSpace).ToString().Append("_").Append(Set.ToString())),
+			bEnabled);
 		if (!ControlName.IsNone())
 		{
 			Result.Add(ControlName);
+			Implementation->NameRecords.AddControl(ControlName, GetControlTypeName(EPhysicsControlType::ParentSpace));
 		}
 		else
 		{
@@ -560,20 +580,20 @@ TArray<FName> UPhysicsControlComponent::MakeControlsFromSkeletalMeshAndConstrain
 
 //======================================================================================================================
 // Note - some params passed by value to allow them to be unconnected and enable inline BP editing
-TMap<FName, FPhysicsControlNameArray> UPhysicsControlComponent::MakeControlsFromLimbBones(
-	FPhysicsControlNameArray&                    AllControls,
+TMap<FName, FPhysicsControlNames> UPhysicsControlComponent::CreateControlsFromLimbBones(
+	FPhysicsControlNames&                        AllControls,
 	const TMap<FName, FPhysicsControlLimbBones>& LimbBones,
-	EPhysicsControlType                          ControlType,
-	FPhysicsControlData                          ControlData,
-	FPhysicsControlSettings                      ControlSettings,
-	bool                                         bEnabled)
+	const EPhysicsControlType                    ControlType,
+	const FPhysicsControlData                    ControlData,
+	const FPhysicsControlSettings                ControlSettings,
+	const bool                                   bEnabled)
 {
-	TMap<FName, FPhysicsControlNameArray> Result;
+	TMap<FName, FPhysicsControlNames> Result;
 	Result.Reserve(LimbBones.Num());
 
-	for (const TPair< FName, FPhysicsControlLimbBones>& LimbBoneEntry : LimbBones)
+	for (const TPair<FName, FPhysicsControlLimbBones>& LimbBoneEntry : LimbBones)
 	{
-		FName LimbName = LimbBoneEntry.Key;
+		const FName LimbName = LimbBoneEntry.Key;
 		const FPhysicsControlLimbBones& BonesInLimb = LimbBoneEntry.Value;
 
 		if (!BonesInLimb.SkeletalMeshComponent)
@@ -587,7 +607,7 @@ TMap<FName, FPhysicsControlNameArray> UPhysicsControlComponent::MakeControlsFrom
 
 		const int32 NumBonesInLimb = BonesInLimb.BoneNames.Num();
 
-		FPhysicsControlNameArray& LimbResult = Result.Add(LimbName);
+		FPhysicsControlNames& LimbResult = Result.Add(LimbName);
 		LimbResult.Names.Reserve(NumBonesInLimb);
 		AllControls.Names.Reserve(AllControls.Names.Num() + NumBonesInLimb);
 
@@ -612,17 +632,21 @@ TMap<FName, FPhysicsControlNameArray> UPhysicsControlComponent::MakeControlsFrom
 					continue;
 				}
 			}
-			FName ControlName = MakeControl(
+			const FName ControlName = CreateControl(
 				ParentMeshComponent, ParentBoneName, BonesInLimb.SkeletalMeshComponent, ChildBoneName,
-				ControlData, FPhysicsControlTarget(), ControlSettings, bEnabled);
+				ControlData, FPhysicsControlTarget(), ControlSettings, 
+				FName(GetControlTypeName(ControlType).ToString().Append("_").Append(LimbName.ToString())),
+				bEnabled);
 			if (!ControlName.IsNone())
 			{
 				LimbResult.Names.Add(ControlName);
 				AllControls.Names.Add(ControlName);
+				Implementation->NameRecords.AddControl(ControlName, GetControlTypeName(ControlType));
 			}
 			else
 			{
-				UE_LOG(LogPhysicsControlComponent, Warning, TEXT("Failed to make control for %s"), *ChildBoneName.ToString());
+				UE_LOG(LogPhysicsControlComponent, Warning, 
+					TEXT("Failed to make control for %s"), *ChildBoneName.ToString());
 			}
 		}
 	}
@@ -630,17 +654,17 @@ TMap<FName, FPhysicsControlNameArray> UPhysicsControlComponent::MakeControlsFrom
 }
 
 //======================================================================================================================
-TMap<FName, FPhysicsControlNameArray> UPhysicsControlComponent::MakeControlsFromLimbBonesAndConstraintProfile(
-	FPhysicsControlNameArray&                    AllControls,
+TMap<FName, FPhysicsControlNames> UPhysicsControlComponent::CreateControlsFromLimbBonesAndConstraintProfile(
+	FPhysicsControlNames&                        AllControls,
 	const TMap<FName, FPhysicsControlLimbBones>& LimbBones,
-	FName                                        ConstraintProfile,
-	bool                                         bEnabled)
+	const FName                                  ConstraintProfile,
+	const bool                                   bEnabled)
 {
-	TMap<FName, FPhysicsControlNameArray> Result;
+	TMap<FName, FPhysicsControlNames> Result;
 	Result.Reserve(LimbBones.Num());
 	for (const TPair< FName, FPhysicsControlLimbBones>& LimbBoneEntry : LimbBones)
 	{
-		FName LimbName = LimbBoneEntry.Key;
+		const FName LimbName = LimbBoneEntry.Key;
 		const FPhysicsControlLimbBones& BonesInLimb = LimbBoneEntry.Value;
 
 		USkeletalMeshComponent* SkeletalMeshComponent = BonesInLimb.SkeletalMeshComponent;
@@ -658,7 +682,7 @@ TMap<FName, FPhysicsControlNameArray> UPhysicsControlComponent::MakeControlsFrom
 
 		const int32 NumBonesInLimb = BonesInLimb.BoneNames.Num();
 
-		FPhysicsControlNameArray& LimbResult = Result.Add(LimbName);
+		FPhysicsControlNames& LimbResult = Result.Add(LimbName);
 		LimbResult.Names.Reserve(NumBonesInLimb);
 		AllControls.Names.Reserve(AllControls.Names.Num() + NumBonesInLimb);
 
@@ -672,13 +696,13 @@ TMap<FName, FPhysicsControlNameArray> UPhysicsControlComponent::MakeControlsFrom
 			// Don't create the parent space control if it's the first bone in a limb that had bIncludeParentBone
 			if (BoneIndex == 0 && BonesInLimb.bFirstBoneIsAdditional)
 			{
-				continue;
+				continue; 
 			}
 
-			FName ChildBoneName = BonesInLimb.BoneNames[BoneIndex];
+			const FName ChildBoneName = BonesInLimb.BoneNames[BoneIndex];
 			FBodyInstance* ChildBone = SkeletalMeshComponent->GetBodyInstance(ChildBoneName);
-			FName ParentBoneName;
-			ParentBoneName = UE::PhysicsControlComponent::GetPhysicalParentBone(SkeletalMeshComponent, ChildBoneName);
+			const FName ParentBoneName = 
+				UE::PhysicsControlComponent::GetPhysicalParentBone(SkeletalMeshComponent, ChildBoneName);
 			if (ParentBoneName.IsNone())
 			{
 				continue;
@@ -688,24 +712,30 @@ TMap<FName, FPhysicsControlNameArray> UPhysicsControlComponent::MakeControlsFrom
 			FConstraintProfileProperties ProfileProperties;
 			if (!SkeletalMeshComponent->GetConstraintProfilePropertiesOrDefault(
 				ProfileProperties, ChildBoneName, ConstraintProfile))
+			
 			{
-				UE_LOG(LogPhysicsControlComponent, Warning, TEXT("Failed get constraint profile for %s"), *ChildBoneName.ToString());
+				UE_LOG(LogPhysicsControlComponent, Warning, 
+					TEXT("Failed get constraint profile for %s"), *ChildBoneName.ToString());
 				continue;
 			}
 
 			UE::PhysicsControlComponent::ConvertConstraintProfileToControlData(ControlData, ProfileProperties);
 
-			FName ControlName = MakeControl(
+			const FName ControlName = CreateControl(
 				SkeletalMeshComponent, ParentBoneName, SkeletalMeshComponent, ChildBoneName,
-				ControlData, FPhysicsControlTarget(), ControlSettings, bEnabled);
+				ControlData, FPhysicsControlTarget(), ControlSettings, 
+				FName(GetControlTypeName(EPhysicsControlType::ParentSpace).ToString().Append("_").Append(LimbName.ToString())),
+				bEnabled);
 			if (!ControlName.IsNone())
 			{
 				LimbResult.Names.Add(ControlName);
 				AllControls.Names.Add(ControlName);
+				Implementation->NameRecords.AddControl(ControlName, GetControlTypeName(EPhysicsControlType::ParentSpace));
 			}
 			else
 			{
-				UE_LOG(LogPhysicsControlComponent, Warning, TEXT("Failed to make control for %s"), *ChildBoneName.ToString());
+				UE_LOG(LogPhysicsControlComponent, Warning, 
+					TEXT("Failed to make control for %s"), *ChildBoneName.ToString());
 			}
 		}
 	}
@@ -715,38 +745,26 @@ TMap<FName, FPhysicsControlNameArray> UPhysicsControlComponent::MakeControlsFrom
 //======================================================================================================================
 bool UPhysicsControlComponent::DestroyControl(const FName Name)
 {
-	FPhysicsControlRecord* PhysicsControlRecord = Implementation->FindControlRecord(Name);
-	if (PhysicsControlRecord)
-	{
-		if (USkeletalMeshComponent* SkeletalMeshComponent = 
-			Cast<USkeletalMeshComponent>(PhysicsControlRecord->PhysicsControl.ParentMeshComponent))
-		{
-			Implementation->RemoveSkeletalMeshReferenceForCaching(SkeletalMeshComponent);
-		}
-		if (USkeletalMeshComponent* SkeletalMeshComponent =
-			Cast<USkeletalMeshComponent>(PhysicsControlRecord->PhysicsControl.ChildMeshComponent))
-		{
-			Implementation->RemoveSkeletalMeshReferenceForCaching(SkeletalMeshComponent);
-		}
-
-		PhysicsControlRecord->PhysicsControlState.bPendingDestroy = true;
-		PhysicsControlRecord->PhysicsControlState.bEnabled = false;
-		return true;
-	}
-	return false;
+	return Implementation->DestroyControl(Name);
 }
 
 //======================================================================================================================
-void UPhysicsControlComponent::DestroyAllControls(const TArray<FName>& Names)
+void UPhysicsControlComponent::DestroyControls(const TArray<FName>& Names)
 {
 	for (FName Name : Names)
 	{
-		DestroyControl(Name);
+		Implementation->DestroyControl(Name);
 	}
 }
 
 //======================================================================================================================
-bool UPhysicsControlComponent::SetControlEnabled(const FName Name, bool bEnable)
+void UPhysicsControlComponent::DestroyControlsInSet(const FName SetName)
+{
+	DestroyControls(GetControlNamesInSet(SetName));
+}
+
+//======================================================================================================================
+bool UPhysicsControlComponent::SetControlEnabled(const FName Name, const bool bEnable)
 {
 	FPhysicsControlRecord* Record = Implementation->FindControlRecord(Name);
 	if (Record)
@@ -758,7 +776,7 @@ bool UPhysicsControlComponent::SetControlEnabled(const FName Name, bool bEnable)
 }
 
 //======================================================================================================================
-void UPhysicsControlComponent::SetAllControlsEnabled(const TArray<FName>& Names, bool bEnable)
+void UPhysicsControlComponent::SetControlsEnabled(const TArray<FName>& Names, const bool bEnable)
 {
 	for (FName Name : Names)
 	{
@@ -766,12 +784,18 @@ void UPhysicsControlComponent::SetAllControlsEnabled(const TArray<FName>& Names,
 	}
 }
 
+//======================================================================================================================
+void UPhysicsControlComponent::SetControlsInSetEnabled(FName SetName, bool bEnable)
+{
+	SetControlsEnabled(GetControlNamesInSet(SetName), bEnable);
+}
+
 //======================================================================================================================
 // Note - some params passed by value to allow them to be unconnected and enable inline BP editing
 bool UPhysicsControlComponent::SetControlData(
-	const FName         Name, 
-	FPhysicsControlData ControlData, 
-	bool                bEnableControl)
+	const FName               Name, 
+	const FPhysicsControlData ControlData, 
+	const bool                bEnableControl)
 {
 	FPhysicsControlRecord* Record = Implementation->FindControlRecord(Name);
 	if (Record)
@@ -787,10 +811,10 @@ bool UPhysicsControlComponent::SetControlData(
 }
 
 //======================================================================================================================
-void UPhysicsControlComponent::SetAllControlDatas(
-	const TArray<FName>& Names, 
-	FPhysicsControlData  ControlData, 
-	bool                 bEnableControl)
+void UPhysicsControlComponent::SetControlDatas(
+	const TArray<FName>&       Names, 
+	const FPhysicsControlData  ControlData, 
+	const bool                 bEnableControl)
 {
 	for (FName Name : Names)
 	{
@@ -798,18 +822,26 @@ void UPhysicsControlComponent::SetAllControlDatas(
 	}
 }
 
+//======================================================================================================================
+void UPhysicsControlComponent::SetControlDatasInSet(
+	const FName                SetName,
+	const FPhysicsControlData  ControlData,
+	const bool                 bEnableControl)
+{
+	SetControlDatas(GetControlNamesInSet(SetName), ControlData, bEnableControl);
+}
 
 //======================================================================================================================
 // Note - some params passed by value to allow them to be unconnected and enable inline BP editing
-bool UPhysicsControlComponent::SetControlMultipliers(
-	const FName                Name, 
-	FPhysicsControlMultipliers ControlMultipliers, 
-	bool                       bEnableControl)
+bool UPhysicsControlComponent::SetControlMultiplier(
+	const FName                      Name, 
+	const FPhysicsControlMultiplier  ControlMultiplier, 
+	const bool                       bEnableControl)
 {
 	FPhysicsControlRecord* Record = Implementation->FindControlRecord(Name);
 	if (Record)
 	{
-		Record->PhysicsControl.ControlMultipliers = ControlMultipliers;
+		Record->PhysicsControl.ControlMultiplier = ControlMultiplier;
 		if (bEnableControl)
 		{
 			Record->PhysicsControlState.bEnabled = true;
@@ -820,21 +852,34 @@ bool UPhysicsControlComponent::SetControlMultipliers(
 }
 
 //======================================================================================================================
-void UPhysicsControlComponent::SetAllControlMultipliers(
-	const TArray<FName>&       Names,
-	FPhysicsControlMultipliers ControlMultipliers, 
-	bool                       bEnableControl)
+void UPhysicsControlComponent::SetControlMultipliers(
+	const TArray<FName>&             Names,
+	const FPhysicsControlMultiplier  ControlMultiplier, 
+	const bool                       bEnableControl)
 {
 	for (FName Name : Names)
 	{
-		SetControlMultipliers(Name, ControlMultipliers, bEnableControl);
+		SetControlMultiplier(Name, ControlMultiplier, bEnableControl);
 	}
 }
 
+//======================================================================================================================
+void UPhysicsControlComponent::SetControlMultipliersInSet(
+	const FName                      SetName,
+	const FPhysicsControlMultiplier  ControlMultiplier, 
+	const bool                       bEnableControl)
+{
+	SetControlMultipliers(GetControlNamesInSet(SetName), ControlMultiplier, bEnableControl);
+}
 
 //======================================================================================================================
 bool UPhysicsControlComponent::SetControlLinearData(
-	const FName Name, float Strength, float DampingRatio, float ExtraDamping, float MaxForce, bool bEnableControl)
+	const FName Name, 
+	const float Strength, 
+	const float DampingRatio, 
+	const float ExtraDamping, 
+	const float MaxForce, 
+	const bool  bEnableControl)
 {
 	FPhysicsControlRecord* Record = Implementation->FindControlRecord(Name);
 	if (Record)
@@ -854,7 +899,12 @@ bool UPhysicsControlComponent::SetControlLinearData(
 
 //======================================================================================================================
 bool UPhysicsControlComponent::SetControlAngularData(
-	const FName Name, float Strength, float DampingRatio, float ExtraDamping, float MaxTorque, bool bEnableControl)
+	const FName Name, 
+	const float Strength, 
+	const float DampingRatio, 
+	const float ExtraDamping, 
+	const float MaxTorque, 
+	const bool  bEnableControl)
 {
 	FPhysicsControlRecord* Record = Implementation->FindControlRecord(Name);
 	if (Record)
@@ -900,8 +950,8 @@ bool UPhysicsControlComponent::ResetControlPoint(const FName Name)
 //======================================================================================================================
 bool UPhysicsControlComponent::SetControlTarget(
 	const FName           Name, 
-	FPhysicsControlTarget ControlTarget, 
-	bool                  bEnableControl)
+	const FPhysicsControlTarget ControlTarget, 
+	const bool                  bEnableControl)
 {
 	FPhysicsControlRecord* Record = Implementation->FindControlRecord(Name);
 	if (Record)
@@ -917,10 +967,10 @@ bool UPhysicsControlComponent::SetControlTarget(
 }
 
 //======================================================================================================================
-void UPhysicsControlComponent::SetAllControlTargets(
-	const TArray<FName>&  Names, 
-	FPhysicsControlTarget ControlTarget, 
-	bool                  bEnableControl)
+void UPhysicsControlComponent::SetControlTargets(
+	const TArray<FName>&        Names, 
+	const FPhysicsControlTarget ControlTarget, 
+	const bool                  bEnableControl)
 {
 	for (FName Name : Names)
 	{
@@ -928,10 +978,23 @@ void UPhysicsControlComponent::SetAllControlTargets(
 	}
 }
 
+//======================================================================================================================
+void UPhysicsControlComponent::SetControlTargetsInSet(
+	const FName                 SetName, 
+	const FPhysicsControlTarget ControlTarget, 
+	const bool                  bEnableControl)
+{
+	SetControlTargets(GetControlNamesInSet(SetName), ControlTarget, bEnableControl);
+}
+
 //======================================================================================================================
 bool UPhysicsControlComponent::SetControlTargetPositionAndOrientation(
-	const FName Name, const FVector Position, const FRotator Orientation, float VelocityDeltaTime, 
-	bool bEnableControl, bool bApplyControlPointToTarget)
+	const FName    Name, 
+	const FVector  Position, 
+	const FRotator Orientation, 
+	const float    VelocityDeltaTime, 
+	const bool     bEnableControl, 
+	const bool     bApplyControlPointToTarget)
 {
 	FPhysicsControlRecord* Record = Implementation->FindControlRecord(Name);
 	if (Record)
@@ -947,8 +1010,11 @@ bool UPhysicsControlComponent::SetControlTargetPositionAndOrientation(
 
 //======================================================================================================================
 bool UPhysicsControlComponent::SetControlTargetPosition(
-	const FName Name, const FVector Position, float VelocityDeltaTime, 
-	bool bEnableControl, bool bApplyControlPointToTarget)
+	const FName   Name, 
+	const FVector Position, 
+	const float   VelocityDeltaTime, 
+	const bool    bEnableControl, 
+	const bool    bApplyControlPointToTarget)
 {
 	FPhysicsControlRecord* Record = Implementation->FindControlRecord(Name);
 	if (Record)
@@ -975,8 +1041,11 @@ bool UPhysicsControlComponent::SetControlTargetPosition(
 
 //======================================================================================================================
 bool UPhysicsControlComponent::SetControlTargetOrientation(
-	const FName Name, const FRotator Orientation, float AngularVelocityDeltaTime, 
-	bool bEnableControl, bool bApplyControlPointToTarget)
+	const FName    Name, 
+	const FRotator Orientation, 
+	const float    AngularVelocityDeltaTime, 
+	const bool     bEnableControl, 
+	const bool     bApplyControlPointToTarget)
 {
 	FPhysicsControlRecord* Record = Implementation->FindControlRecord(Name);
 	if (Record)
@@ -1008,10 +1077,13 @@ bool UPhysicsControlComponent::SetControlTargetOrientation(
 
 //======================================================================================================================
 bool UPhysicsControlComponent::SetControlTargetPoses(
-	const FName Name,
-	const FVector ParentPosition, const FRotator ParentOrientation,
-	const FVector ChildPosition, const FRotator ChildOrientation,
-	float VelocityDeltaTime, bool bEnableControl)
+	const FName    Name,
+	const FVector  ParentPosition, 
+	const FRotator ParentOrientation,
+	const FVector  ChildPosition, 
+	const FRotator ChildOrientation,
+	const float    VelocityDeltaTime, 
+	const bool     bEnableControl)
 {
 	FPhysicsControlRecord* Record = Implementation->FindControlRecord(Name);
 	if (Record)
@@ -1055,8 +1127,8 @@ bool UPhysicsControlComponent::SetControlTargetPoses(
 //======================================================================================================================
 bool UPhysicsControlComponent::SetControlUseSkeletalAnimation(
 	const FName Name,
-	bool        bUseSkeletalAnimation,
-	float       SkeletalAnimationVelocityMultiplier)
+	const bool  bUseSkeletalAnimation,
+	const float SkeletalAnimationVelocityMultiplier)
 {
 	FPhysicsControl* PhysicsControl = Implementation->FindControl(Name);
 	if (PhysicsControl)
@@ -1069,10 +1141,10 @@ bool UPhysicsControlComponent::SetControlUseSkeletalAnimation(
 }
 
 //======================================================================================================================
-void UPhysicsControlComponent::SetAllControlsUseSkeletalAnimation(
+void UPhysicsControlComponent::SetControlsUseSkeletalAnimation(
 	const TArray<FName>& Names,
-	bool                 bUseSkeletalAnimation,
-	float                SkeletalAnimationVelocityMultiplier)
+	const bool           bUseSkeletalAnimation,
+	const float          SkeletalAnimationVelocityMultiplier)
 {
 	for (FName Name : Names)
 	{
@@ -1081,24 +1153,22 @@ void UPhysicsControlComponent::SetAllControlsUseSkeletalAnimation(
 }
 
 //======================================================================================================================
-bool UPhysicsControlComponent::SetControlAutoDisable(const FName Name, bool bAutoDisable)
+void UPhysicsControlComponent::SetControlsInSetUseSkeletalAnimation(
+	const FName SetName,
+	const bool  bUseSkeletalAnimation,
+	const float SkeletalAnimationVelocityMultiplier)
 {
-	FPhysicsControl* PhysicsControl = Implementation->FindControl(Name);
-	if (PhysicsControl)
-	{
-		PhysicsControl->ControlSettings.bAutoDisable = bAutoDisable;
-		return true;
-	}
-	return false;
+	SetControlsUseSkeletalAnimation(
+		GetControlNamesInSet(SetName), bUseSkeletalAnimation, SkeletalAnimationVelocityMultiplier);
 }
 
 //======================================================================================================================
-bool UPhysicsControlComponent::GetControl(const FName Name, FPhysicsControl& Control) const
+bool UPhysicsControlComponent::SetControlAutoDisable(const FName Name, const bool bAutoDisable)
 {
 	FPhysicsControl* PhysicsControl = Implementation->FindControl(Name);
 	if (PhysicsControl)
 	{
-		Control = *PhysicsControl;
+		PhysicsControl->ControlSettings.bAutoDisable = bAutoDisable;
 		return true;
 	}
 	return false;
@@ -1117,12 +1187,12 @@ bool UPhysicsControlComponent::GetControlData(const FName Name, FPhysicsControlD
 }
 
 //======================================================================================================================
-bool UPhysicsControlComponent::GetControlMultipliers(const FName Name, FPhysicsControlMultipliers& ControlMultipliers) const
+bool UPhysicsControlComponent::GetControlMultiplier(const FName Name, FPhysicsControlMultiplier& ControlMultiplier) const
 {
 	FPhysicsControl* PhysicsControl = Implementation->FindControl(Name);
 	if (PhysicsControl)
 	{
-		ControlMultipliers = PhysicsControl->ControlMultipliers;
+		ControlMultiplier = PhysicsControl->ControlMultiplier;
 		return true;
 	}
 	return false;
@@ -1163,17 +1233,18 @@ bool UPhysicsControlComponent::GetControlEnabled(const FName Name) const
 }
 
 //======================================================================================================================
-FName UPhysicsControlComponent::MakeBodyModifier(
-	UMeshComponent*         MeshComponent,
-	const FName             BoneName,
-	EPhysicsMovementType    MovementType, 
-	ECollisionEnabled::Type CollisionType,
-	float                   GravityMultiplier,
-	bool                    bUseSkeletalAnimation)
+FName UPhysicsControlComponent::CreateBodyModifier(
+	UMeshComponent*               MeshComponent,
+	const FName                   BoneName,
+	const FName                   Set,
+	const EPhysicsMovementType    MovementType,
+	const ECollisionEnabled::Type CollisionType,
+	const float                   GravityMultiplier,
+	const bool                    bUseSkeletalAnimation)
 {
-	FName Name = Implementation->GetUniqueBodyModifierName(BoneName);
-	if (MakeNamedBodyModifier(
-		Name, MeshComponent, BoneName, MovementType, CollisionType, GravityMultiplier, bUseSkeletalAnimation))
+	const FName Name = Implementation->GetUniqueBodyModifierName(BoneName);
+	if (CreateNamedBodyModifier(
+		Name, MeshComponent, BoneName, Set, MovementType, CollisionType, GravityMultiplier, bUseSkeletalAnimation))
 	{
 		return Name;
 	}
@@ -1181,17 +1252,17 @@ FName UPhysicsControlComponent::MakeBodyModifier(
 }
 
 //======================================================================================================================
-bool UPhysicsControlComponent::MakeNamedBodyModifier(
-	const FName             Name,
-	UMeshComponent*         MeshComponent,
-	const FName             BoneName,
-	EPhysicsMovementType    MovementType,
-	ECollisionEnabled::Type CollisionType,
-	float                   GravityMultiplier,
-	bool                    bUseSkeletalAnimation)
+bool UPhysicsControlComponent::CreateNamedBodyModifier(
+	const FName                   Name,
+	UMeshComponent*               MeshComponent,
+	const FName                   BoneName,
+	const FName                   Set,
+	const EPhysicsMovementType    MovementType,
+	const ECollisionEnabled::Type CollisionType,
+	const float                   GravityMultiplier,
+	const bool                    bUseSkeletalAnimation)
 {
-	FPhysicsBodyModifier* BodyModifier = Implementation->FindBodyModifier(Name);
-	if (BodyModifier)
+	if (Implementation->FindBodyModifier(Name))
 	{
 		return false;
 	}
@@ -1214,18 +1285,21 @@ bool UPhysicsControlComponent::MakeNamedBodyModifier(
 		Implementation->AddSkeletalMeshReferenceForModifier(SkeletalMeshComponent);
 	}
 
+	Implementation->NameRecords.AddBodyModifier(Name, Set);
+
 	return true;
 }
 
 //======================================================================================================================
-TArray<FName> UPhysicsControlComponent::MakeBodyModifiersFromSkeletalMeshBelow(
-	USkeletalMeshComponent* SkeletalMeshComponent,
-	FName                   BoneName,
-	bool                    bIncludeSelf,
-	EPhysicsMovementType    MovementType,
-	ECollisionEnabled::Type CollisionType,
-	float                   GravityMultiplier,
-	bool                    bUseSkeletalAnimation)
+TArray<FName> UPhysicsControlComponent::CreateBodyModifiersFromSkeletalMeshBelow(
+	USkeletalMeshComponent*       SkeletalMeshComponent,
+	const FName                   BoneName,
+	const bool                    bIncludeSelf,
+	const FName                   Set,
+	const EPhysicsMovementType    MovementType,
+	const ECollisionEnabled::Type CollisionType,
+	const float                   GravityMultiplier,
+	const bool                    bUseSkeletalAnimation)
 {
 	TArray<FName> Result;
 	UPhysicsAsset* PhysicsAsset = SkeletalMeshComponent ? SkeletalMeshComponent->GetPhysicsAsset() : nullptr;
@@ -1237,15 +1311,16 @@ TArray<FName> UPhysicsControlComponent::MakeBodyModifiersFromSkeletalMeshBelow(
 	SkeletalMeshComponent->ForEachBodyBelow(
 		BoneName, bIncludeSelf, /*bSkipCustomType=*/false,
 		[
-			this, PhysicsAsset, SkeletalMeshComponent, MovementType, CollisionType,
+			this, PhysicsAsset, SkeletalMeshComponent, Set, MovementType, CollisionType,
 			GravityMultiplier, bUseSkeletalAnimation, &Result
 		](const FBodyInstance* BI)
 		{
 			if (USkeletalBodySetup* BodySetup = Cast<USkeletalBodySetup>(BI->BodySetup.Get()))
 			{
 				const FName BoneName = PhysicsAsset->SkeletalBodySetups[BI->InstanceBodyIndex]->BoneName;
-				FName BodyModifierName = MakeBodyModifier(
-					SkeletalMeshComponent, BoneName, MovementType, CollisionType, GravityMultiplier, bUseSkeletalAnimation);
+				const FName BodyModifierName = CreateBodyModifier(
+					SkeletalMeshComponent, BoneName, Set, MovementType, 
+					CollisionType, GravityMultiplier, bUseSkeletalAnimation);
 				Result.Add(BodyModifierName);
 			}
 		});
@@ -1254,20 +1329,20 @@ TArray<FName> UPhysicsControlComponent::MakeBodyModifiersFromSkeletalMeshBelow(
 }
 
 //======================================================================================================================
-TMap<FName, FPhysicsControlNameArray> UPhysicsControlComponent::MakeBodyModifiersFromLimbBones(
-	FPhysicsControlNameArray&                    AllBodyModifiers,
+TMap<FName, FPhysicsControlNames> UPhysicsControlComponent::CreateBodyModifiersFromLimbBones(
+	FPhysicsControlNames&                        AllBodyModifiers,
 	const TMap<FName, FPhysicsControlLimbBones>& LimbBones,
-	EPhysicsMovementType                         MovementType,
-	ECollisionEnabled::Type                      CollisionType,
-	float                                        GravityMultiplier,
-	bool                                         bUseSkeletalAnimation)
+	const EPhysicsMovementType                   MovementType,
+	const ECollisionEnabled::Type                CollisionType,
+	const float                                  GravityMultiplier,
+	const bool                                   bUseSkeletalAnimation)
 {
-	TMap<FName, FPhysicsControlNameArray> Result;
+	TMap<FName, FPhysicsControlNames> Result;
 	Result.Reserve(LimbBones.Num());
 
 	for (const TPair<FName, FPhysicsControlLimbBones>& LimbBoneEntry : LimbBones)
 	{
-		FName LimbName = LimbBoneEntry.Key;
+		const FName LimbName = LimbBoneEntry.Key;
 		const FPhysicsControlLimbBones& BonesInLimb = LimbBoneEntry.Value;
 
 		if (!BonesInLimb.SkeletalMeshComponent)
@@ -1278,14 +1353,15 @@ TMap<FName, FPhysicsControlNameArray> UPhysicsControlComponent::MakeBodyModifier
 
 		const int32 NumBonesInLimb = BonesInLimb.BoneNames.Num();
 
-		FPhysicsControlNameArray& LimbResult = Result.Add(LimbName);
+		FPhysicsControlNames& LimbResult = Result.Add(LimbName);
 		LimbResult.Names.Reserve(NumBonesInLimb);
 		AllBodyModifiers.Names.Reserve(AllBodyModifiers.Names.Num() + NumBonesInLimb);
 
-		for (FName BoneName : BonesInLimb.BoneNames)
+		for (const FName BoneName : BonesInLimb.BoneNames)
 		{
-			FName BodyModifierName = MakeBodyModifier(
-				BonesInLimb.SkeletalMeshComponent, BoneName, MovementType, CollisionType, GravityMultiplier, bUseSkeletalAnimation);
+			const FName BodyModifierName = CreateBodyModifier(
+				BonesInLimb.SkeletalMeshComponent, BoneName, LimbName, MovementType, 
+				CollisionType, GravityMultiplier, bUseSkeletalAnimation);
 			if (!BodyModifierName.IsNone())
 			{
 				LimbResult.Names.Add(BodyModifierName);
@@ -1293,7 +1369,8 @@ TMap<FName, FPhysicsControlNameArray> UPhysicsControlComponent::MakeBodyModifier
 			}
 			else
 			{
-				UE_LOG(LogPhysicsControlComponent, Warning, TEXT("Failed to make body modifier for %s"), *BoneName.ToString());
+				UE_LOG(LogPhysicsControlComponent, Warning, 
+					TEXT("Failed to make body modifier for %s"), *BoneName.ToString());
 			}
 		}
 	}
@@ -1303,37 +1380,30 @@ TMap<FName, FPhysicsControlNameArray> UPhysicsControlComponent::MakeBodyModifier
 //======================================================================================================================
 bool UPhysicsControlComponent::DestroyBodyModifier(const FName Name)
 {
-	FPhysicsBodyModifier* PhysicsBodyModifier = Implementation->FindBodyModifier(Name);
-	if (PhysicsBodyModifier)
-	{
-		if (USkeletalMeshComponent* SkeletalMeshComponent = 
-			Cast<USkeletalMeshComponent>(PhysicsBodyModifier->MeshComponent))
-		{
-			Implementation->RemoveSkeletalMeshReferenceForCaching(SkeletalMeshComponent);
-			Implementation->RemoveSkeletalMeshReferenceForModifier(SkeletalMeshComponent);
-		}
-
-		PhysicsBodyModifier->bPendingDestroy = true;
-		return true;
-	}
-	return false;
+	return Implementation->DestroyBodyModifier(Name);
 }
 
 //======================================================================================================================
-void UPhysicsControlComponent::DestroyAllBodyModifiers(const TArray<FName>& Names)
+void UPhysicsControlComponent::DestroyBodyModifiers(const TArray<FName>& Names)
 {
 	for (FName Name : Names)
 	{
-		DestroyBodyModifier(Name);
+		Implementation->DestroyBodyModifier(Name);
 	}
 }
 
+//======================================================================================================================
+void UPhysicsControlComponent::DestroyBodyModifiersInSet(const FName SetName)
+{
+	DestroyBodyModifiers(GetBodyModifierNamesInSet(SetName));
+}
+
 //======================================================================================================================
 bool UPhysicsControlComponent::SetBodyModifierKinematicTarget(
-	const FName Name,
-	FVector     KinematicTargetPosition,
-	FRotator    KinematicTargetOrienation,
-	bool        bMakeKinematic)
+	const FName    Name,
+	const FVector  KinematicTargetPosition,
+	const FRotator KinematicTargetOrienation,
+	const bool     bMakeKinematic)
 {
 	FPhysicsBodyModifier* PhysicsBodyModifier = Implementation->FindBodyModifier(Name);
 	if (PhysicsBodyModifier)
@@ -1351,8 +1421,8 @@ bool UPhysicsControlComponent::SetBodyModifierKinematicTarget(
 
 //======================================================================================================================
 bool UPhysicsControlComponent::SetBodyModifierMovementType(
-		FName                Name,
-		EPhysicsMovementType MovementType)
+	const FName                Name,
+	const EPhysicsMovementType MovementType)
 {
 	FPhysicsBodyModifier* PhysicsBodyModifier = Implementation->FindBodyModifier(Name);
 	if (PhysicsBodyModifier)
@@ -1364,9 +1434,9 @@ bool UPhysicsControlComponent::SetBodyModifierMovementType(
 }
 
 //======================================================================================================================
-void UPhysicsControlComponent::SetAllBodyModifierMovementType(
-	const TArray<FName>& Names,
-	EPhysicsMovementType MovementType)
+void UPhysicsControlComponent::SetBodyModifiersMovementType(
+	const TArray<FName>&       Names,
+	const EPhysicsMovementType MovementType)
 {
 	for (FName Name : Names)
 	{
@@ -1374,10 +1444,18 @@ void UPhysicsControlComponent::SetAllBodyModifierMovementType(
 	}
 }
 
+//======================================================================================================================
+void UPhysicsControlComponent::SetBodyModifiersInSetMovementType(
+	const FName                SetName,
+	const EPhysicsMovementType MovementType)
+{
+	SetBodyModifiersMovementType(GetBodyModifierNamesInSet(SetName), MovementType);
+}
+
 //======================================================================================================================
 bool UPhysicsControlComponent::SetBodyModifierCollisionType(
-	FName                   Name,
-	ECollisionEnabled::Type CollisionType)
+	const FName                   Name,
+	const ECollisionEnabled::Type CollisionType)
 {
 	FPhysicsBodyModifier* PhysicsBodyModifier = Implementation->FindBodyModifier(Name);
 	if (PhysicsBodyModifier)
@@ -1389,9 +1467,9 @@ bool UPhysicsControlComponent::SetBodyModifierCollisionType(
 }
 
 //======================================================================================================================
-void UPhysicsControlComponent::SetAllBodyModifierCollisionType(
-	const TArray<FName>&    Names,
-	ECollisionEnabled::Type CollisionType)
+void UPhysicsControlComponent::SetBodyModifiersCollisionType(
+	const TArray<FName>&          Names,
+	const ECollisionEnabled::Type CollisionType)
 {
 	for (FName Name : Names)
 	{
@@ -1399,10 +1477,18 @@ void UPhysicsControlComponent::SetAllBodyModifierCollisionType(
 	}
 }
 
+//======================================================================================================================
+void UPhysicsControlComponent::SetBodyModifiersInSetCollisionType(
+	const FName                   SetName,
+	const ECollisionEnabled::Type CollisionType)
+{
+	SetBodyModifiersCollisionType(GetBodyModifierNamesInSet(SetName), CollisionType);
+}
+
 //======================================================================================================================
 bool UPhysicsControlComponent::SetBodyModifierGravityMultiplier(
 	const FName Name,
-	float       GravityMultiplier)
+	const float GravityMultiplier)
 {
 	FPhysicsBodyModifier* PhysicsBodyModifier = Implementation->FindBodyModifier(Name);
 	if (PhysicsBodyModifier)
@@ -1414,9 +1500,9 @@ bool UPhysicsControlComponent::SetBodyModifierGravityMultiplier(
 }
 
 //======================================================================================================================
-void UPhysicsControlComponent::SetAllBodyModifierGravityMultipliers(
+void UPhysicsControlComponent::SetBodyModifiersGravityMultiplier(
 	const TArray<FName>& Names,
-	float                GravityMultiplier)
+	const float          GravityMultiplier)
 {
 	for (FName Name : Names)
 	{
@@ -1424,10 +1510,18 @@ void UPhysicsControlComponent::SetAllBodyModifierGravityMultipliers(
 	}
 }
 
+//======================================================================================================================
+void UPhysicsControlComponent::SetBodyModifiersInSetGravityMultiplier(
+	const FName SetName,
+	const float GravityMultiplier)
+{
+	SetBodyModifiersGravityMultiplier(GetBodyModifierNamesInSet(SetName), GravityMultiplier);
+}
+
 //======================================================================================================================
 bool UPhysicsControlComponent::SetBodyModifierUseSkeletalAnimation(
 	const FName Name,
-	bool        bUseSkeletalAnimation)
+	const bool  bUseSkeletalAnimation)
 {
 	FPhysicsBodyModifier* PhysicsBodyModifier = Implementation->FindBodyModifier(Name);
 	if (PhysicsBodyModifier)
@@ -1439,9 +1533,9 @@ bool UPhysicsControlComponent::SetBodyModifierUseSkeletalAnimation(
 }
 
 //======================================================================================================================
-void UPhysicsControlComponent::SetAllBodyModifiersUseSkeletalAnimation(
+void UPhysicsControlComponent::SetBodyModifiersUseSkeletalAnimation(
 	const TArray<FName>& Names,
-	bool                 bUseSkeletalAnimation)
+	const bool           bUseSkeletalAnimation)
 {
 	for (FName Name : Names)
 	{
@@ -1449,20 +1543,297 @@ void UPhysicsControlComponent::SetAllBodyModifiersUseSkeletalAnimation(
 	}
 }
 
+//======================================================================================================================
+void UPhysicsControlComponent::SetBodyModifiersInSetUseSkeletalAnimation(
+	const FName SetName,
+	const bool  bUseSkeletalAnimation)
+{
+	SetBodyModifiersUseSkeletalAnimation(GetBodyModifierNamesInSet(SetName), bUseSkeletalAnimation);
+}
+
 #if WITH_EDITOR
+
+//======================================================================================================================
+void UPhysicsControlComponent::OnRegister()
+{
+	Super::OnRegister();
+
+	if (SpriteComponent)
+	{
+		UpdateSpriteTexture();
+		SpriteComponent->SpriteInfo.Category = TEXT("Physics");
+		SpriteComponent->SpriteInfo.DisplayName = NSLOCTEXT("SpriteCategory", "Physics", "Physics");
+	}
+}
+
+//======================================================================================================================
+void UPhysicsControlComponent::UpdateSpriteTexture()
+{
+	if (SpriteComponent)
+	{
+		SpriteComponent->SetSprite(LoadObject<UTexture2D>(NULL, TEXT("/Engine/EditorResources/S_KBSJoint.S_KBSJoint")));
+	}
+}
+
+//======================================================================================================================
+const TArray<FName>& UPhysicsControlComponent::GetAllControlNames() const
+{
+	return GetControlNamesInSet("All");
+}
+
+//======================================================================================================================
+void UPhysicsControlComponent::CreateControlsAndBodyModifiersFromLimbBones(
+	FPhysicsControlNames&                       AllWorldSpaceControls,
+	TMap<FName, FPhysicsControlNames>&          LimbWorldSpaceControls,
+	FPhysicsControlNames&                       AllParentSpaceControls,
+	TMap<FName, FPhysicsControlNames>&          LimbParentSpaceControls,
+	FPhysicsControlNames&                       AllBodyModifiers,
+	TMap<FName, FPhysicsControlNames>&          LimbBodyModifiers,
+	USkeletalMeshComponent*                     SkeletalMeshComponent,
+	const TArray<FPhysicsControlLimbSetupData>& LimbSetupData,
+	const FPhysicsControlData                   WorldSpaceControlData,
+	const FPhysicsControlSettings               WorldSpaceControlSettings,
+	const bool                                  bEnableWorldSpaceControls,
+	const FPhysicsControlData                   ParentSpaceControlData,
+	const FPhysicsControlSettings               ParentSpaceControlSettings,
+	const bool                                  bEnableParentSpaceControls,
+	const EPhysicsMovementType                  PhysicsMovementType,
+	const float                                 GravityMultiplier
+	)
+{
+	TMap<FName, FPhysicsControlLimbBones> LimbBones = 
+		GetLimbBonesFromSkeletalMesh(SkeletalMeshComponent, LimbSetupData);
+
+	LimbWorldSpaceControls = CreateControlsFromLimbBones(
+		AllWorldSpaceControls, LimbBones, EPhysicsControlType::WorldSpace, 
+		WorldSpaceControlData, WorldSpaceControlSettings, bEnableWorldSpaceControls);
+
+	LimbParentSpaceControls = CreateControlsFromLimbBones(
+		AllParentSpaceControls, LimbBones, EPhysicsControlType::ParentSpace,
+		ParentSpaceControlData, ParentSpaceControlSettings, bEnableParentSpaceControls);
+
+	LimbBodyModifiers = CreateBodyModifiersFromLimbBones(
+		AllBodyModifiers, LimbBones, PhysicsMovementType, ECollisionEnabled::QueryAndPhysics, GravityMultiplier, true);
+}
+
+//======================================================================================================================
+void UPhysicsControlComponent::AddControlToSet(
+	FPhysicsControlNames& NewSet, 
+	const FName           Control, 
+	const FName           SetName)
+{
+	Implementation->NameRecords.AddControl(Control, SetName);
+	NewSet.Names = GetControlNamesInSet(SetName);
+}
+
+//======================================================================================================================
+void UPhysicsControlComponent::AddControlsToSet(
+	FPhysicsControlNames& NewSet, 
+	const TArray<FName>&  Controls, 
+	const FName           SetName)
+{
+	for (FName Control : Controls)
+	{
+		Implementation->NameRecords.AddControl(Control, SetName);
+	}
+	NewSet.Names = GetControlNamesInSet(SetName);
+}
+
+//======================================================================================================================
+const TArray<FName>& UPhysicsControlComponent::GetControlNamesInSet(const FName SetName) const
+{
+	return Implementation->NameRecords.GetControlNamesInSet(SetName);
+}
+
+//======================================================================================================================
+const TArray<FName>& UPhysicsControlComponent::GetAllBodyModifierNames() const
+{
+	return GetBodyModifierNamesInSet("All");
+}
+
+//======================================================================================================================
+const TArray<FName>& UPhysicsControlComponent::GetBodyModifierNamesInSet(const FName SetName) const
+{
+	return Implementation->NameRecords.GetBodyModifierNamesInSet(SetName);
+}
+
+//======================================================================================================================
+void UPhysicsControlComponent::AddBodyModifierToSet(
+	FPhysicsControlNames& NewSet, 
+	const FName           BodyModifier, 
+	const FName           SetName)
+{
+	Implementation->NameRecords.AddBodyModifier(BodyModifier, SetName);
+	NewSet.Names = GetBodyModifierNamesInSet(SetName);
+}
+
+
+//======================================================================================================================
+void UPhysicsControlComponent::AddBodyModifiersToSet(
+	FPhysicsControlNames& NewSet, 
+	const TArray<FName>&  BodyModifiers, 
+	const FName           SetName)
+{
+	for (FName BodyModifier : BodyModifiers)
+	{
+		Implementation->NameRecords.AddBodyModifier(BodyModifier, SetName);
+	}
+	NewSet.Names = GetBodyModifierNamesInSet(SetName);
+}
+
+//======================================================================================================================
+TArray<FName> UPhysicsControlComponent::GetSetsContainingControl(const FName Control) const
+{
+	TArray<FName> Result;
+	for (const TPair<FName, TArray<FName>>& ControlSetPair : Implementation->NameRecords.ControlSets)
+	{
+		for (const FName& ControlName : ControlSetPair.Value)
+		{
+			if (ControlName == Control)
+			{
+				Result.Add(ControlSetPair.Key);
+			}
+		}
+	}
+	return Result;
+}
+
+//======================================================================================================================
+TArray<FName> UPhysicsControlComponent::GetSetsContainingBodyModifier(const FName BodyModifier) const
+{
+	TArray<FName> Result;
+	for (const TPair<FName, TArray<FName>>& BodyModifierSetPair : Implementation->NameRecords.BodyModifierSets)
+	{
+		for (const FName& BodyModifierName : BodyModifierSetPair.Value)
+		{
+			if (BodyModifierName == BodyModifier)
+			{
+				Result.Add(BodyModifierSetPair.Key);
+			}
+		}
+	}
+	return Result;
+}
+
 //======================================================================================================================
 void UPhysicsControlComponent::DebugDraw(FPrimitiveDrawInterface* PDI) const
 {
+	// Draw gizmos
 	if (bShowDebugVisualization && VisualizationSizeScale > 0)
 	{
 		for (const TPair<FName, FPhysicsControlRecord>& PhysicsControlRecordPair : 
 			Implementation->PhysicsControlRecords)
 		{
-			const FName& Name = PhysicsControlRecordPair.Key;
+			const FName Name = PhysicsControlRecordPair.Key;
 			const FPhysicsControlRecord& Record = PhysicsControlRecordPair.Value;
 			DebugDrawControl(PDI, Record, Name);
 		} 
 	}
+
+	// Detailed controls - if there's a filter
+	if (!DebugControlDetailFilter.IsEmpty())
+	{
+		for (const TPair<FName, FPhysicsControlRecord>& PhysicsControlRecordPair :
+			Implementation->PhysicsControlRecords)
+		{
+			const FName Name = PhysicsControlRecordPair.Key;
+
+			if (Name.ToString().Contains(DebugControlDetailFilter))
+			{
+				const FPhysicsControlRecord& Record = PhysicsControlRecordPair.Value;
+
+				FString ParentComponentName = Record.PhysicsControl.ParentMeshComponent ?
+					Record.PhysicsControl.ParentMeshComponent->GetName() : TEXT("NoParent");
+				FString ChildComponentName = Record.PhysicsControl.ChildMeshComponent ?
+					Record.PhysicsControl.ChildMeshComponent->GetName() : TEXT("NoChild");
+
+				FString Text = FString::Printf(
+					TEXT("%s: Parent %s (%s) Child %s (%s): Linear strength %f Angular strength %f"),
+					*Name.ToString(),
+					*ParentComponentName,
+					*Record.PhysicsControl.ParentBoneName.ToString(),
+					*ChildComponentName,
+					*Record.PhysicsControl.ChildBoneName.ToString(),
+					Record.PhysicsControl.ControlData.LinearStrength,
+					Record.PhysicsControl.ControlData.AngularStrength);
+
+				GEngine->AddOnScreenDebugMessage(
+					-1, 0.0f,
+					Record.PhysicsControlState.bEnabled ? FColor::Green : FColor::Red, Text);
+			}
+		}
+	}
+
+	// Summary of control list
+	if (bShowDebugControlList)
+	{
+		FString AllNames;
+
+		for (const TPair<FName, FPhysicsControlRecord>& PhysicsControlRecordPair :
+			Implementation->PhysicsControlRecords)
+		{
+			const FName Name = PhysicsControlRecordPair.Key;
+			AllNames += Name.ToString() + TEXT(" ");
+			if (AllNames.Len() > 256)
+			{
+				GEngine->AddOnScreenDebugMessage(-1, 0.0f, FColor::White, *AllNames);
+				AllNames.Reset();
+			}
+		}
+		GEngine->AddOnScreenDebugMessage(
+			-1, 0.0f, FColor::White,
+			FString::Printf(TEXT("%d Controls: %s"), Implementation->PhysicsControlRecords.Num(), *AllNames));
+
+	}
+
+	// Detailed body modifiers - if there's a filter
+	if (!DebugBodyModifierDetailFilter.IsEmpty())
+	{
+		for (const TPair<FName, FPhysicsBodyModifier>& PhysicsBodyModifierPair :
+			Implementation->PhysicsBodyModifiers)
+		{
+			const FName Name = PhysicsBodyModifierPair.Key;
+
+			if (Name.ToString().Contains(DebugBodyModifierDetailFilter))
+			{
+				const FPhysicsBodyModifier& Record = PhysicsBodyModifierPair.Value;
+
+				FString ComponentName = Record.MeshComponent ? Record.MeshComponent->GetName() : TEXT("None");
+
+				FString Text = FString::Printf(
+					TEXT("%s: %s: %s %s GravityMultiplier %f"),
+					*Name.ToString(),
+					*ComponentName,
+					*UEnum::GetValueAsString(Record.MovementType),
+					*UEnum::GetValueAsString(Record.CollisionType),
+					Record.GravityMultiplier);
+
+				GEngine->AddOnScreenDebugMessage(-1, 0.0f, 
+					Record.MovementType == EPhysicsMovementType::Simulated ? FColor::Green : FColor::Red, Text);
+			}
+		}
+	}
+
+	// Summary of control list
+	if (bShowDebugBodyModifierList)
+	{
+		FString AllNames;
+
+		for (const TPair<FName, FPhysicsBodyModifier>& PhysicsBodyModifierPair : Implementation->PhysicsBodyModifiers)
+		{
+			const FName Name = PhysicsBodyModifierPair.Key;
+			AllNames += Name.ToString() + TEXT(" ");
+			if (AllNames.Len() > 256)
+			{
+				GEngine->AddOnScreenDebugMessage(-1, 0.0f, FColor::White, *AllNames);
+				AllNames.Reset();
+			}
+		}
+		GEngine->AddOnScreenDebugMessage(
+			-1, 0.0f, FColor::White,
+			FString::Printf(TEXT("%d Body modifiers: %s"), Implementation->PhysicsBodyModifiers.Num(), *AllNames));
+
+	}
 }
 
 //======================================================================================================================
@@ -1476,8 +1847,8 @@ void UPhysicsControlComponent::DebugDrawControl(
 
 	const FConstraintInstance* ConstraintInstance = Record.PhysicsControlState.ConstraintInstance.Get();
 
-	bool bHaveLinear = Record.PhysicsControl.ControlData.LinearStrength > 0;
-	bool bHaveAngular = Record.PhysicsControl.ControlData.AngularStrength > 0;
+	const bool bHaveLinear = Record.PhysicsControl.ControlData.LinearStrength > 0;
+	const bool bHaveAngular = Record.PhysicsControl.ControlData.AngularStrength > 0;
 
 	if (Record.PhysicsControlState.bEnabled && ConstraintInstance)
 	{
@@ -1491,7 +1862,7 @@ void UPhysicsControlComponent::DebugDrawControl(
 
 		FBodyInstance* ParentBodyInstance = UE::PhysicsControlComponent::GetBodyInstance(
 			Record.PhysicsControl.ParentMeshComponent, Record.PhysicsControl.ParentBoneName);
-		FTransform ParentBodyTM = ParentBodyInstance ? ParentBodyInstance->GetUnrealWorldTransform() : FTransform();
+		const FTransform ParentBodyTM = ParentBodyInstance ? ParentBodyInstance->GetUnrealWorldTransform() : FTransform();
 
 		FTransform TargetTM;
 		FVector TargetVelocity;
@@ -1499,14 +1870,14 @@ void UPhysicsControlComponent::DebugDrawControl(
 		Implementation->CalculateControlTargetData(TargetTM, TargetVelocity, TargetAngularVelocity, Record, true);
 
 		// WorldChildFrameTM is the world-space transform of the child (driven) constraint frame
-		FTransform WorldChildFrameTM = ConstraintInstance->GetRefFrame(EConstraintFrame::Frame1) * ChildBodyTM;
+		const FTransform WorldChildFrameTM = ConstraintInstance->GetRefFrame(EConstraintFrame::Frame1) * ChildBodyTM;
 
 		// WorldParentFrameTM is the world-space transform of the parent constraint frame
-		FTransform WorldParentFrameTM = ConstraintInstance->GetRefFrame(EConstraintFrame::Frame2) * ParentBodyTM;
+		const FTransform WorldParentFrameTM = ConstraintInstance->GetRefFrame(EConstraintFrame::Frame2) * ParentBodyTM;
 
-		FTransform WorldTargetTM = TargetTM * WorldParentFrameTM;
-		FTransform WorldCurrentTM = WorldChildFrameTM;
+		const FTransform WorldCurrentTM = WorldChildFrameTM;
 
+		FTransform WorldTargetTM = TargetTM * WorldParentFrameTM;
 		if (!bHaveLinear)
 		{
 			WorldTargetTM.SetTranslation(WorldCurrentTM.GetTranslation());
diff --git a/Engine/Plugins/Experimental/PhysicsControl/Source/PhysicsControl/Private/PhysicsControlComponentImpl.cpp b/Engine/Plugins/Experimental/PhysicsControl/Source/PhysicsControl/Private/PhysicsControlComponentImpl.cpp
--- a/Engine/Plugins/Experimental/PhysicsControl/Source/PhysicsControl/Private/PhysicsControlComponentImpl.cpp
+++ b/Engine/Plugins/Experimental/PhysicsControl/Source/PhysicsControl/Private/PhysicsControlComponentImpl.cpp
@@ -528,43 +528,43 @@ bool FPhysicsControlComponentImpl::ApplyControlStrengths(
 	FPhysicsControlRecord& Record, FConstraintInstance* ConstraintInstance)
 {
 	const FPhysicsControlData& Data = Record.PhysicsControl.ControlData;
-	const FPhysicsControlMultipliers& Multipliers = Record.PhysicsControl.ControlMultipliers;
+	const FPhysicsControlMultiplier& Multiplier = Record.PhysicsControl.ControlMultiplier;
 
 	double AngularSpring;
 	double AngularDamping;
-	double MaxTorque = Data.MaxTorque * Multipliers.MaxTorqueMultiplier;
+	double MaxTorque = Data.MaxTorque * Multiplier.MaxTorqueMultiplier;
 
 	FVector LinearSpring;
 	FVector LinearDamping;
-	FVector MaxForce = Data.MaxForce * Multipliers.MaxForceMultiplier;
+	FVector MaxForce = Data.MaxForce * Multiplier.MaxForceMultiplier;
 
 	UE::PhysicsControlComponent::ConvertStrengthToSpringParams(
 		AngularSpring, AngularDamping,
-		Data.AngularStrength * Multipliers.AngularStrengthMultiplier,
+		Data.AngularStrength * Multiplier.AngularStrengthMultiplier,
 		Data.AngularDampingRatio,
-		Data.AngularExtraDamping * Multipliers.AngularExtraDampingMultiplier);
+		Data.AngularExtraDamping * Multiplier.AngularExtraDampingMultiplier);
 	UE::PhysicsControlComponent::ConvertStrengthToSpringParams(
 		LinearSpring, LinearDamping,
-		Data.LinearStrength * Multipliers.LinearStrengthMultiplier,
+		Data.LinearStrength * Multiplier.LinearStrengthMultiplier,
 		Data.LinearDampingRatio,
-		Data.LinearExtraDamping * Multipliers.LinearExtraDampingMultiplier);
+		Data.LinearExtraDamping * Multiplier.LinearExtraDampingMultiplier);
 
-	if (Multipliers.MaxTorqueMultiplier <= 0.0)
+	if (Multiplier.MaxTorqueMultiplier <= 0.0)
 	{
 		AngularSpring = 0.0;
 		AngularDamping = 0.0;
 	}
-	if (Multipliers.MaxForceMultiplier.X <= 0.0)
+	if (Multiplier.MaxForceMultiplier.X <= 0.0)
 	{
 		LinearSpring.X = 0.0;
 		LinearDamping.X = 0.0;
 	}
-	if (Multipliers.MaxForceMultiplier.Y <= 0.0)
+	if (Multiplier.MaxForceMultiplier.Y <= 0.0)
 	{
 		LinearSpring.Y = 0.0;
 		LinearDamping.Y = 0.0;
 	}
-	if (Multipliers.MaxForceMultiplier.Z <= 0.0)
+	if (Multiplier.MaxForceMultiplier.Z <= 0.0)
 	{
 		LinearSpring.Z = 0.0;
 		LinearDamping.Z = 0.0;
@@ -635,3 +635,57 @@ FPhysicsBodyModifier* FPhysicsControlComponentImpl::FindBodyModifier(const FName
 	return PhysicsBodyModifiers.Find(Name);
 }
 
+//======================================================================================================================
+bool FPhysicsControlComponentImpl::DestroyControl(
+	const FName            Name,
+	const EDestroyBehavior DestroyBehavior)
+{
+	FPhysicsControlRecord* PhysicsControlRecord = FindControlRecord(Name);
+	if (PhysicsControlRecord)
+	{
+		if (USkeletalMeshComponent* SkeletalMeshComponent =
+			Cast<USkeletalMeshComponent>(PhysicsControlRecord->PhysicsControl.ParentMeshComponent))
+		{
+			RemoveSkeletalMeshReferenceForCaching(SkeletalMeshComponent);
+		}
+		if (USkeletalMeshComponent* SkeletalMeshComponent =
+			Cast<USkeletalMeshComponent>(PhysicsControlRecord->PhysicsControl.ChildMeshComponent))
+		{
+			RemoveSkeletalMeshReferenceForCaching(SkeletalMeshComponent);
+		}
+
+		PhysicsControlRecord->PhysicsControlState.Reset();
+		NameRecords.RemoveControl(Name);
+		if (DestroyBehavior == EDestroyBehavior::RemoveRecord)
+		{
+			check(PhysicsControlRecords.Remove(Name) == 1);
+		}
+		return true;
+	}
+	return false;
+}
+
+//======================================================================================================================
+bool FPhysicsControlComponentImpl::DestroyBodyModifier(
+	const FName            Name,
+	const EDestroyBehavior DestroyBehavior)
+{
+	FPhysicsBodyModifier* PhysicsBodyModifier = FindBodyModifier(Name);
+	if (PhysicsBodyModifier)
+	{
+		if (USkeletalMeshComponent* SkeletalMeshComponent =
+			Cast<USkeletalMeshComponent>(PhysicsBodyModifier->MeshComponent))
+		{
+			RemoveSkeletalMeshReferenceForCaching(SkeletalMeshComponent);
+			RemoveSkeletalMeshReferenceForModifier(SkeletalMeshComponent);
+		}
+		NameRecords.RemoveBodyModifier(Name);
+		if (DestroyBehavior == EDestroyBehavior::RemoveRecord)
+		{
+			check(PhysicsBodyModifiers.Remove(Name) == 1);
+		}
+		return true;
+	}
+	return false;
+}
+
diff --git a/Engine/Plugins/Experimental/PhysicsControl/Source/PhysicsControl/Private/PhysicsControlComponentImpl.h b/Engine/Plugins/Experimental/PhysicsControl/Source/PhysicsControl/Private/PhysicsControlComponentImpl.h
--- a/Engine/Plugins/Experimental/PhysicsControl/Source/PhysicsControl/Private/PhysicsControlComponentImpl.h
+++ b/Engine/Plugins/Experimental/PhysicsControl/Source/PhysicsControl/Private/PhysicsControlComponentImpl.h
@@ -3,6 +3,8 @@
 #pragma once
 
 #include "PhysicsControlRecord.h"
+#include "PhysicsControlNameRecords.h"
+
 #include "UObject/ObjectPtr.h"
 #include "UObject/NameTypes.h"
 
@@ -118,7 +120,26 @@ public:
 	 */
 	FPhysicsBodyModifier* FindBodyModifier(const FName Name);
 
+	enum class EDestroyBehavior : uint8
+	{
+		KeepRecord,
+		RemoveRecord
+	};
+
+	/** Destroys the control. It will optionally be removed from the array of records too */
+	bool DestroyControl(
+		const FName            Name, 
+		const EDestroyBehavior DestroyBehavior = EDestroyBehavior::RemoveRecord);
+
+	/** Destroys the modifier. It will optionally be removed from the array of records too */
+	bool DestroyBodyModifier(
+		const FName            Name,
+		const EDestroyBehavior DestroyBehavior = EDestroyBehavior::RemoveRecord);
+
+	/** Calculates and returns a unique name based on the bone names and any existing records */
 	FName GetUniqueControlName(const FName ParentBoneName, const FName ChildBoneName) const;
+
+	/** Calculates and returns a unique name based on the bone name and any existing records */
 	FName GetUniqueBodyModifierName(const FName BoneName) const;
 
 public:
@@ -134,7 +155,12 @@ public:
 	TMap<FName, FPhysicsControlRecord> PhysicsControlRecords;
 	TMap<FName, FPhysicsBodyModifier> PhysicsBodyModifiers;
 
+public:
+	// Keep track of the names of everything we have created
+	FPhysicsControlNameRecords NameRecords;
+
 private:
-	UPhysicsControlComponent* Owner;
+	UPhysicsControlComponent* Owner = nullptr;
+
 };
 
diff --git a/Engine/Plugins/Experimental/PhysicsControl/Source/PhysicsControl/Private/PhysicsControlNameRecords.cpp b/Engine/Plugins/Experimental/PhysicsControl/Source/PhysicsControl/Private/PhysicsControlNameRecords.cpp
new file mode 100644
--- /dev/null
+++ b/Engine/Plugins/Experimental/PhysicsControl/Source/PhysicsControl/Private/PhysicsControlNameRecords.cpp
@@ -0,0 +1,57 @@
+#include "PhysicsControlNameRecords.h"
+
+//======================================================================================================================
+void FPhysicsControlNameRecords::AddControl(FName Name, FName SetName)
+{
+	ControlSets.FindOrAdd(SetName).Add(Name);
+	ControlSets.FindOrAdd("All").Add(Name);
+}
+
+//======================================================================================================================
+void FPhysicsControlNameRecords::RemoveControl(FName Name)
+{
+	for (TPair<FName, TArray<FName>>& Set : ControlSets)
+	{
+		Set.Value.Remove(Name);
+	}
+}
+
+//======================================================================================================================
+const TArray<FName>& FPhysicsControlNameRecords::GetControlNamesInSet(FName SetName) const
+{
+	const TArray<FName>* Names = ControlSets.Find(SetName);
+	if (Names)
+	{
+		return *Names;
+	}
+	static TArray<FName> FailureResult;
+	return FailureResult;
+}
+
+//======================================================================================================================
+void FPhysicsControlNameRecords::AddBodyModifier(FName Name, FName Set)
+{
+	BodyModifierSets.FindOrAdd(Set).Add(Name);	
+	BodyModifierSets.FindOrAdd("All").Add(Name);
+}
+
+//======================================================================================================================
+void FPhysicsControlNameRecords::RemoveBodyModifier(FName Name)
+{
+	for (TPair<FName, TArray<FName>>& Set : BodyModifierSets)
+	{
+		Set.Value.Remove(Name);
+	}
+}
+
+//======================================================================================================================
+const TArray<FName>& FPhysicsControlNameRecords::GetBodyModifierNamesInSet(FName Set) const
+{
+	const TArray<FName>* Names = BodyModifierSets.Find(Set);
+	if (Names)
+	{
+		return *Names;
+	}
+	static TArray<FName> FailureResult;
+	return FailureResult;
+}
diff --git a/Engine/Plugins/Experimental/PhysicsControl/Source/PhysicsControl/Private/PhysicsControlNameRecords.h b/Engine/Plugins/Experimental/PhysicsControl/Source/PhysicsControl/Private/PhysicsControlNameRecords.h
new file mode 100644
--- /dev/null
+++ b/Engine/Plugins/Experimental/PhysicsControl/Source/PhysicsControl/Private/PhysicsControlNameRecords.h
@@ -0,0 +1,40 @@
+// Copyright Epic Games, Inc. All Rights Reserved.
+
+#pragma once
+
+#include "PhysicsControlLimbData.h"
+
+/**
+ * All the controls and body modifiers we create get given names, so we can keep a record in a form
+ * that will be useful to users. The three general uses are:
+ * 1. All controls or modifiers
+ * 2. The type of control - the obvious ones are "world space" and "parent space", but users might make others
+ * 3. Sets of control/modifiers - the obvious ones are limbs (e.g. a set of world-space controls on the leg)
+ *    but users might want to make other sets - e.g. "UpperBody"
+ */
+struct FPhysicsControlNameRecords
+{
+	/** Adds Name to SetName as well as adding it to the set "All" */
+	void AddControl(FName Name, FName SetName);
+	/** Removes Name from all Sets */
+	void RemoveControl(FName Name);
+	const TArray<FName>& GetControlNamesInSet(FName SetName) const;
+
+	/** Adds Name to SetName as well as adding it to the set "All" */
+	void AddBodyModifier(FName Name, FName SetName);
+	/** Removes Name from all Sets */
+	void RemoveBodyModifier(FName Name);
+	const TArray<FName>& GetBodyModifierNamesInSet(FName SetName) const;
+
+	/**
+	 * All the control sets we've created, arranged by set name. Note that some controls will
+	 * not necessarily have been added to a set.
+	 */
+	TMap<FName, TArray<FName>> ControlSets;
+
+	/**
+	 * All the body modifiers we've created, arranged by set name. Note that some modifiers will
+	 * not necessarily have been added to a set.
+	 * */
+	TMap<FName, TArray<FName>> BodyModifierSets;
+};
diff --git a/Engine/Plugins/Experimental/PhysicsControl/Source/PhysicsControl/Private/PhysicsControlRecord.h b/Engine/Plugins/Experimental/PhysicsControl/Source/PhysicsControl/Private/PhysicsControlRecord.h
--- a/Engine/Plugins/Experimental/PhysicsControl/Source/PhysicsControl/Private/PhysicsControlRecord.h
+++ b/Engine/Plugins/Experimental/PhysicsControl/Source/PhysicsControl/Private/PhysicsControlRecord.h
@@ -15,18 +15,12 @@ class UMeshComponent;
  */
 struct FPhysicsControlState
 {
-	FPhysicsControlState()
-		: bEnabled(false), bPendingDestroy(false)
-	{
-	}
-
 	/** Removes any constraint and resets the state */
 	void Reset();
 
 	TSharedPtr<FConstraintInstance> ConstraintInstance;
 
-	bool bEnabled;
-	bool bPendingDestroy;
+	bool bEnabled = false;
 };
 
 /**
@@ -34,8 +28,7 @@ struct FPhysicsControlState
  */
 struct FPhysicsControlRecord
 {
-	FPhysicsControlRecord(const FPhysicsControl& InControl)
-		: PhysicsControl(InControl)
+	FPhysicsControlRecord(const FPhysicsControl& InControl) : PhysicsControl(InControl)
 	{
 	}
 
@@ -78,7 +71,6 @@ struct FPhysicsBodyModifier
 		, KinematicTargetPosition(FVector::ZeroVector)
 		, KinematicTargetOrientation(FQuat::Identity)
 		, bUseSkeletalAnimation(InUseSkeletalAnimation)
-		, bPendingDestroy(false)
 	{}
 
 	/**  The mesh that will be modified. */
@@ -88,34 +80,31 @@ struct FPhysicsBodyModifier
 	FName BoneName;
 
 	/** How the associated body should move. */
-	EPhysicsMovementType MovementType;
+	EPhysicsMovementType MovementType = EPhysicsMovementType::Kinematic;
 
 	/** How the associated body should collide/interact */
-	ECollisionEnabled::Type CollisionType;
+	ECollisionEnabled::Type CollisionType = ECollisionEnabled::QueryAndPhysics;
 
 	/**
 	 * Multiplier for gravity applied to the body. Note that if the body itself has gravity disabled, then
 	 * setting this to 1 will not enable gravity.
 	 */
-	float GravityMultiplier;
+	float GravityMultiplier = 1.0f;
 
 	/** 
 	 * The target position when kinematic. Note that this is applied on top of any animation 
 	 * target if bUseSkeletalAnimation is set. 
 	 */
-	FVector KinematicTargetPosition;
+	FVector KinematicTargetPosition = FVector::ZeroVector;
 
 	/**
 	 * The target orientation when kinematic. Note that this is applied on top of any animation
 	 * target if bUseSkeletalAnimation is set.
 	 */
-	FQuat KinematicTargetOrientation;
+	FQuat KinematicTargetOrientation = FQuat::Identity;
 
 	/** If true then the target will be applied on top of the skeletal animation (if there is any) */
 	uint8 bUseSkeletalAnimation:1;
-
-	/** Indicates if this should be deleted on the next tick */
-	uint8 bPendingDestroy:1;
 };
 
 /**
diff --git a/Engine/Plugins/Experimental/PhysicsControl/Source/PhysicsControl/Public/PhysicsControlActor.h b/Engine/Plugins/Experimental/PhysicsControl/Source/PhysicsControl/Public/PhysicsControlActor.h
new file mode 100644
--- /dev/null
+++ b/Engine/Plugins/Experimental/PhysicsControl/Source/PhysicsControl/Public/PhysicsControlActor.h
@@ -0,0 +1,265 @@
+// Copyright Epic Games, Inc. All Rights Reserved.
+
+#pragma once
+
+#include "GameFramework/Actor.h"
+#include "PhysicsControlData.h"
+#include "PhysicsControlActor.generated.h"
+
+class UPhysicsControlComponent;
+struct FPhysicsControl;
+struct FPhysicsControlLimbSetupData;
+
+/**
+ * Structure that determines a Physics Control used during initialization of the physics control actor
+ */
+USTRUCT(BlueprintType)
+struct PHYSICSCONTROL_API FInitialPhysicsControl
+{
+	GENERATED_BODY()
+
+	FInitialPhysicsControl() {}
+
+	/** The owner of the mesh that will be doing the driving. Blank/non-existent means it will happen in world space */
+	UPROPERTY(EditAnywhere, BlueprintReadWrite, Category = PhysicsControl)
+	TObjectPtr<AActor> ParentActor;
+
+	/** 
+	 * The mesh that will be doing the driving. If this is blank but there is an actor, then we'll attempt to
+	 * use the root component. If that doesn't work then it will happen in world space.
+	 */
+	UPROPERTY(EditAnywhere, BlueprintReadWrite, Category = PhysicsControl)
+	FName ParentMeshComponentName;
+
+	/** 
+	 * If the parent mesh component is skeletal, then the name of the skeletal mesh bone that will be doing 
+	 * the driving. 
+	 */
+	UPROPERTY(EditAnywhere, BlueprintReadWrite, Category = PhysicsControl)
+	FName ParentBoneName;
+
+	/** The owner of the mesh that the control will be driving */
+	UPROPERTY(EditAnywhere, BlueprintReadWrite, Category = PhysicsControl)
+	TObjectPtr<AActor> ChildActor;
+
+	/** 
+	 * The mesh that the control will be driving. If this is blank but there is an actor, then we'll attempt 
+	 * to use the root component 
+	 */
+	UPROPERTY(EditAnywhere, BlueprintReadWrite, Category = PhysicsControl)
+	FName ChildMeshComponentName;
+
+	/** 
+	 * If the child mesh component is skeletal, then the name of the skeletal mesh bone that the control 
+	 * will be driving. 
+	 */
+	UPROPERTY(EditAnywhere, BlueprintReadWrite, Category = PhysicsControl)
+	FName ChildBoneName;
+
+	/** 
+	 * Strength and damping parameters. Can be modified at any time, but will sometimes have 
+	 * been set once during initialization 
+	 */
+	UPROPERTY(EditAnywhere, BlueprintReadWrite, Category = PhysicsControl)
+	FPhysicsControlData ControlData;
+
+	/**
+	 * Multiplier for the ControlData. This will typically be modified dynamically, and also expose the ability
+	 * to set directional strengths
+	 */
+	UPROPERTY(EditAnywhere, BlueprintReadWrite, Category = PhysicsControl)
+	FPhysicsControlMultiplier ControlMultiplier;
+
+	/**
+	 * The position/orientation etc targets for the controls. These are procedural/explicit control targets -
+	 * skeletal meshes have the option to use skeletal animation as well, in which case these targets are 
+	 * expressed as relative to that animation.
+	 */
+	UPROPERTY(EditAnywhere, BlueprintReadWrite, Category = PhysicsControl)
+	FPhysicsControlTarget ControlTarget;
+
+	/**
+	 * More general settings for the control
+	 */
+	UPROPERTY(EditAnywhere, BlueprintReadWrite, Category = PhysicsControl)
+	FPhysicsControlSettings ControlSettings;
+};
+
+/**
+ * Structure that determines a Body Modifier used during initialization of the physics control actor
+ */
+USTRUCT(BlueprintType)
+struct PHYSICSCONTROL_API FInitialBodyModifier
+{
+	GENERATED_BODY()
+
+	FInitialBodyModifier() {}
+
+	/** The owner of the mesh that that we will modify */
+	UPROPERTY(EditAnywhere, BlueprintReadWrite, Category = PhysicsControl)
+	TObjectPtr<AActor> Actor;
+
+	/** 
+	 * The mesh that will be modify. If this is blank but there is an actor, then we'll attempt to
+	 * use the root component. 
+	 */
+	UPROPERTY(EditAnywhere, BlueprintReadWrite, Category = PhysicsControl)
+	FName MeshComponentName;
+
+	/** 
+	 * If the mesh component is skeletal, then the name of the skeletal mesh bone to modify
+	 */
+	UPROPERTY(EditAnywhere, BlueprintReadWrite, Category = PhysicsControl)
+	FName BoneName;
+
+	/** How the associated body should move. */
+	UPROPERTY(EditAnywhere, BlueprintReadWrite, Category = PhysicsControl)
+	EPhysicsMovementType MovementType = EPhysicsMovementType::Kinematic;
+
+	/** How the associated body should collide/interact */
+	ECollisionEnabled::Type CollisionType = ECollisionEnabled::QueryAndPhysics;
+
+	/**
+	 * Multiplier for gravity applied to the body. Note that if the body itself has gravity disabled, then
+	 * setting this to 1 will not enable gravity.
+	 */
+	UPROPERTY(EditAnywhere, BlueprintReadWrite, Category = PhysicsControl)
+	float GravityMultiplier = 1.0f;
+
+	/**
+	 * The target position when kinematic. Note that this is applied on top of any animation
+	 * target if bUseSkeletalAnimation is set.
+	 */
+	UPROPERTY(EditAnywhere, BlueprintReadWrite, Category = PhysicsControl)
+	FVector KinematicTargetPosition = FVector::ZeroVector;
+
+	/**
+	 * The target orientation when kinematic. Note that this is applied on top of any animation
+	 * target if bUseSkeletalAnimation is set.
+	 */
+	UPROPERTY(EditAnywhere, BlueprintReadWrite, Category = PhysicsControl)
+	FQuat KinematicTargetOrientation = FQuat::Identity;
+
+	/** If true then the target will be applied on top of the skeletal animation (if there is any) */
+	UPROPERTY(EditAnywhere, BlueprintReadWrite, Category = PhysicsControl)
+	bool bUseSkeletalAnimation = true;
+};
+
+/**
+ * Structure that holds the data necessary to set up a default set of limb controls for a character
+ */
+USTRUCT(BlueprintType)
+struct PHYSICSCONTROL_API FInitialCharacterControls
+{
+	GENERATED_BODY()
+
+	FInitialCharacterControls() {}
+
+	/** The owner of the character skeletal mesh */
+	UPROPERTY(EditAnywhere, BlueprintReadWrite, Category = PhysicsControl)
+	TObjectPtr<AActor> CharacterActor;
+
+	/**
+	 * The skeletal mesh that will have controls set up. If this is blank but there is an actor, then we'll attempt to
+	 * use the root component. 
+	 */
+	UPROPERTY(EditAnywhere, BlueprintReadWrite, Category = PhysicsControl)
+	FName SkeletalMeshComponentName;
+
+	UPROPERTY(EditAnywhere, BlueprintReadWrite, Category = PhysicsControl)
+	TArray<FPhysicsControlLimbSetupData> LimbSetupData;
+
+	UPROPERTY(EditAnywhere, BlueprintReadWrite, Category = PhysicsControl)
+	FPhysicsControlData WorldSpaceControlData;
+
+	UPROPERTY(EditAnywhere, BlueprintReadWrite, Category = PhysicsControl)
+	FPhysicsControlSettings WorldSpaceControlSettings;
+
+	UPROPERTY(EditAnywhere, BlueprintReadWrite, Category = PhysicsControl)
+	bool bEnableWorldSpaceControls = false;
+
+	UPROPERTY(EditAnywhere, BlueprintReadWrite, Category = PhysicsControl)
+	FPhysicsControlData ParentSpaceControlData;
+
+	UPROPERTY(EditAnywhere, BlueprintReadWrite, Category = PhysicsControl)
+	FPhysicsControlSettings ParentSpaceControlSettings;
+
+	UPROPERTY(EditAnywhere, BlueprintReadWrite, Category = PhysicsControl)
+	bool bEnableParentSpaceControls = false;
+
+	UPROPERTY(EditAnywhere, BlueprintReadWrite, Category = PhysicsControl)
+	EPhysicsMovementType PhysicsMovementType = EPhysicsMovementType::Kinematic;
+};
+
+/**
+ * 
+ */
+UCLASS(meta = (BlueprintSpawnableComponent), ClassGroup = Physics, Experimental)
+class PHYSICSCONTROL_API UPhysicsControlInitializerComponent : public USceneComponent
+{
+	GENERATED_UCLASS_BODY()
+
+public:
+	virtual void BeginPlay() override;
+
+public:
+	/**
+	 * This creates all the controls specified in the Initial properties. You can call it explicitly,
+	 * or you can opt to have it called at BeginPlay using CreateControlsAtBeginPlay.
+	 */
+	UFUNCTION(BlueprintCallable, Category = PhysicsControl)
+	void CreateControls(UPhysicsControlComponent* PhysicsControlComponent);
+
+	/**
+	 * This can be filled in to automatically create whole-character controls (by specifying limbs etc) for
+	 * a skeletal mesh during the BeginPlay event.
+	 */
+	UPROPERTY(EditAnywhere, BlueprintReadWrite, Category = PhysicsControls)
+	FInitialCharacterControls InitialCharacterControls;
+
+	/**
+	 * This can be filled in to automatically create controls during the BeginPlay event. If a control
+	 * already exists with the name (e.g. created as part of InitialCharacterControls) then it will be updated.
+	 */
+	UPROPERTY(EditAnywhere, BlueprintReadWrite, Category = PhysicsControls)
+	TMap<FName, FInitialPhysicsControl> InitialControls;
+
+	/**
+	 * This can be filled in to automatically create body modifiers during the BeginPlay event. If a body modifier
+	 * already exists with the name (e.g. created as part of InitialCharacterControls) then it will be updated.
+	 */
+	UPROPERTY(EditAnywhere, BlueprintReadWrite, Category = PhysicsControls)
+	TMap<FName, FInitialBodyModifier> InitialBodyModifiers;
+
+	/**
+	 * If set, then CreateControls will be called in our BeginPlay event, attempting to find a 
+	 * PhysicsControlComponent in our parent actor.
+	 */
+	UPROPERTY(EditAnywhere, BlueprintReadWrite, Category = PhysicsControls)
+	bool bCreateControlsAtBeginPlay = true;
+
+private:
+	void CreateInitialCharacterControls(UPhysicsControlComponent* ControlComponent);
+	void CreateOrUpdateInitialControls(UPhysicsControlComponent* ControlComponent);
+	void CreateOrUpdateInitialBodyModifiers(UPhysicsControlComponent* ControlComponent);
+};
+
+/**
+ * Actor that holds a PhysicsControlComponent which is then used to apply controls/body modifiers to
+ * other actors from "outside", and a PhysicsControlInitializerComponent which provides the ability
+ * to create those controls automatically
+ */
+UCLASS(ConversionRoot, MinimalAPI, ComponentWrapperClass)
+class APhysicsControlActor : public AActor
+{
+	GENERATED_UCLASS_BODY()
+
+private:
+	UPROPERTY(Category = ControlActor, VisibleAnywhere, BlueprintReadOnly, meta = (AllowPrivateAccess = "true"))
+	TObjectPtr<UPhysicsControlComponent> ControlComponent;
+
+	UPROPERTY(Category = ControlActor, VisibleAnywhere, BlueprintReadOnly, meta = (AllowPrivateAccess = "true"))
+	TObjectPtr<UPhysicsControlInitializerComponent> ControlInitializerComponent;
+};
+
+
diff --git a/Engine/Plugins/Experimental/PhysicsControl/Source/PhysicsControl/Public/PhysicsControlComponent.h b/Engine/Plugins/Experimental/PhysicsControl/Source/PhysicsControl/Public/PhysicsControlComponent.h
--- a/Engine/Plugins/Experimental/PhysicsControl/Source/PhysicsControl/Public/PhysicsControlComponent.h
+++ b/Engine/Plugins/Experimental/PhysicsControl/Source/PhysicsControl/Public/PhysicsControlComponent.h
@@ -6,7 +6,7 @@
 #include "PhysicsControlLimbData.h"
 
 #include "UObject/ObjectMacros.h"
-#include "Components/ActorComponent.h"
+#include "Components/SceneComponent.h"
 #include "EngineDefines.h"
 #include "Templates/PimplPtr.h"
 
@@ -47,56 +47,64 @@ enum class EPhysicsControlType : uint8
  * (perhaps in arrays) to make it easy to quickly change multiple Controls/Body Modifiers.
  */
 UCLASS(meta = (BlueprintSpawnableComponent), ClassGroup = Physics, Experimental)
-class PHYSICSCONTROL_API UPhysicsControlComponent : public UActorComponent
+class PHYSICSCONTROL_API UPhysicsControlComponent : public USceneComponent
 {
 	GENERATED_UCLASS_BODY()
 
 public:
 	/**
-	 * Makes a new control for mesh components
+	 * Creates a new control for mesh components
 	 * 
 	 * @param ControlData   Describes the initial strength etc of the new control
 	 * @param ControlTarget Describes the initial target for the new control
 	 * @param ControlSettings General settings for the control
-	 * @param bEnabled      If true then the control will be enabled immediately. If false you will need to call
-	 *                      SetControlEnabled(true) in order to activate it.
+	 * @param Set Which set to include the control in (optional). Note that it automatically
+	 *        gets added to the set "All"
+	 * @param bEnabled If true then the control will be enabled immediately. If false you will need to call
+	 *        SetControlEnabled(true) in order to activate it.
 	 * @return The name of the new control
 	 */
 	UFUNCTION(BlueprintCallable, Category = PhysicsControl)
-	FName MakeControl(
-		UMeshComponent*         ParentMeshComponent,
-		FName                   ParentBoneName,
-		UMeshComponent*         ChildMeshComponent,
-		FName                   ChildBoneName,
-		FPhysicsControlData     ControlData,
-		FPhysicsControlTarget   ControlTarget, 
-		FPhysicsControlSettings ControlSettings,
-		bool                    bEnabled = true);
+	FName CreateControl(
+		UMeshComponent*               ParentMeshComponent,
+		FName                         ParentBoneName,
+		UMeshComponent*               ChildMeshComponent,
+		const FName                   ChildBoneName,
+		const FPhysicsControlData     ControlData,
+		const FPhysicsControlTarget   ControlTarget, 
+		const FPhysicsControlSettings ControlSettings,
+		FName                         Set,
+		const bool                    bEnabled = true
+	);
 
 	/**
-	 * Makes a new control for mesh components
+	 * Creates a new control for mesh components
 	 * 
+	 * @param The name of the control that will be created. Creation will fail if this name is already in use.
 	 * @param ControlData   Describes the initial strength etc of the new control
 	 * @param ControlTarget Describes the initial target for the new control
 	 * @param ControlSettings General settings for the control
-	 * @param bEnabled      If true then the control will be enabled immediately. If false you will need to call
-	 *                      SetControlEnabled(true) in order to enable it.
+	 * @param Set Which set to include the control in (optional). Note that it automatically
+	 *        gets added to the set "All"
+	 * @param bEnabled If true then the control will be enabled immediately. If false you will need to call
+	 *        SetControlEnabled(true) in order to enable it.
 	 * @return True if a new control was created, false if a control of the specified name already exists
 	 */
 	UFUNCTION(BlueprintCallable, Category = PhysicsControl)
-	bool MakeNamedControl(
-		FName                   Name,
-		UMeshComponent*         ParentMeshComponent,
-		FName                   ParentBoneName,
-		UMeshComponent*         ChildMeshComponent,
-		FName                   ChildBoneName,
-		FPhysicsControlData     ControlData, 
-		FPhysicsControlTarget   ControlTarget, 
-		FPhysicsControlSettings ControlSettings, 
-		bool                    bEnabled = true);
+	bool CreateNamedControl(
+		FName                         Name,
+		UMeshComponent*               ParentMeshComponent,
+		const FName                   ParentBoneName,
+		UMeshComponent*               ChildMeshComponent,
+		const FName                   ChildBoneName,
+		const FPhysicsControlData     ControlData, 
+		const FPhysicsControlTarget   ControlTarget, 
+		const FPhysicsControlSettings ControlSettings, 
+		const FName                   Set,
+		const bool                    bEnabled = true);
 
 	/**
-	 * Makes a collection of controls controlling a skeletal mesh
+	 * Creates a collection of controls controlling a skeletal mesh
 	 * 
 	 * @param SkeletalMeshComponent The skeletal mesh which will have controls
 	 * @param BoneName The name of the bone below which controls should be created. Each bone will be the child in a control
@@ -104,22 +112,25 @@ public:
 	 * @param ControlType What type of control to create. This determines what the parent will be for each control
 	 * @param ControlData   Describes the initial strength etc of the new control
 	 * @param ControlSettings General settings for the control
-	 * @param bEnabled      If true then the control will be enabled immediately. If false you will need to call
-	 *                      SetControlEnabled(true) in order to enable it.
+	 * @param Set Which set to include the control in (optional). Note that it automatically
+	 *        gets added to the set "All"
+	 * @param bEnabled If true then the control will be enabled immediately. If false you will need to call
+	 *        SetControlEnabled(true) in order to enable it.
 	 * @return An array of the controls that have been created
 	 */
 	UFUNCTION(BlueprintCallable, Category = PhysicsControl)
-	TArray<FName> MakeControlsFromSkeletalMeshBelow(
-		USkeletalMeshComponent* SkeletalMeshComponent,
-		FName                   BoneName,
-		bool                    bIncludeSelf,
-		EPhysicsControlType     ControlType,
-		FPhysicsControlData     ControlData,
-		FPhysicsControlSettings ControlSettings,
-		bool                    bEnabled = true);
+	TArray<FName> CreateControlsFromSkeletalMeshBelow(
+		USkeletalMeshComponent*       SkeletalMeshComponent,
+		const FName                   BoneName,
+		const bool                    bIncludeSelf,
+		const EPhysicsControlType     ControlType,
+		const FPhysicsControlData     ControlData,
+		const FPhysicsControlSettings ControlSettings,
+		const FName                   Set,
+		const bool                    bEnabled = true);
 
 	/**
-	 * Makes a collection of ParentSpace controls controlling a skeletal mesh, initializing
+	 * Creates a collection of ParentSpace controls controlling a skeletal mesh, initializing
 	 * them with a constraint profile
 	 * 
 	 * @param SkeletalMeshComponent The skeletal mesh which will have controls
@@ -135,41 +146,47 @@ public:
 	 *                          instead of slerp. Note also that the joint constraints do not use the animation
 	 *                          velocity as a target, so when creating controls in this way the control settings
 	 *                          will set the skeletal animation velocity multiplier to zero.
-	 * @param bEnabled      If true then the control will be enabled immediately. If false you will need to call
-	 *                      SetControlEnabled(true) in order to enable it.
+	 * @param Set Which set to include the control in (optional). Note that it automatically
+	 *        gets added to the set "All"
+	 * @param bEnabled If true then the control will be enabled immediately. If false you will need to call
+	 *        SetControlEnabled(true) in order to enable it.
 	 * @return An array of the controls that have been created
 	 */
 	UFUNCTION(BlueprintCallable, Category = PhysicsControl)
-	TArray<FName> MakeControlsFromSkeletalMeshAndConstraintProfileBelow(
+	TArray<FName> CreateControlsFromSkeletalMeshAndConstraintProfileBelow(
 		USkeletalMeshComponent* SkeletalMeshComponent,
-		FName                   BoneName,
-		bool                    bIncludeSelf,
-		FName                   ConstraintProfile,
-		bool                    bEnabled = true);
+		const FName             BoneName,
+		const bool              bIncludeSelf,
+		const FName             ConstraintProfile,
+		const FName             Set,
+		const bool              bEnabled = true);
 
 	/**
-	 * Makes a collection of controls controlling a skeletal mesh
+	 * Creates a collection of controls controlling a skeletal mesh
 	 *
 	 * @param SkeletalMeshComponent The skeletal mesh which will have controls
 	 * @param BoneNames The names of bones for which controls should be created. Each bone will be the child in a control
 	 * @param ControlType What type of control to create. This determines what the parent will be for each control
 	 * @param ControlData   Describes the initial strength etc of the new control
 	 * @param ControlSettings General settings for the control
-	 * @param bEnabled      If true then the control will be enabled immediately. If false you will need to call
-	 *                      SetControlEnabled(true) in order to enable it.
+	 * @param Set Which set to include the control in (optional). Note that it automatically
+	 *        gets added to the set "All"
+	 * @param bEnabled If true then the control will be enabled immediately. If false you will need to call
+	 *        SetControlEnabled(true) in order to enable it.
 	 * @return An array of the controls that have been created
 	 */
 	UFUNCTION(BlueprintCallable, Category = PhysicsControl)
-	TArray<FName> MakeControlsFromSkeletalMesh(
-		USkeletalMeshComponent* SkeletalMeshComponent,
-		const TArray<FName>&    BoneNames,
-		EPhysicsControlType     ControlType,
-		FPhysicsControlData     ControlData,
-		FPhysicsControlSettings ControlSettings,
-		bool                    bEnabled = true);
+	TArray<FName> CreateControlsFromSkeletalMesh(
+		USkeletalMeshComponent*       SkeletalMeshComponent,
+		const TArray<FName>&          BoneNames,
+		const EPhysicsControlType     ControlType,
+		const FPhysicsControlData     ControlData,
+		const FPhysicsControlSettings ControlSettings,
+		const FName                   Set,
+		const bool                    bEnabled = true);
 
 	/**
-	 * Makes a collection of ParentSpace controls controlling a skeletal mesh, initializing them 
+	 * Creates a collection of ParentSpace controls controlling a skeletal mesh, initializing them 
 	 * with a constraint profile
 	 *
 	 * @param SkeletalMeshComponent The skeletal mesh which will have controls
@@ -184,16 +201,19 @@ public:
 	 *                          instead of slerp. Note also that the joint constraints do not use the animation
 	 *                          velocity as a target, so when creating controls in this way the control settings
 	 *                          will set the skeletal animation velocity multiplier to zero.
-	 * @param bEnabled      If true then the control will be enabled immediately. If false you will need to call
-	 *                      SetControlEnabled(true) in order to enable it.
+	 * @param Set Which set to include the control in (optional). Note that it automatically
+	 *        gets added to the set "All"
+	 * @param bEnabled If true then the control will be enabled immediately. If false you will need to call
+	 *        SetControlEnabled(true) in order to enable it.
 	 * @return An array of the controls that have been created
 	 */
 	UFUNCTION(BlueprintCallable, Category = PhysicsControl)
-	TArray<FName> MakeControlsFromSkeletalMeshAndConstraintProfile(
+	TArray<FName> CreateControlsFromSkeletalMeshAndConstraintProfile(
 		USkeletalMeshComponent* SkeletalMeshComponent,
 		const TArray<FName>&    BoneNames,
-		FName                   ConstraintProfile,
-		bool                    bEnabled = true);
+		const FName             ConstraintProfile,
+		const FName             Set,
+		const bool              bEnabled = true);
 
 	/**
 	 * Calculates which bones belong to which limb in a skeletal mesh
@@ -208,11 +228,11 @@ public:
 	 */
 	UFUNCTION(BlueprintCallable, Category = PhysicsControl)
 	TMap<FName, FPhysicsControlLimbBones> GetLimbBonesFromSkeletalMesh(
-		USkeletalMeshComponent* SkeletalMeshComponent,
+		USkeletalMeshComponent*                     SkeletalMeshComponent,
 		const TArray<FPhysicsControlLimbSetupData>& LimbSetupData) const;
 
 	/**
-	 * Makes a collection of controls controlling a skeletal mesh, grouped together in limbs
+	 * Creates a collection of controls controlling a skeletal mesh, grouped together in limbs
 	 *
 	 * @param AllControls A single container for all the controls that have been created
 	 * @param LimbBones A map relating the limbs and the bones that they contain. Typically create this 
@@ -225,16 +245,16 @@ public:
 	 * @return A map containing the controls for each limb
 	 */
 	UFUNCTION(BlueprintCallable, Category = PhysicsControl)
-	TMap<FName, FPhysicsControlNameArray> MakeControlsFromLimbBones(
-		FPhysicsControlNameArray&                    AllControls,
+	TMap<FName, FPhysicsControlNames> CreateControlsFromLimbBones(
+		FPhysicsControlNames&                        AllControls,
 		const TMap<FName, FPhysicsControlLimbBones>& LimbBones,
-		EPhysicsControlType                          ControlType,
-		FPhysicsControlData                          ControlData,
-		FPhysicsControlSettings                      ControlSettings,
-		bool                                         bEnabled = true);
+		const EPhysicsControlType                    ControlType,
+		const FPhysicsControlData                    ControlData,
+		const FPhysicsControlSettings                ControlSettings,
+		const bool                                   bEnabled = true);
 
 	/**
-	 * Makes a collection of ParentSpace controls controlling a skeletal mesh, grouped together in limbs, initializing
+	 * Creates a collection of ParentSpace controls controlling a skeletal mesh, grouped together in limbs, initializing
 	 * them with a constraint profile
 	 *
 	 * @param AllControls A single container for all the controls that have been created
@@ -255,11 +275,11 @@ public:
 	 * @return A map containing the controls for each limb
 	 */
 	UFUNCTION(BlueprintCallable, Category = PhysicsControl)
-	TMap<FName, FPhysicsControlNameArray> MakeControlsFromLimbBonesAndConstraintProfile(
-		FPhysicsControlNameArray&                    AllControls,
+	TMap<FName, FPhysicsControlNames> CreateControlsFromLimbBonesAndConstraintProfile(
+		FPhysicsControlNames&                        AllControls,
 		const TMap<FName, FPhysicsControlLimbBones>& LimbBones,
-		FName                                        ConstraintProfile,
-		bool                                         bEnabled = true);
+		const FName                                  ConstraintProfile,
+		const bool                                   bEnabled = true);
 
 	/**
 	 * Destroys a control
@@ -268,7 +288,7 @@ public:
 	 * @return     Returns true if the control was found and destroyed, false if not
 	 */
 	UFUNCTION(BlueprintCallable, Category = PhysicsControl)
-	bool DestroyControl(FName Name);
+	bool DestroyControl(const FName Name);
 
 	/**
 	 * Destroys all controls
@@ -277,7 +297,16 @@ public:
 	 *              then it can be split.
 	 */
 	UFUNCTION(BlueprintCallable, Category = PhysicsControl)
-	void DestroyAllControls(const TArray<FName>& Names);
+	void DestroyControls(const TArray<FName>& Names);
+
+	/**
+	 * Destroys all controls in a set
+	 *
+	 * @param Set The set of controls to use to destroy. Standard sets will include "All", "WorldSpace",
+	 * "ParentSpace" and things like "WorldSpace-ArmLeft", depending on how controls have been created.
+	 */
+	UFUNCTION(BlueprintCallable, Category = PhysicsControl)
+	void DestroyControlsInSet(const FName Set);
 
 	/**
 	 * Modifies an existing control data - i.e. the strengths etc of the control driving towards the target
@@ -288,7 +317,7 @@ public:
 	 * @return Returns true if the control was found and modified, false if not
 	 */
 	UFUNCTION(BlueprintCallable, Category = PhysicsControl)
-	bool SetControlData(FName Name, FPhysicsControlData ControlData, bool bEnableControl = true);
+	bool SetControlData(const FName Name, const FPhysicsControlData ControlData, const bool bEnableControl = true);
 
 	/**
 	 * Modifies existing control data - i.e. the strengths etc of the controls driving towards the targets
@@ -299,7 +328,18 @@ public:
 	 * @param bEnableControl Enables the control if it is currently disabled
 	 */
 	UFUNCTION(BlueprintCallable, Category = PhysicsControl)
-	void SetAllControlDatas(const TArray<FName>& Names, FPhysicsControlData ControlData, bool bEnableControl = true);
+	void SetControlDatas(const TArray<FName>& Names, const FPhysicsControlData ControlData, const bool bEnableControl = true);
+
+	/**
+	 * Modifies existing control data - i.e. the strengths etc of the controls driving towards the targets
+	 *
+	 * @param Set The set of controls to modify. Standard sets will include "All", "WorldSpace",
+	 *        "ParentSpace" and things like "WorldSpace-ArmLeft", depending on how controls have been created.
+	 * @param ControlData The new control data
+	 * @param bEnableControl Enables the controls if currently disabled
+	 */
+	UFUNCTION(BlueprintCallable, Category = PhysicsControl)
+	void SetControlDatasInSet(const FName Set, const FPhysicsControlData ControlData, const bool bEnableControl = true);
 
 	/**
 	 * Modifies an existing control data using the multipliers
@@ -310,21 +350,35 @@ public:
 	 * @return Returns true if the control was found and modified, false if not
 	 */
 	UFUNCTION(BlueprintCallable, Category = PhysicsControl)
-	bool SetControlMultipliers(FName Name, FPhysicsControlMultipliers ControlMultipliers, bool bEnableControl = true);
+	bool SetControlMultiplier(const FName Name, const FPhysicsControlMultiplier ControlMultiplier, const bool bEnableControl = true);
 
 	/**
 	 * Modifies existing control data using the multipliers
 	 *
 	 * @param Names The names of the controls to modify. Note that if you have these in a FPhysicsControlNameArray
 	 *              then it can be split.
-	 * @param ControlMultipliers The new control multipliers
-	 * @param bEnableControl Enables the control if it is currently disabled
+	 * @param ControlMultiplier The new control multiplier
+	 * @param bEnableControl Enables the controls if currently disabled
 	 */
 	UFUNCTION(BlueprintCallable, Category = PhysicsControl)
-	void SetAllControlMultipliers(
-		const TArray<FName>&       Names, 
-		FPhysicsControlMultipliers ControlMultipliers, 
-		bool                       bEnableControl = true);
+	void SetControlMultipliers(
+		const TArray<FName>&            Names, 
+		const FPhysicsControlMultiplier ControlMultiplier, 
+		const bool                      bEnableControl = true);
+
+	/**
+	 * Modifies existing control data using the multipliers
+	 *
+	 * @param Set The set of controls to modify. Standard sets will include "All", "WorldSpace",
+	 *        "ParentSpace" and things like "WorldSpace-ArmLeft", depending on how controls have been created.
+	 * @param ControlMultiplier The new control multiplier
+	 * @param bEnableControl Enables the controls if currently disabled
+	 */
+	UFUNCTION(BlueprintCallable, Category = PhysicsControl)
+	void SetControlMultipliersInSet(
+		const FName                      Set, 
+		const FPhysicsControlMultiplier  ControlMultiplier, 
+		const bool                       bEnableControl = true);
 
 	/**
 	 * Modifies an existing control's linear data - i.e. the strengths etc of the control driving towards the target
@@ -339,12 +393,12 @@ public:
 	 */
 	UFUNCTION(BlueprintCallable, Category = PhysicsControl)
 	bool SetControlLinearData(
-		FName Name,
-		float Strength = 1.0f, 
-		float DampingRatio = 1.0f, 
-		float ExtraDamping = 0.0f, 
-		float MaxForce = 0.0f, 
-		bool  bEnableControl = true);
+		const FName Name,
+		const float Strength = 1.0f, 
+		const float DampingRatio = 1.0f, 
+		const float ExtraDamping = 0.0f, 
+		const float MaxForce = 0.0f, 
+		const bool  bEnableControl = true);
 
 	/**
 	 * Modifies an existing control's angular data - i.e. the strengths etc of the control driving towards the target
@@ -359,12 +413,12 @@ public:
 	 */
 	UFUNCTION(BlueprintCallable, Category = PhysicsControl)
 	bool SetControlAngularData(
-		FName Name,
-		float Strength = 1.0f, 
-		float DampingRatio = 1.0f, 
-		float ExtraDamping = 0.0f, 
-		float MaxTorque = 0.0f, 
-		bool  bEnableControl = true);
+		const FName Name,
+		const float Strength = 1.0f, 
+		const float DampingRatio = 1.0f, 
+		const float ExtraDamping = 0.0f, 
+		const float MaxTorque = 0.0f, 
+		const bool  bEnableControl = true);
 
 	/**
 	 * Sets the point at which controls will "push" the child object.
@@ -374,7 +428,7 @@ public:
 	 *        object is in use and is being simulated)
 	 */
 	UFUNCTION(BlueprintCallable, Category = PhysicsControl)
-	bool SetControlPoint(FName Name, const FVector Position);
+	bool SetControlPoint(const FName Name, const FVector Position);
 
 	/**
 	 * Resets the control point to the center of mass of the mesh
@@ -382,7 +436,7 @@ public:
 	 * @param Name The name of the control to modify. 
 	 */
 	UFUNCTION(BlueprintCallable, Category = PhysicsControl)
-	bool ResetControlPoint(FName Name);
+	bool ResetControlPoint(const FName Name);
 
 	/**
 	 * Modifies an existing control target - i.e. what it is driving towards, relative to the parent object
@@ -393,7 +447,7 @@ public:
 	 * @return Returns true if the control was found and modified, false if not
 	 */
 	UFUNCTION(BlueprintCallable, Category = PhysicsControl)
-	bool SetControlTarget(FName Name, FPhysicsControlTarget ControlTarget, bool bEnableControl = true);
+	bool SetControlTarget(const FName Name, const FPhysicsControlTarget ControlTarget, const bool bEnableControl = true);
 
 	/**
 	 * Modifies existing control targets - i.e. what they are driving towards, relative to the parent objects
@@ -401,14 +455,27 @@ public:
 	 * @param Names The names of the controls to modify. Note that if you have these in a FPhysicsControlNameArray
 	 *              then it can be split.
 	 * @param ControlTarget The new target for the controls
-	 * @param bEnableControl Enables the control if it is currently disabled
+	 * @param bEnableControl Enables the controls if currently disabled
 	 */
 	UFUNCTION(BlueprintCallable, Category = PhysicsControl)
-	void SetAllControlTargets(
-		const TArray<FName>&  Names, 
-		FPhysicsControlTarget ControlTarget, 
-		bool                  bEnableControl = true);
+	void SetControlTargets(
+		const TArray<FName>&        Names, 
+		const FPhysicsControlTarget ControlTarget, 
+		const bool                  bEnableControl = true);
 
+	/**
+	 * Modifies existing control targets - i.e. what they are driving towards, relative to the parent objects
+	 *
+	 * @param Set The set of controls to modify. Standard sets will include "All", "WorldSpace",
+	 *        "ParentSpace" and things like "WorldSpace-ArmLeft", depending on how controls have been created.
+	 * @param ControlTarget The new target for the controls
+	 * @param bEnableControl Enables the controls if currently disabled
+	 */
+	UFUNCTION(BlueprintCallable, Category = PhysicsControl)
+	void SetControlTargetsInSet(
+		const FName                 Set, 
+		const FPhysicsControlTarget ControlTarget, 
+		const bool                  bEnableControl = true);
 
 	/**
 	 * Modifies an existing control target - i.e. what it is driving towards, relative to the parent object
@@ -429,12 +496,12 @@ public:
 	 */
 	UFUNCTION(BlueprintCallable, Category = PhysicsControl)
 	bool SetControlTargetPositionAndOrientation(
-		FName            Name,
-		const FVector    Position,
-		const FRotator   Orientation,
-		float            VelocityDeltaTime,
-		bool             bEnableControl = true,
-		bool             bApplyControlPointToTarget = false);
+		const FName    Name,
+		const FVector  Position,
+		const FRotator Orientation,
+		const float    VelocityDeltaTime,
+		const bool     bEnableControl = true,
+		const bool     bApplyControlPointToTarget = false);
 
 	/**
 	 * Modifies an existing control target - i.e. what it is driving towards, relative to the parent object
@@ -455,11 +522,11 @@ public:
 	 */
 	UFUNCTION(BlueprintCallable, Category = PhysicsControl)
 	bool SetControlTargetPosition(
-		FName         Name,
+		const FName   Name,
 		const FVector Position, 
-		float         VelocityDeltaTime, 
-		bool          bEnableControl = true,
-		bool          bApplyControlPointToTarget = false);
+		const float   VelocityDeltaTime, 
+		const bool    bEnableControl = true,
+		const bool    bApplyControlPointToTarget = false);
 
 	/**
 	 * Modifies an existing control target - i.e. what it is driving towards, relative to the parent object
@@ -480,11 +547,11 @@ public:
 	 */
 	UFUNCTION(BlueprintCallable, Category = PhysicsControl)
 	bool SetControlTargetOrientation(
-		FName          Name,
+		const FName    Name,
 		const FRotator Orientation, 
-		float          AngularVelocityDeltaTime, 
-		bool           bEnableControl = true,
-		bool           bApplyControlPointToTarget = false);
+		const float    AngularVelocityDeltaTime, 
+		const bool     bEnableControl = true,
+		const bool     bApplyControlPointToTarget = false);
 
 	/**
 	 * Calculates and sets an existing control target. This takes the "virtual" position/orientation of the parent 
@@ -502,13 +569,13 @@ public:
 	 */
 	UFUNCTION(BlueprintCallable, Category = PhysicsControl)
 	bool SetControlTargetPoses(
-		FName          Name,
+		const FName    Name,
 		const FVector  ParentPosition, 
 		const FRotator ParentOrientation,
 		const FVector  ChildPosition, 
 		const FRotator ChildOrientation,
-		float          VelocityDeltaTime, 
-		bool           bEnableControl = true);
+		const float    VelocityDeltaTime, 
+		const bool     bEnableControl = true);
 
 
 	/**
@@ -523,9 +590,9 @@ public:
 	 */
 	UFUNCTION(BlueprintCallable, Category = PhysicsControl)
 	bool SetControlUseSkeletalAnimation(
-		FName Name,
-		bool  bUseSkeletalAnimation = true,
-		float SkeletalAnimationVelocityMultiplier = 1.0f);
+		const FName Name,
+		const bool  bUseSkeletalAnimation = true,
+		const float SkeletalAnimationVelocityMultiplier = 1.0f);
 
 	/**
 	 * Sets whether or not the controls should use skeletal animation for the targets
@@ -533,51 +600,69 @@ public:
 	 * @param Names The names of the controls to modify. Note that if you have these in a FPhysicsControlNameArray 
 	 *              then it can be split.
 	 * @param bUseSkeletalAnimation If true then the targets will be a combination of the skeletal animation (if
-	 *                              there is any) and the control target that has been set
+	 *              there is any) and the control target that has been set
 	 * @param SkeletalAnimationVelocityMultiplier If skeletal animation is being used, then this determines the amount of 
-	 *                              velocity extracted from the animation that is used as targets for the controls
+	 *              velocity extracted from the animation that is used as targets for the controls
 	 */
 	UFUNCTION(BlueprintCallable, Category = PhysicsControl)
-	void SetAllControlsUseSkeletalAnimation(
+	void SetControlsUseSkeletalAnimation(
 		const TArray<FName>& Names,
-		bool                 bUseSkeletalAnimation = true,
-		float                SkeletalAnimationVelocityMultiplier = 1.0f);
+		const bool           bUseSkeletalAnimation = true,
+		const float          SkeletalAnimationVelocityMultiplier = 1.0f);
+
+	/**
+	 * Sets whether or not the controls should use skeletal animation for the targets
+	 *
+	 * @param Set The set of controls to modify. Standard sets will include "All", "WorldSpace",
+	 *        "ParentSpace" and things like "WorldSpace-ArmLeft", depending on how controls have been created.
+	 * @param bUseSkeletalAnimation If true then the targets will be a combination of the skeletal animation (if
+	 *        there is any) and the control target that has been set
+	 * @param SkeletalAnimationVelocityMultiplier If skeletal animation is being used, then this determines the amount of
+	 *        velocity extracted from the animation that is used as targets for the controls
+	 */
+	UFUNCTION(BlueprintCallable, Category = PhysicsControl)
+	void SetControlsInSetUseSkeletalAnimation(
+		const FName Set,
+		const bool  bUseSkeletalAnimation = true,
+		const float SkeletalAnimationVelocityMultiplier = 1.0f);
 
 	/**
 	 * Activates or deactivates a control
 	 *
 	 * @param Name     The name of the control to modify. 
-	 * @param bEnable  The control to enable/disable
+	 * @param bEnable  Whether to enable/disable the control
 	 * @return         Returns true if the control was found and modified, false if not
 	 */
 	UFUNCTION(BlueprintCallable, Category = PhysicsControl)
-	bool SetControlEnabled(FName Name, bool bEnable = true);
+	bool SetControlEnabled(const FName Name, const bool bEnable = true);
 
 	/**
 	 * Activates or deactivates controls
 	 *
 	 * @param Names The names of the controls to modify. Note that if you have these in a FPhysicsControlNameArray
 	 *              then it can be split.
-	 * @param bEnable  The controls to enable/disable
+	 * @param bEnable  Whether to enable/disable the controls
 	 */
 	UFUNCTION(BlueprintCallable, Category = PhysicsControl)
-	void SetAllControlsEnabled(const TArray<FName>& Names, bool bEnable = true);
+	void SetControlsEnabled(const TArray<FName>& Names, const bool bEnable = true);
 
 	/**
-	 * @param Name The name of the control to modify. 
-	 * @param bAutoDisable If set then the control will automatically deactivate after each tick.
-	 * @return     Returns true if the control was found and modified, false if not
+	 * Activates or deactivates controls
+	 *
+	 * @param Set The set of controls to modify. Standard sets will include "All", "WorldSpace",
+	 *        "ParentSpace" and things like "WorldSpace-ArmLeft", depending on how controls have been created.
+	 * @param bEnable  Whether to enable/disable the controls
 	 */
 	UFUNCTION(BlueprintCallable, Category = PhysicsControl)
-	bool SetControlAutoDisable(FName Name, bool bAutoDisable);
+	void SetControlsInSetEnabled(const FName Set, const bool bEnable = true);
 
 	/**
-	 * @param Name     The name of the control to access. 
-	 * @param Control  The control that will be filled in, if found
-	 * @return         Returns true if the control was found, false if not
+	 * @param Name The name of the control to modify. 
+	 * @param bAutoDisable If set then the control will automatically deactivate after each tick.
+	 * @return     Returns true if the control was found and modified, false if not
 	 */
-	//UFUNCTION(BlueprintCallable, Category = PhysicsControl)
-	bool GetControl(FName Name, FPhysicsControl& Control) const;
+	UFUNCTION(BlueprintCallable, Category = PhysicsControl)
+	bool SetControlAutoDisable(const FName Name, const bool bAutoDisable);
 
 	/**
 	 * @param Name     The name of the control to access. 
@@ -585,7 +670,7 @@ public:
 	 * @return         Returns true if the control was found, false if not
 	 */
 	UFUNCTION(BlueprintCallable, Category = PhysicsControl)
-	bool GetControlData(FName Name, FPhysicsControlData& ControlData) const;
+	bool GetControlData(const FName Name, FPhysicsControlData& ControlData) const;
 
 	/**
 	 * @param Name     The name of the control to access. 
@@ -593,7 +678,7 @@ public:
 	 * @return         Returns true if the control was found, false if not
 	 */
 	UFUNCTION(BlueprintCallable, Category = PhysicsControl)
-	bool GetControlMultipliers(FName Name, FPhysicsControlMultipliers& ControlMultipliers) const;
+	bool GetControlMultiplier(const FName Name, FPhysicsControlMultiplier& ControlMultiplier) const;
 
 	/**
 	 * @param Name     The name of the control to access. 
@@ -601,82 +686,122 @@ public:
 	 * @return         Returns true if the control was found, false if not
 	 */
 	UFUNCTION(BlueprintCallable, Category = PhysicsControl)
-	bool GetControlTarget(FName Name, FPhysicsControlTarget& ControlTarget) const;
+	bool GetControlTarget(const FName Name, FPhysicsControlTarget& ControlTarget) const;
 
 	/**
 	 * @param Name        The name of the control to access. 
 	 * @return            Returns true if the control is marked to automatically disable after each tick
 	 */
 	UFUNCTION(BlueprintCallable, Category = PhysicsControl)
-	bool GetControlAutoDisable(FName Name) const;
+	bool GetControlAutoDisable(const FName Name) const;
 
 	/**
 	 * @param Name        The name of the control to access. 
 	 * @return            Returns true if the control is enabled
 	 */
 	UFUNCTION(BlueprintCallable, Category = PhysicsControl)
-	bool GetControlEnabled(FName Name) const;
+	bool GetControlEnabled(const FName Name) const;
 
 	/**
-	 * Makes a new body modifier for mesh components
+	 * Creates a new body modifier for mesh components
+	 * 
+	 * @param MeshComponent The Mesh Component used as a target for the modifier
+	 * @param BoneName The bone name, if a skeletal mesh is used
+	 * @param Set Which set to include the body modifier in (optional). Note that it automatically 
+	 *        gets added to the set "All"
+	 * @param MovementType Whether to enable/disable simulation on the body
+	 * @param CollisionType Collision type to set on the body
+	 * @param GravityMultiplier The amount of gravity to apply when simulating
+	 * @param bUseSkeletalAnimation Whether the kinematic target is specified in the frame of the
+	 *                              skeletal animation, rather than world space. Only relevant if the
+	 *                              body is part of a skeletal mesh.
 	 */
 	UFUNCTION(BlueprintCallable, Category = PhysicsControl)
-	FName MakeBodyModifier(
-		UMeshComponent*         MeshComponent,
-		FName                   BoneName,
-		EPhysicsMovementType    MovementType = EPhysicsMovementType::Simulated, 
-		ECollisionEnabled::Type CollisionType = ECollisionEnabled::QueryAndPhysics,
-		float                   GravityMultiplier = 1.0f,
-		bool                    bUseSkeletalAnimation = true);
+	FName CreateBodyModifier(
+		UMeshComponent*               MeshComponent,
+		const FName                   BoneName,
+		const FName                   Set,
+		const EPhysicsMovementType    MovementType = EPhysicsMovementType::Simulated,
+		const ECollisionEnabled::Type CollisionType = ECollisionEnabled::QueryAndPhysics,
+		const float                   GravityMultiplier = 1.0f,
+		const bool                    bUseSkeletalAnimation = true);
 
 	/**
-	 * Makes a new body modifier for mesh components
+	 * Creates a new body modifier for mesh components
+	 * 
+	 * @param The name of the body modifier that will be created. Creation will fail if this name is already in use.
+	 * @param MeshComponent The Mesh Component used as a target for the modifier
+	 * @param BoneName The bone name, if a skeletal mesh is used
+	 * @param Set Which set to include the body modifier in (optional). Note that it automatically
+	 *        gets added to the set "All"
+	 * @param MovementType Whether to enable/disable simulation on the body
+	 * @param CollisionType Collision type to set on the body
+	 * @param GravityMultiplier The amount of gravity to apply when simulating
+	 * @param bUseSkeletalAnimation Whether the kinematic target is specified in the frame of the
+	 *                              skeletal animation, rather than world space. Only relevant if the
+	 *                              body is part of a skeletal mesh.
 	 */
 	UFUNCTION(BlueprintCallable, Category = PhysicsControl)
-	bool MakeNamedBodyModifier(
-		FName                   Name,
-		UMeshComponent*         MeshComponent,
-		FName                   BoneName,
-		EPhysicsMovementType    MovementType = EPhysicsMovementType::Simulated,
-		ECollisionEnabled::Type CollisionType = ECollisionEnabled::QueryAndPhysics,
-		float                   GravityMultiplier = 1.0f,
-		bool                    bUseSkeletalAnimation = true);
+	bool CreateNamedBodyModifier(
+		const FName                   Name,
+		UMeshComponent*               MeshComponent,
+		const FName                   BoneName,
+		const FName                   Set,
+		const EPhysicsMovementType    MovementType = EPhysicsMovementType::Simulated,
+		const ECollisionEnabled::Type CollisionType = ECollisionEnabled::QueryAndPhysics,
+		const float                   GravityMultiplier = 1.0f,
+		const bool                    bUseSkeletalAnimation = true);
 
 	/**
-	 * Makes new body modifiers for skeletal mesh components
+	 * Creates new body modifiers for skeletal mesh components
+	 * 
+	 * @param SkeletalMeshComponent The skeletal mesh which will have body modifiers
+	 * @param BoneName The bone name below which modifiers should be created
+	 * @param bIncludeSelf Whether or not to include BoneName when creating modifiers
+	 * @param Set Which set to include the body modifier in (optional). Note that it automatically
+	 *        gets added to the set "All"
+	 * @param MovementType Whether to enable/disable simulation on the body
+	 * @param CollisionType Collision type to set on the body
+	 * @param GravityMultiplier The amount of gravity to apply when simulating
+	 * @param bUseSkeletalAnimation Whether the kinematic target is specified in the frame of the
+	 *                              skeletal animation, rather than world space. Only relevant if the
+	 *                              body is part of a skeletal mesh.
 	 */
 	UFUNCTION(BlueprintCallable, Category = PhysicsControl)
-	TArray<FName> MakeBodyModifiersFromSkeletalMeshBelow(
-		USkeletalMeshComponent* SkeletalMeshComponent,
-		FName                   BoneName,
-		bool                    bIncludeSelf,
-		EPhysicsMovementType    MovementType = EPhysicsMovementType::Simulated,
-		ECollisionEnabled::Type CollisionType = ECollisionEnabled::QueryAndPhysics,
-		float                   GravityMultiplier = 1.0f,
-		bool                    bUseSkeletalAnimation = true);
+	TArray<FName> CreateBodyModifiersFromSkeletalMeshBelow(
+		USkeletalMeshComponent*       SkeletalMeshComponent,
+		const FName                   BoneName,
+		const bool                    bIncludeSelf,
+		const FName                   Set,
+		const EPhysicsMovementType    MovementType = EPhysicsMovementType::Simulated,
+		const ECollisionEnabled::Type CollisionType = ECollisionEnabled::QueryAndPhysics,
+		const float                   GravityMultiplier = 1.0f,
+		const bool                    bUseSkeletalAnimation = true);
 
 
 	/**
-	 * Makes a collection of controls controlling a skeletal mesh, grouped together in limbs
+	 * Creates a collection of controls controlling a skeletal mesh, grouped together in limbs
 	 *
 	 * @param AllControls A single container for all the controls that have been created
 	 * @param LimbBones A map relating the limbs and the bones that they contain. Typically create this 
 	 *                  using GetLimbBonesFromSkeletalMesh
-	 * @param ControlType What type of control to create. This determines what the parent will be for each control
-	 * @param ControlData   Describes the initial strength etc of the new control
-	 * @param ControlSettings General settings for the control
-	 * @param bEnabled      If true then the control will be enabled immediately. If false you will need to call
-	 *                      SetControlEnabled(true) in order to enable it.
-	 * @return A map containing the controls for each limb
+	 * @param MovementType Whether to enable/disable simulation on the body
+	 * @param CollisionType Collision type to set on the body
+	 * @param GravityMultiplier The amount of gravity to apply when simulating
+	 * @param bUseSkeletalAnimation Whether the kinematic target is specified in the frame of the
+	 *                              skeletal animation, rather than world space. Only relevant if the
+	 *                              body is part of a skeletal mesh.
+	 * 
+	 * @return A map containing the modifiers for each limb
 	 */
 	UFUNCTION(BlueprintCallable, Category = PhysicsControl)
-	TMap<FName, FPhysicsControlNameArray> MakeBodyModifiersFromLimbBones(
-		FPhysicsControlNameArray&                    AllBodyModifiers,
+	TMap<FName, FPhysicsControlNames> CreateBodyModifiersFromLimbBones(
+		FPhysicsControlNames&                        AllBodyModifiers,
 		const TMap<FName, FPhysicsControlLimbBones>& LimbBones,
-		EPhysicsMovementType                         MovementType = EPhysicsMovementType::Simulated,
-		ECollisionEnabled::Type                      CollisionType = ECollisionEnabled::QueryAndPhysics,
-		float                                        GravityMultiplier = 1.0f,
-		bool                                         bUseSkeletalAnimation = true);
+		const EPhysicsMovementType                   MovementType = EPhysicsMovementType::Simulated,
+		const ECollisionEnabled::Type                CollisionType = ECollisionEnabled::QueryAndPhysics,
+		const float                                  GravityMultiplier = 1.0f,
+		const bool                                   bUseSkeletalAnimation = true);
 
 	/**
 	 * Destroys a BodyModifier
@@ -685,7 +810,7 @@ public:
 	 * @return            Returns true if the body modifier was found and destroyed, false if not
 	 */
 	UFUNCTION(BlueprintCallable, Category = PhysicsControl)
-	bool DestroyBodyModifier(FName Name);
+	bool DestroyBodyModifier(const FName Name);
 
 	/**
 	 * Destroys BodyModifiers
@@ -694,7 +819,16 @@ public:
 	 *              then it can be split.
 	 */
 	UFUNCTION(BlueprintCallable, Category = PhysicsControl)
-	void DestroyAllBodyModifiers(const TArray<FName>& Names);
+	void DestroyBodyModifiers(const TArray<FName>& Names);
+
+	/**
+	 * Destroys BodyModifiers
+	 *
+	 * @param Set The set of body modifiers to destroy. Standard sets will include "All" and things like
+	 *        "ArmLeft", depending on how body modifiers have been created.
+	 */
+	UFUNCTION(BlueprintCallable, Category = PhysicsControl)
+	void DestroyBodyModifiersInSet(const FName Set);
 
 	/**
 	 * Sets the kinematic target transform for a body modifier. 
@@ -709,10 +843,10 @@ public:
 	 */
 	UFUNCTION(BlueprintCallable, Category = PhysicsControl)
 	bool SetBodyModifierKinematicTarget(
-		FName    Name,
-		FVector  KinematicTargetPosition, 
-		FRotator KinematicTargetOrienation,
-		bool     bMakeKinematic);
+		const FName    Name,
+		const FVector  KinematicTargetPosition, 
+		const FRotator KinematicTargetOrienation,
+		const bool     bMakeKinematic);
 
 	/**
 	 * Sets the movement type for a body modifier
@@ -723,8 +857,8 @@ public:
 	 */
 	UFUNCTION(BlueprintCallable, Category = PhysicsControl)
 	bool SetBodyModifierMovementType(
-		FName                Name,
-		EPhysicsMovementType MovementType = EPhysicsMovementType::Simulated);
+		const FName                Name,
+		const EPhysicsMovementType MovementType = EPhysicsMovementType::Simulated);
 
 	/**
 	 * Sets the movement type for body modifiers
@@ -734,9 +868,21 @@ public:
 	 * @param MovementType Whether to enable/disable simulation on the bodies
 	 */
 	UFUNCTION(BlueprintCallable, Category = PhysicsControl)
-	void SetAllBodyModifierMovementType(
-		const TArray<FName>& Names,
-		EPhysicsMovementType MovementType = EPhysicsMovementType::Simulated);
+	void SetBodyModifiersMovementType(
+		const TArray<FName>&       Names,
+		const EPhysicsMovementType MovementType = EPhysicsMovementType::Simulated);
+
+	/**
+	 * Sets the movement type for body modifiers
+	 *
+	 * @param Set The set of body modifiers to modify. Standard sets will include "All" and things like
+	 *        "ArmLeft", depending on how body modifiers have been created.
+	 * @param MovementType Whether to enable/disable simulation on the bodies
+	 */
+	UFUNCTION(BlueprintCallable, Category = PhysicsControl)
+	void SetBodyModifiersInSetMovementType(
+		const FName                Set,
+		const EPhysicsMovementType MovementType = EPhysicsMovementType::Simulated);
 
 	/**
 	 * Sets the collision type for a body modifier
@@ -747,8 +893,8 @@ public:
 	 */
 	UFUNCTION(BlueprintCallable, Category = PhysicsControl)
 	bool SetBodyModifierCollisionType(
-		FName                   Name,
-		ECollisionEnabled::Type CollisionType = ECollisionEnabled::QueryAndPhysics);
+		const FName                   Name,
+		const ECollisionEnabled::Type CollisionType = ECollisionEnabled::QueryAndPhysics);
 
 	/**
 	 * Sets the collision type for body modifiers
@@ -758,9 +904,21 @@ public:
 	 * @param CollisionType Collision type to set on the bodies
 	 */
 	UFUNCTION(BlueprintCallable, Category = PhysicsControl)
-	void SetAllBodyModifierCollisionType(
-		const TArray<FName>&    Names,
-		ECollisionEnabled::Type CollisionType = ECollisionEnabled::QueryAndPhysics);
+	void SetBodyModifiersCollisionType(
+		const TArray<FName>&          Names,
+		const ECollisionEnabled::Type CollisionType = ECollisionEnabled::QueryAndPhysics);
+
+	/**
+	 * Sets the collision type for body modifiers
+	 *
+	 * @param Set The set of body modifiers to modify. Standard sets will include "All" and things like
+	 *        "ArmLeft", depending on how body modifiers have been created.
+	 * @param CollisionType Collision type to set on the bodies
+	 */
+	UFUNCTION(BlueprintCallable, Category = PhysicsControl)
+	void SetBodyModifiersInSetCollisionType(
+		const FName                   Set,
+		const ECollisionEnabled::Type CollisionType = ECollisionEnabled::QueryAndPhysics);
 
 	/**
 	 * Sets the gravity multiplier for a body modifier
@@ -771,8 +929,8 @@ public:
 	 */
 	UFUNCTION(BlueprintCallable, Category = PhysicsControl)
 	bool SetBodyModifierGravityMultiplier(
-		FName Name,
-		float GravityMultiplier = 1.0f);
+		const FName Name,
+		const float GravityMultiplier = 1.0f);
 
 	/**
 	 * Sets the gravity multiplier for body modifiers
@@ -782,9 +940,21 @@ public:
 	 * @param GravityMultiplier The amount of gravity to apply when simulating
 	 */
 	UFUNCTION(BlueprintCallable, Category = PhysicsControl)
-	void SetAllBodyModifierGravityMultipliers(
+	void SetBodyModifiersGravityMultiplier(
 		const TArray<FName>& Names,
-		float                GravityMultiplier = 1.0f);
+		const float          GravityMultiplier = 1.0f);
+
+	/**
+	 * Sets the gravity multiplier for body modifiers
+	 *
+	 * @param Set The set of body modifiers to modify. Standard sets will include "All" and things like
+	 *        "ArmLeft", depending on how body modifiers have been created.
+	 * @param GravityMultiplier The amount of gravity to apply when simulating
+	 */
+	UFUNCTION(BlueprintCallable, Category = PhysicsControl)
+	void SetBodyModifiersInSetGravityMultiplier(
+		const FName Set,
+		const float GravityMultiplier = 1.0f);
 
 	/**
 	 * Sets whether a body modifier should use skeletal animation for its kinematic targets
@@ -797,22 +967,145 @@ public:
 	 */
 	UFUNCTION(BlueprintCallable, Category = PhysicsControl)
 	bool SetBodyModifierUseSkeletalAnimation(
-		FName Name,
-		bool  bUseSkeletalAnimation);
+		const FName Name,
+		const bool  bUseSkeletalAnimation);
 
 	/**
 	 * Sets whether body modifiers should use skeletal animation for their kinematic targets
 	 *
 	 * @param Names The names of the body modifiers to access. Note that if you have these in a FPhysicsControlNameArray
-	 *              then it can be split.
+	 *        then it can be split.
 	 * @param bUseSkeletalAnimation Whether the kinematic target is specified in the frame of the 
-	 *                              skeletal animation, rather than world space. Only relevant if the
-	 *                              body is part of a skeletal mesh.
+	 *        skeletal animation, rather than world space. Only relevant if the
+	 *        body is part of a skeletal mesh.
 	 */
 	UFUNCTION(BlueprintCallable, Category = PhysicsControl)
-	void SetAllBodyModifiersUseSkeletalAnimation(
+	void SetBodyModifiersUseSkeletalAnimation(
 		const TArray<FName>& Names,
-		bool                 bUseSkeletalAnimation);
+		const bool           bUseSkeletalAnimation);
+
+	/**
+	 * Sets whether body modifiers should use skeletal animation for their kinematic targets
+	 *
+	 * @param Set The set of body modifiers to modify. Standard sets will include "All" and things like
+	 *        "ArmLeft", depending on how body modifiers have been created.
+	 * @param bUseSkeletalAnimation Whether the kinematic target is specified in the frame of the
+	 *        skeletal animation, rather than world space. Only relevant if the
+	 *        body is part of a skeletal mesh.
+	 */
+	UFUNCTION(BlueprintCallable, Category = PhysicsControl)
+	void SetBodyModifiersInSetUseSkeletalAnimation(
+		const FName Set,
+		const bool  bUseSkeletalAnimation);
+
+	/**
+	 * Creates a collections of controls and body modifiers for a character, based on the description passed in. 
+	 * This makes:
+	 * - World-space controls
+	 * - Parent-space controls
+	 * - Body modifiers
+	 * for all the body parts. In addition, they get added to sets, so they can be referenced later. Each control 
+	 * is added to three sets:
+	 * - "All"
+	 * - "ControlType - i.e. "WorldSpace" or "ParentSpace", each of which will end up containing all controls of that type
+	 * - "ControlType_LimbName" - e.g. "WorldSpace_ArmLeft" or "ParentSpace_Head"
+	 * Each body modifier is added to "All" and a set named after the limb - e.g. "Spine" or "LegRight".
+	 */
+	UFUNCTION(BlueprintCallable, Category = PhysicsControl)
+	void CreateControlsAndBodyModifiersFromLimbBones(
+		FPhysicsControlNames&                       AllWorldSpaceControls,
+		TMap<FName, FPhysicsControlNames>&          LimbWorldSpaceControls,
+		FPhysicsControlNames&                       AllParentSpaceControls,
+		TMap<FName, FPhysicsControlNames>&          LimbParentSpaceControls,
+		FPhysicsControlNames&                       AllBodyModifiers,
+		TMap<FName, FPhysicsControlNames>&          LimbBodyModifiers,
+		USkeletalMeshComponent*                     SkeletalMeshComponent,
+		const TArray<FPhysicsControlLimbSetupData>& LimbSetupData,
+		const FPhysicsControlData                   WorldSpaceControlData,
+		const FPhysicsControlSettings               WorldSpaceControlSettings,
+		const bool                                  bEnableWorldSpaceControls,
+		const FPhysicsControlData                   ParentSpaceControlData,
+		const FPhysicsControlSettings               ParentSpaceControlSettings,
+		const bool                                  bEnableParentSpaceControls,
+		const EPhysicsMovementType                  PhysicsMovementType = EPhysicsMovementType::Static,
+		const float                                 GravityMultiplier = 1.0f
+	);
+
+	/**
+	 * Adds a Control to a Set. This will add a new set if necessary. For example, you might
+	 * make a set of Controls called "ParentSpace_Feet" by calling this twice, passing in the left and right
+	 * foot ParentSpace controls.
+	 * 
+	 * @return The new/updated set of controls, in case you want to store it
+	 */
+	UFUNCTION(BlueprintCallable, Category = PhysicsControl)
+	void AddControlToSet(FPhysicsControlNames& NewSet, const FName Control, const FName Set);
+
+	/**
+	 * Adds Controls to a Set. This will add a new set if necessary. For example, you might
+	 * make a set of ParentSpace Arm controls by calling this twice, passing in the left and right
+	 * arm ParentSpace controls.
+	 * 
+	 * @return The new/updated set of controls, in case you want to store it
+	 */
+	UFUNCTION(BlueprintCallable, Category = PhysicsControl)
+	void AddControlsToSet(FPhysicsControlNames& NewSet, const TArray<FName>& Controls, const FName Set);
+
+	/**
+	 * Returns a reference to all the control names that have been created.
+	 */
+	UFUNCTION(BlueprintCallable, Category = PhysicsControl)
+	const TArray<FName>& GetAllControlNames() const;
+
+	/**
+	 * Returns a reference to all the control names that have been created and are in the specified 
+	 * set, which could be a limb, or a subsequently created set. Standard sets will include "All", "WorldSpace",
+	 * "ParentSpace" and things like "WorldSpace-ArmLeft", depending on how controls have been created.
+	 */
+	UFUNCTION(BlueprintCallable, Category = PhysicsControl)
+	const TArray<FName>& GetControlNamesInSet(const FName Set) const;
+
+	/**
+	 * Returns a reference to all the body modifier names that have been created.
+	 */
+	UFUNCTION(BlueprintCallable, Category = PhysicsControl)
+	const TArray<FName>& GetAllBodyModifierNames() const;
+
+	/**
+	 * Returns a reference to all the body modifier names that have been created and are in the specified
+	 * set, which could be a limb, or a subsequently created set. Standard sets will include "All" and things like
+	 * "ArmLeft", depending on how body modifiers have been created.
+	 */
+	UFUNCTION(BlueprintCallable, Category = PhysicsControl)
+	const TArray<FName>& GetBodyModifierNamesInSet(const FName Set) const;
+
+	/**
+	 * Adds a BodyModifier to a Set. This will add a new set if necessary. For example, you might
+	 * make a set of body modifiers called "Feet" by calling this twice, passing in the left and right
+	 * foot body modifiers.
+	 */
+	UFUNCTION(BlueprintCallable, Category = PhysicsControl)
+	void AddBodyModifierToSet(FPhysicsControlNames& NewSet, const FName BodyModifier, const FName Set);
+
+	/**
+	 * Adds BodyModifiers to a Set. This will add a new set if necessary. For example, you might
+	 * make a set of Arm body modifiers by calling this twice, passing in the left and right
+	 * arm body modifiers.
+	 */
+	UFUNCTION(BlueprintCallable, Category = PhysicsControl)
+	void AddBodyModifiersToSet(FPhysicsControlNames& NewSet, const TArray<FName>& BodyModifiers, const FName Set);
+
+	/**
+	 * Returns the names of all sets containing the control (may be empty - e.g. if it doesn't exist)
+	 */
+	UFUNCTION(BlueprintCallable, Category = PhysicsControl)
+	TArray<FName> GetSetsContainingControl(const FName Control) const;
+
+	/**
+	 * Returns the names of all sets containing the body modifier (may be empty - e.g. if it doesn't exist)
+	 */
+	UFUNCTION(BlueprintCallable, Category = PhysicsControl)
+	TArray<FName> GetSetsContainingBodyModifier(const FName Control) const;
 
 public:
 
@@ -822,7 +1115,7 @@ public:
 	 * teleporting when moving kinematic objects. Zero or negative disables.
 	 */
 	UPROPERTY(EditAnywhere, BlueprintReadWrite, Category = Teleport)
-	float TeleportDistanceThreshold;
+	float TeleportDistanceThreshold = 300.0f;
 
 	/**
 	 * If the component rotates by more than this angle (in degrees) then it is treated as a teleport,
@@ -830,22 +1123,38 @@ public:
 	 * teleporting when moving kinematic objects. Zero or negative disables.
 	 */
 	UPROPERTY(EditAnywhere, BlueprintReadWrite, Category = Teleport)
-	float TeleportRotationThreshold;
+	float TeleportRotationThreshold = 0.0f;
 
-	/** Visualize the controls when this actor/component is selected */
+	/** Visualize the controls when this component is selected */
 	UPROPERTY(EditAnywhere, BlueprintReadWrite, Category = Debug)
-	bool bShowDebugVisualization;
+	bool bShowDebugVisualization = true;
 
 	/** Size of the gizmos etc used during visualization */
 	UPROPERTY(EditAnywhere, BlueprintReadWrite, Category = Debug)
-	float VisualizationSizeScale;
+	float VisualizationSizeScale = 5.0f;
+
+	/** Display all the controls and their basic properties when this component is selected*/
+	UPROPERTY(EditAnywhere, BlueprintReadWrite, Category = Debug)
+	bool bShowDebugControlList = false;
+
+	/** Display detailed info for controls containing this string (if non-empty) when this component is selected*/
+	UPROPERTY(EditAnywhere, BlueprintReadWrite, Category = Debug)
+	FString DebugControlDetailFilter;
+
+	/** Display all the body modifiers and their basic properties when this component is selected*/
+	UPROPERTY(EditAnywhere, BlueprintReadWrite, Category = Debug)
+	bool bShowDebugBodyModifierList = false;
+
+	/** Display detailed info for body modifiers containing this string (if non-empty) when this component is selected*/
+	UPROPERTY(EditAnywhere, BlueprintReadWrite, Category = Debug)
+	FString DebugBodyModifierDetailFilter;
 
 	/**
 	 * The time used when "predicting" the target position/orientation. Zero will disable the visualization
 	 * of this.
 	 */
 	UPROPERTY(EditAnywhere, BlueprintReadWrite, Category = Debug)
-	float VelocityPredictionTime;
+	float VelocityPredictionTime = 0.2f;
 
 	/**
 	 * Upper limit on the number of controls or modifiers that will be created using the same name (which
@@ -855,7 +1164,7 @@ public:
 	 * modifiers you expect to create.
 	 */
 	UPROPERTY(EditAnywhere, BlueprintReadWrite, Category = Settings)
-	int32 MaxNumControlsOrModifiersPerName;
+	int32 MaxNumControlsOrModifiersPerName = 256;
 
 protected:
 
@@ -867,9 +1176,14 @@ protected:
 
 public:
 #if WITH_EDITOR
+	//Begin ActorComponent interface
+	virtual void OnRegister() override;
+	//End ActorComponent interface
+
+	void UpdateSpriteTexture();
 	// Used by the component visualizer
 	void DebugDraw(FPrimitiveDrawInterface* PDI) const;
-	void DebugDrawControl(FPrimitiveDrawInterface* PDI, const FPhysicsControlRecord& Record, FName ControlName) const;
+	void DebugDrawControl(FPrimitiveDrawInterface* PDI, const FPhysicsControlRecord& Record, const FName ControlName) const;
 #endif
 
 protected:
diff --git a/Engine/Plugins/Experimental/PhysicsControl/Source/PhysicsControl/Public/PhysicsControlData.h b/Engine/Plugins/Experimental/PhysicsControl/Source/PhysicsControl/Public/PhysicsControlData.h
--- a/Engine/Plugins/Experimental/PhysicsControl/Source/PhysicsControl/Public/PhysicsControlData.h
+++ b/Engine/Plugins/Experimental/PhysicsControl/Source/PhysicsControl/Public/PhysicsControlData.h
@@ -102,11 +102,11 @@ struct PHYSICSCONTROL_API FPhysicsControlData
  *    tick (typically between 0 and 1, though that is up to the user).
  */
 USTRUCT(BlueprintType)
-struct PHYSICSCONTROL_API FPhysicsControlMultipliers
+struct PHYSICSCONTROL_API FPhysicsControlMultiplier
 {
 	GENERATED_BODY()
 
-	FPhysicsControlMultipliers()
+	FPhysicsControlMultiplier()
 		: LinearStrengthMultiplier(1.0)
 		, LinearExtraDampingMultiplier(1.0)
 		, MaxForceMultiplier(1.0)
@@ -293,11 +293,11 @@ struct PHYSICSCONTROL_API FPhysicsControl
 	FPhysicsControlData ControlData;
 
 	/**
-	 * Multipliers for the ControlData. These will typically be modified dynamically, and also expose the ability
+	 * Multiplier for the ControlData. This will typically be modified dynamically, and also expose the ability
 	 * to set directional strengths
 	 */
 	UPROPERTY(EditAnywhere, BlueprintReadWrite, Category = PhysicsControl)
-	FPhysicsControlMultipliers ControlMultipliers;
+	FPhysicsControlMultiplier ControlMultiplier;
 
 	/**
 	 * The position/orientation etc targets for the controls. These are procedural/explicit control targets -
diff --git a/Engine/Plugins/Experimental/PhysicsControl/Source/PhysicsControl/Public/PhysicsControlLimbData.h b/Engine/Plugins/Experimental/PhysicsControl/Source/PhysicsControl/Public/PhysicsControlLimbData.h
--- a/Engine/Plugins/Experimental/PhysicsControl/Source/PhysicsControl/Public/PhysicsControlLimbData.h
+++ b/Engine/Plugins/Experimental/PhysicsControl/Source/PhysicsControl/Public/PhysicsControlLimbData.h
@@ -60,11 +60,11 @@ struct PHYSICSCONTROL_API FPhysicsControlLimbBones
 };
 
 /**
- * Wrapper for the array of names, so that the limbs can we can work with a map of arrays of names. 
- * Note that in practice, the names will be of either controls or body modifiers.
+ * Wrapper for the array of names, so that the control/modifier sets can we work with a map 
+ * of arrays of names. Note that in practice, the names will be of either controls or body modifiers.
  */
 USTRUCT(BlueprintType)
-struct PHYSICSCONTROL_API FPhysicsControlNameArray
+struct PHYSICSCONTROL_API FPhysicsControlNames
 {
 	GENERATED_BODY()
 
