Add some missing functions that should be in Physics Control Component

diff --git a/Engine/Plugins/Experimental/PhysicsControl/Source/PhysicsControl/Private/PhysicsControlComponent.cpp b/Engine/Plugins/Experimental/PhysicsControl/Source/PhysicsControl/Private/PhysicsControlComponent.cpp
--- a/Engine/Plugins/Experimental/PhysicsControl/Source/PhysicsControl/Private/PhysicsControlComponent.cpp
+++ b/Engine/Plugins/Experimental/PhysicsControl/Source/PhysicsControl/Private/PhysicsControlComponent.cpp
@@ -121,12 +121,15 @@ void UPhysicsControlComponent::TickComponent(
 
 		if (ConstraintInstance)
 		{
+			// Always control collision, because otherwise maintaining it is very difficult, since
+			// constraint-controlled collision doesn't interact nicely when there are multiple constraints.
+			ConstraintInstance->SetDisableCollision(Record.PhysicsControl.ControlSettings.bDisableCollision);
+
 			// Constraint is not enabled
 			if (!Record.PhysicsControlState.bEnabled)
 			{
 				ConstraintInstance->SetAngularDriveParams(0.0f, 0.0f, 0.0f);
 				ConstraintInstance->SetLinearDriveParams(0.0f, 0.0f, 0.0f);
-				ConstraintInstance->SetDisableCollision(false);
 			}
 			else
 			{
@@ -1174,6 +1177,48 @@ bool UPhysicsControlComponent::SetControlAutoDisable(const FName Name, const boo
 	return false;
 }
 
+//======================================================================================================================
+void UPhysicsControlComponent::SetControlsAutoDisable(const TArray<FName>& Names, const bool bAutoDisable)
+{
+	for (FName Name : Names)
+	{
+		SetControlAutoDisable(Name, bAutoDisable);
+	}
+}
+
+//======================================================================================================================
+void UPhysicsControlComponent::SetControlsInSetAutoDisable(const FName SetName, const bool bAutoDisable)
+{
+	SetControlsAutoDisable(GetControlNamesInSet(SetName), bAutoDisable);
+}
+
+//======================================================================================================================
+bool UPhysicsControlComponent::SetControlDisableCollision(const FName Name, const bool bDisableCollision)
+{
+	FPhysicsControl* PhysicsControl = Implementation->FindControl(Name);
+	if (PhysicsControl)
+	{
+		PhysicsControl->ControlSettings.bDisableCollision = bDisableCollision;
+		return true;
+	}
+	return false;
+}
+
+//======================================================================================================================
+void UPhysicsControlComponent::SetControlsDisableCollision(const TArray<FName>& Names, const bool bDisableCollision)
+{
+	for (FName Name : Names)
+	{
+		SetControlDisableCollision(Name, bDisableCollision);
+	}
+}
+
+//======================================================================================================================
+void UPhysicsControlComponent::SetControlsInSetDisableCollision(const FName SetName, const bool bDisableCollision)
+{
+	SetControlsDisableCollision(GetControlNamesInSet(SetName), bDisableCollision);
+}
+
 //======================================================================================================================
 bool UPhysicsControlComponent::GetControlData(const FName Name, FPhysicsControlData& ControlData) const
 {
diff --git a/Engine/Plugins/Experimental/PhysicsControl/Source/PhysicsControl/Private/PhysicsControlComponentImpl.cpp b/Engine/Plugins/Experimental/PhysicsControl/Source/PhysicsControl/Private/PhysicsControlComponentImpl.cpp
--- a/Engine/Plugins/Experimental/PhysicsControl/Source/PhysicsControl/Private/PhysicsControlComponentImpl.cpp
+++ b/Engine/Plugins/Experimental/PhysicsControl/Source/PhysicsControl/Private/PhysicsControlComponentImpl.cpp
@@ -603,8 +603,6 @@ void FPhysicsControlComponentImpl::ApplyControl(FPhysicsControlRecord& Record)
 		return;
 	}
 
-	ConstraintInstance->SetDisableCollision(Record.PhysicsControl.ControlSettings.bDisableCollision);
-
 	// Set strengths etc
 	if (ApplyControlStrengths(Record, ConstraintInstance))
 	{
diff --git a/Engine/Plugins/Experimental/PhysicsControl/Source/PhysicsControl/Private/PhysicsControlRecord.cpp b/Engine/Plugins/Experimental/PhysicsControl/Source/PhysicsControl/Private/PhysicsControlRecord.cpp
--- a/Engine/Plugins/Experimental/PhysicsControl/Source/PhysicsControl/Private/PhysicsControlRecord.cpp
+++ b/Engine/Plugins/Experimental/PhysicsControl/Source/PhysicsControl/Private/PhysicsControlRecord.cpp
@@ -62,6 +62,8 @@ FConstraintInstance* FPhysicsControlRecord::CreateConstraint(UObject* Constraint
 	ConstraintInstance->SetOrientationDriveSLERP(true);
 	ConstraintInstance->SetAngularVelocityDriveSLERP(true);
 
+	ConstraintInstance->SetDisableCollision(PhysicsControl.ControlSettings.bDisableCollision);
+
 	return ConstraintInstance;
 }
 
diff --git a/Engine/Plugins/Experimental/PhysicsControl/Source/PhysicsControl/Public/PhysicsControlComponent.h b/Engine/Plugins/Experimental/PhysicsControl/Source/PhysicsControl/Public/PhysicsControlComponent.h
--- a/Engine/Plugins/Experimental/PhysicsControl/Source/PhysicsControl/Public/PhysicsControlComponent.h
+++ b/Engine/Plugins/Experimental/PhysicsControl/Source/PhysicsControl/Public/PhysicsControlComponent.h
@@ -314,7 +314,7 @@ public:
 	 * @param Name The name of the control to modify. 
 	 * @param ControlData The new control data
 	 * @param bEnableControl Enables the control if it is currently disabled
-	 * @return Returns true if the control was found and modified, false if not
+	 * @return true if the control was found and modified, false if not
 	 */
 	UFUNCTION(BlueprintCallable, Category = PhysicsControl)
 	bool SetControlData(const FName Name, const FPhysicsControlData ControlData, const bool bEnableControl = true);
@@ -347,7 +347,7 @@ public:
 	 * @param Name The name of the control to modify. 
 	 * @param ControlMultipliers The new control multipliers
 	 * @param bEnableControl Enables the control if it is currently disabled
-	 * @return Returns true if the control was found and modified, false if not
+	 * @return true if the control was found and modified, false if not
 	 */
 	UFUNCTION(BlueprintCallable, Category = PhysicsControl)
 	bool SetControlMultiplier(const FName Name, const FPhysicsControlMultiplier ControlMultiplier, const bool bEnableControl = true);
@@ -389,7 +389,7 @@ public:
 	 * @param ExtraDamping The amount of additional linear damping
 	 * @param MaxForce The maximum force used to drive the linear motion. Zero indicates no limit.
 	 * @param bEnableControl Enables the control if it is currently disabled
-	 * @return Returns true if the control was found and modified, false if not
+	 * @return true if the control was found and modified, false if not
 	 */
 	UFUNCTION(BlueprintCallable, Category = PhysicsControl)
 	bool SetControlLinearData(
@@ -409,7 +409,7 @@ public:
 	 * @param ExtraDamping The amount of additional angular damping
 	 * @param MaxTorque The maximum torque used to drive the angular motion. Zero indicates no limit.
 	 * @param bEnableControl Enables the control if it is currently disabled
-	 * @return Returns true if the control was found and modified, false if not
+	 * @return true if the control was found and modified, false if not
 	 */
 	UFUNCTION(BlueprintCallable, Category = PhysicsControl)
 	bool SetControlAngularData(
@@ -444,7 +444,7 @@ public:
 	 * @param Name The name of the control to modify. 
 	 * @param ControlTarget The new target for the control
 	 * @param bEnableControl Enables the control if it is currently disabled
-	 * @return Returns true if the control was found and modified, false if not
+	 * @return true if the control was found and modified, false if not
 	 */
 	UFUNCTION(BlueprintCallable, Category = PhysicsControl)
 	bool SetControlTarget(const FName Name, const FPhysicsControlTarget ControlTarget, const bool bEnableControl = true);
@@ -518,7 +518,7 @@ public:
 	 *        transform for the object as a whole. If false, then the target transform
 	 *        is used as is, and the system drives the control point towards this
 	 *        transform.
-	 * @return Returns true if the control was found and modified, false if not
+	 * @return true if the control was found and modified, false if not
 	 */
 	UFUNCTION(BlueprintCallable, Category = PhysicsControl)
 	bool SetControlTargetPosition(
@@ -543,7 +543,7 @@ public:
 	 *        transform for the object as a whole. If false, then the target transform
 	 *        is used as is, and the system drives the control point towards this
 	 *        transform.
-	 * @return Returns true if the control was found and modified, false if not
+	 * @return true if the control was found and modified, false if not
 	 */
 	UFUNCTION(BlueprintCallable, Category = PhysicsControl)
 	bool SetControlTargetOrientation(
@@ -586,7 +586,7 @@ public:
 	 *        there is any) and the control target that has been set
 	 * @param SkeletalAnimationVelocityMultiplier If skeletal animation is being used, then this determines 
 	 *        the amount of velocity extracted from the animation that is used as targets for the controls
-	 * @return Returns true if the control was found and modified, false if not
+	 * @return true if the control was found and modified, false if not
 	 */
 	UFUNCTION(BlueprintCallable, Category = PhysicsControl)
 	bool SetControlUseSkeletalAnimation(
@@ -659,11 +659,49 @@ public:
 	/**
 	 * @param Name The name of the control to modify. 
 	 * @param bAutoDisable If set then the control will automatically deactivate after each tick.
-	 * @return     Returns true if the control was found and modified, false if not
+	 * @return true if the control was found and modified, false if not
 	 */
 	UFUNCTION(BlueprintCallable, Category = PhysicsControl)
 	bool SetControlAutoDisable(const FName Name, const bool bAutoDisable);
 
+	/**
+	 * @param Names The names of the controls to modify. 
+	 * @param bAutoDisable If set then the control will automatically deactivate after each tick.
+	 */
+	UFUNCTION(BlueprintCallable, Category = PhysicsControl)
+	void SetControlsAutoDisable(const TArray<FName>& Names, const bool bAutoDisable);
+
+	/**
+	 * @param Set The set of controls to modify. Standard sets will include "All", "WorldSpace",
+	 *        "ParentSpace" and things like "WorldSpace-ArmLeft", depending on how controls have been created.
+	 * @param bAutoDisable If set then the control will automatically deactivate after each tick.
+	 */
+	UFUNCTION(BlueprintCallable, Category = PhysicsControl)
+	void SetControlsInSetAutoDisable(const FName Set, const bool bAutoDisable);
+
+	/**
+	 * @param Name The name of the control to modify. 
+	 * @param bDisableCollision If set then the control will disable collision between the bodies it connects.
+	 * @return true if the control was found and modified, false if not
+	 */
+	UFUNCTION(BlueprintCallable, Category = PhysicsControl)
+	bool SetControlDisableCollision(const FName Name, const bool bDisableCollision);
+
+	/**
+	 * @param Names The names of the controls to modify. 
+	 * @param bDisableCollision If set then the control will disable collision between the bodies it connects.
+	 */
+	UFUNCTION(BlueprintCallable, Category = PhysicsControl)
+	void SetControlsDisableCollision(const TArray<FName>& Names, const bool bDisableCollision);
+
+	/**
+	 * @param Set The set of controls to modify. Standard sets will include "All", "WorldSpace",
+	 *        "ParentSpace" and things like "WorldSpace-ArmLeft", depending on how controls have been created.
+	 * @param bDisableCollision If set then the control will disable collision between the bodies it connects.
+	 */
+	UFUNCTION(BlueprintCallable, Category = PhysicsControl)
+	void SetControlsInSetDisableCollision(const FName Set, const bool bDisableCollision);
+
 	/**
 	 * @param Name     The name of the control to access. 
 	 * @param Control  The control data that will be filled in if found
@@ -839,7 +877,7 @@ public:
 	 * @param KinematicTargetOrientation The orientation to use as the kinematic target of the associated body,
 	 *                                   if it is kinematic
 	 * @param bMakeKinematic If set then the body will be made kinematic. If not set, then it won't be changed.
-	 * @return Returns true if the body modifier was found, false if not
+	 * @return true if the body modifier was found, false if not
 	 */
 	UFUNCTION(BlueprintCallable, Category = PhysicsControl)
 	bool SetBodyModifierKinematicTarget(
@@ -853,7 +891,7 @@ public:
 	 *
 	 * @param Name The name of the body modifier to access. 
 	 * @param MovementType Whether to enable/disable simulation on the body
-	 * @return Returns true if the body modifier was found, false if not
+	 * @return true if the body modifier was found, false if not
 	 */
 	UFUNCTION(BlueprintCallable, Category = PhysicsControl)
 	bool SetBodyModifierMovementType(
@@ -889,7 +927,7 @@ public:
 	 *
 	 * @param Name The name of the body modifier to access.
 	 * @param CollisionType Collision type to set on the body
-	 * @return Returns true if the body modifier was found, false if not
+	 * @return true if the body modifier was found, false if not
 	 */
 	UFUNCTION(BlueprintCallable, Category = PhysicsControl)
 	bool SetBodyModifierCollisionType(
@@ -925,7 +963,7 @@ public:
 	 *
 	 * @param Name The name of the body modifier to access. 
 	 * @param GravityMultiplier The amount of gravity to apply when simulating
-	 * @return Returns true if the body modifier was found, false if not
+	 * @return true if the body modifier was found, false if not
 	 */
 	UFUNCTION(BlueprintCallable, Category = PhysicsControl)
 	bool SetBodyModifierGravityMultiplier(
@@ -960,10 +998,9 @@ public:
 	 * Sets whether a body modifier should use skeletal animation for its kinematic targets
 	 *
 	 * @param Name The name of the body modifier to access. 
-	 * @param bUseSkeletalAnimation Whether the kinematic target is specified in the frame of the 
-	 *                              skeletal animation, rather than world space. Only relevant if the
-	 *                              body is part of a skeletal mesh.
-	 * @return Returns true if the body modifier was found, false if not
+	 * @param bUseSkeletalAnimation Whether the kinematic target is specified in the frame of the skeletal
+	 *        animation, rather than world space. Only relevant if the body is part of a skeletal mesh.
+	 * @return true if the body modifier was found, false if not
 	 */
 	UFUNCTION(BlueprintCallable, Category = PhysicsControl)
 	bool SetBodyModifierUseSkeletalAnimation(
