Fix for: Packaging Content Examples fails with error: undefined symbol: UPhysicsControlComponent

diff --git a/Engine/Plugins/Experimental/PhysicsControl/PhysicsControl.uplugin b/Engine/Plugins/Experimental/PhysicsControl/PhysicsControl.uplugin
--- a/Engine/Plugins/Experimental/PhysicsControl/PhysicsControl.uplugin
+++ b/Engine/Plugins/Experimental/PhysicsControl/PhysicsControl.uplugin
@@ -3,8 +3,8 @@
 	"Version": 1,
 	"VersionName": "1.0",
 	"FriendlyName": "PhysicsControl",
-	"Description": "",
-	"Category": "Other",
+	"Description": "Additional Support for controlling static and skeletal meshes through physical controls",
+	"Category": "Physics",
 	"CreatedBy": "Epic Games, Inc.",
 	"CreatedByURL": "https://epicgames.com",
 	"DocsURL": "",
diff --git a/Engine/Plugins/Experimental/PhysicsControl/Source/PhysicsControl/Private/PhysicsControlComponent.cpp b/Engine/Plugins/Experimental/PhysicsControl/Source/PhysicsControl/Private/PhysicsControlComponent.cpp
--- a/Engine/Plugins/Experimental/PhysicsControl/Source/PhysicsControl/Private/PhysicsControlComponent.cpp
+++ b/Engine/Plugins/Experimental/PhysicsControl/Source/PhysicsControl/Private/PhysicsControlComponent.cpp
@@ -24,6 +24,8 @@
 
 #include "SceneManagement.h"
 
+#include UE_INLINE_GENERATED_CPP_BY_NAME(PhysicsControlComponent)
+
 //======================================================================================================================
 // This is used, rather than UEnum::GetValueAsString, so that we have more control over the string returned, which 
 // gets used as a prefix for the automatically named sets etc
@@ -1596,30 +1598,6 @@ void UPhysicsControlComponent::SetBodyModifiersInSetUseSkeletalAnimation(
 	SetBodyModifiersUseSkeletalAnimation(GetBodyModifierNamesInSet(SetName), bUseSkeletalAnimation);
 }
 
-#if WITH_EDITOR
-
-//======================================================================================================================
-void UPhysicsControlComponent::OnRegister()
-{
-	Super::OnRegister();
-
-	if (SpriteComponent)
-	{
-		UpdateSpriteTexture();
-		SpriteComponent->SpriteInfo.Category = TEXT("Physics");
-		SpriteComponent->SpriteInfo.DisplayName = NSLOCTEXT("SpriteCategory", "Physics", "Physics");
-	}
-}
-
-//======================================================================================================================
-void UPhysicsControlComponent::UpdateSpriteTexture()
-{
-	if (SpriteComponent)
-	{
-		SpriteComponent->SetSprite(LoadObject<UTexture2D>(NULL, TEXT("/Engine/EditorResources/S_KBSJoint.S_KBSJoint")));
-	}
-}
-
 //======================================================================================================================
 const TArray<FName>& UPhysicsControlComponent::GetAllControlNames() const
 {
@@ -1643,8 +1621,7 @@ void UPhysicsControlComponent::CreateControlsAndBodyModifiersFromLimbBones(
 	const FPhysicsControlSettings               ParentSpaceControlSettings,
 	const bool                                  bEnableParentSpaceControls,
 	const EPhysicsMovementType                  PhysicsMovementType,
-	const float                                 GravityMultiplier
-	)
+	const float                                 GravityMultiplier)
 {
 	TMap<FName, FPhysicsControlLimbBones> LimbBones = 
 		GetLimbBonesFromSkeletalMesh(SkeletalMeshComponent, LimbSetupData);
@@ -1760,6 +1737,30 @@ TArray<FName> UPhysicsControlComponent::GetSetsContainingBodyModifier(const FNam
 	return Result;
 }
 
+#if WITH_EDITOR
+
+//======================================================================================================================
+void UPhysicsControlComponent::OnRegister()
+{
+	Super::OnRegister();
+
+	if (SpriteComponent)
+	{
+		UpdateSpriteTexture();
+		SpriteComponent->SpriteInfo.Category = TEXT("Physics");
+		SpriteComponent->SpriteInfo.DisplayName = NSLOCTEXT("SpriteCategory", "Physics", "Physics");
+	}
+}
+
+//======================================================================================================================
+void UPhysicsControlComponent::UpdateSpriteTexture()
+{
+	if (SpriteComponent)
+	{
+		SpriteComponent->SetSprite(LoadObject<UTexture2D>(NULL, TEXT("/Engine/EditorResources/S_KBSJoint.S_KBSJoint")));
+	}
+}
+
 //======================================================================================================================
 void UPhysicsControlComponent::DebugDraw(FPrimitiveDrawInterface* PDI) const
 {
