Add missing guards to Physics Control Component to avoid crashing when using the skeletal mesh component

#10017

diff --git a/Engine/Plugins/Experimental/PhysicsControl/Source/PhysicsControl/Private/PhysicsControlComponentImpl.cpp b/Engine/Plugins/Experimental/PhysicsControl/Source/PhysicsControl/Private/PhysicsControlComponentImpl.cpp
--- a/Engine/Plugins/Experimental/PhysicsControl/Source/PhysicsControl/Private/PhysicsControlComponentImpl.cpp
+++ b/Engine/Plugins/Experimental/PhysicsControl/Source/PhysicsControl/Private/PhysicsControlComponentImpl.cpp
@@ -9,6 +9,45 @@
 #include "Components/SkeletalMeshComponent.h"
 #include "PhysicsEngine/ConstraintInstance.h"
 
+//======================================================================================================================
+// Helper to get a valid skeletal mesh component pointer from a record
+static USkeletalMeshComponent* GetValidSkeletalMeshComponentFromControlParent(
+	const FPhysicsControlRecord& Record)
+{
+	if (Record.PhysicsControl.ParentMeshComponent &&
+		Record.PhysicsControl.ParentMeshComponent->IsValidLowLevel())
+	{
+		return Cast<USkeletalMeshComponent>(Record.PhysicsControl.ParentMeshComponent.Get());
+	}
+	return nullptr;
+}
+
+//======================================================================================================================
+// Helper to get a valid skeletal mesh component pointer from a record
+static USkeletalMeshComponent* GetValidSkeletalMeshComponentFromControlChild(
+	const FPhysicsControlRecord& Record)
+{
+	if (Record.PhysicsControl.ChildMeshComponent &&
+		Record.PhysicsControl.ChildMeshComponent->IsValidLowLevel())
+	{
+		return Cast<USkeletalMeshComponent>(Record.PhysicsControl.ChildMeshComponent.Get());
+	}
+	return nullptr;
+}
+
+//======================================================================================================================
+// Helper to get a valid skeletal mesh component pointer from a record
+static USkeletalMeshComponent* GetValidSkeletalMeshComponentFromBodyModifier(
+	const FPhysicsBodyModifier& PhysicsBodyModifier)
+{
+	if (PhysicsBodyModifier.MeshComponent &&
+		PhysicsBodyModifier.MeshComponent->IsValidLowLevel())
+	{
+		return Cast<USkeletalMeshComponent>(PhysicsBodyModifier.MeshComponent.Get());
+	}
+	return nullptr;
+}
+
 //======================================================================================================================
 bool FPhysicsControlComponentImpl::GetBoneData(
 	FCachedSkeletalMeshData::FBoneData& OutBoneData,
@@ -55,9 +94,9 @@ FPhysicsControlRecord* FPhysicsControlComponentImpl::FindControlRecord(const FNa
 //======================================================================================================================
 FPhysicsControl* FPhysicsControlComponentImpl::FindControl(const FName Name)
 {
-	if (FPhysicsControlRecord* ControlRecord = FindControlRecord(Name))
+	if (FPhysicsControlRecord* Record = FindControlRecord(Name))
 	{
-		return &ControlRecord->PhysicsControl;
+		return &Record->PhysicsControl;
 	}
 	return nullptr;
 }
@@ -155,7 +194,7 @@ void FPhysicsControlComponentImpl::ResetControls(bool bKeepControlRecords)
 void FPhysicsControlComponentImpl::ApplyKinematicTarget(const FPhysicsBodyModifier& BodyModifier) const
 {
 	// Seems like static and skeletal meshes need to be handled differently
-	if (USkeletalMeshComponent* SkeletalMeshComponent = Cast<USkeletalMeshComponent>(BodyModifier.MeshComponent.Get()))
+	if (USkeletalMeshComponent* SkeletalMeshComponent = GetValidSkeletalMeshComponentFromBodyModifier(BodyModifier))
 	{
 		FBodyInstance* BodyInstance = UE::PhysicsControlComponent::GetBodyInstance(
 			BodyModifier.MeshComponent, BodyModifier.BoneName);
@@ -388,15 +427,13 @@ void FPhysicsControlComponentImpl::CalculateControlTargetData(
 		bool bHaveChildBoneData = false;
 		bool bHaveParentBoneData = false;
 
-		if (USkeletalMeshComponent* ChildSkeletalMeshComponent =
-			Cast<USkeletalMeshComponent>(Record.PhysicsControl.ChildMeshComponent.Get()))
+		if (USkeletalMeshComponent* ChildSkeletalMeshComponent = GetValidSkeletalMeshComponentFromControlChild(Record))
 		{
 			bHaveChildBoneData = GetBoneData(
 				ChildBoneData, ChildSkeletalMeshComponent, Record.PhysicsControl.ChildBoneName);
 		}
 
-		if (USkeletalMeshComponent* ParentSkeletalMeshComponent =
-			Cast<USkeletalMeshComponent>(Record.PhysicsControl.ParentMeshComponent.Get()))
+		if (USkeletalMeshComponent* ParentSkeletalMeshComponent = GetValidSkeletalMeshComponentFromControlParent(Record))
 		{
 			bHaveParentBoneData = GetBoneData(
 				ParentBoneData, ParentSkeletalMeshComponent, Record.PhysicsControl.ParentBoneName);
@@ -638,21 +675,19 @@ bool FPhysicsControlComponentImpl::DestroyControl(
 	const FName            Name,
 	const EDestroyBehavior DestroyBehavior)
 {
-	FPhysicsControlRecord* PhysicsControlRecord = FindControlRecord(Name);
-	if (PhysicsControlRecord)
+	FPhysicsControlRecord* Record = FindControlRecord(Name);
+	if (Record)
 	{
-		if (USkeletalMeshComponent* SkeletalMeshComponent =
-			Cast<USkeletalMeshComponent>(PhysicsControlRecord->PhysicsControl.ParentMeshComponent))
+		if (USkeletalMeshComponent* SkeletalMeshComponent = GetValidSkeletalMeshComponentFromControlParent(*Record))
 		{
 			RemoveSkeletalMeshReferenceForCaching(SkeletalMeshComponent);
 		}
-		if (USkeletalMeshComponent* SkeletalMeshComponent =
-			Cast<USkeletalMeshComponent>(PhysicsControlRecord->PhysicsControl.ChildMeshComponent))
+		if (USkeletalMeshComponent* SkeletalMeshComponent = GetValidSkeletalMeshComponentFromControlChild(*Record))
 		{
 			RemoveSkeletalMeshReferenceForCaching(SkeletalMeshComponent);
 		}
 
-		PhysicsControlRecord->PhysicsControlState.Reset();
+		Record->PhysicsControlState.Reset();
 		NameRecords.RemoveControl(Name);
 		if (DestroyBehavior == EDestroyBehavior::RemoveRecord)
 		{
@@ -668,11 +703,10 @@ bool FPhysicsControlComponentImpl::DestroyBodyModifier(
 	const FName            Name,
 	const EDestroyBehavior DestroyBehavior)
 {
-	FPhysicsBodyModifier* PhysicsBodyModifier = FindBodyModifier(Name);
-	if (PhysicsBodyModifier)
+	FPhysicsBodyModifier* BodyModifier = FindBodyModifier(Name);
+	if (BodyModifier)
 	{
-		if (USkeletalMeshComponent* SkeletalMeshComponent =
-			Cast<USkeletalMeshComponent>(PhysicsBodyModifier->MeshComponent))
+		if (USkeletalMeshComponent* SkeletalMeshComponent = GetValidSkeletalMeshComponentFromBodyModifier(*BodyModifier))
 		{
 			RemoveSkeletalMeshReferenceForCaching(SkeletalMeshComponent);
 			RemoveSkeletalMeshReferenceForModifier(SkeletalMeshComponent);
