Physics Control Component: Fix for deleting controls in a set

diff --git a/Engine/Plugins/Experimental/PhysicsControl/Source/PhysicsControl/Private/PhysicsControlComponent.cpp b/Engine/Plugins/Experimental/PhysicsControl/Source/PhysicsControl/Private/PhysicsControlComponent.cpp
--- a/Engine/Plugins/Experimental/PhysicsControl/Source/PhysicsControl/Private/PhysicsControlComponent.cpp
+++ b/Engine/Plugins/Experimental/PhysicsControl/Source/PhysicsControl/Private/PhysicsControlComponent.cpp
@@ -765,7 +765,9 @@ void UPhysicsControlComponent::DestroyControls(const TArray<FName>& Names)
 //======================================================================================================================
 void UPhysicsControlComponent::DestroyControlsInSet(const FName SetName)
 {
-	DestroyControls(GetControlNamesInSet(SetName));
+	// Make a copy as the set will be being modified during 
+	TArray<FName> Names = GetControlNamesInSet(SetName);
+	DestroyControls(Names);
 }
 
 //======================================================================================================================
@@ -1442,7 +1444,8 @@ void UPhysicsControlComponent::DestroyBodyModifiers(const TArray<FName>& Names)
 //======================================================================================================================
 void UPhysicsControlComponent::DestroyBodyModifiersInSet(const FName SetName)
 {
-	DestroyBodyModifiers(GetBodyModifierNamesInSet(SetName));
+	TArray<FName> Names = GetBodyModifierNamesInSet(SetName);
+	DestroyBodyModifiers(Names);
 }
 
 //======================================================================================================================
