Support Physics Blend Weight with Physics Control Component

diff --git a/Engine/Plugins/Experimental/PhysicsControl/Source/PhysicsControl/Private/PhysicsControlComponent.cpp b/Engine/Plugins/Experimental/PhysicsControl/Source/PhysicsControl/Private/PhysicsControlComponent.cpp
--- a/Engine/Plugins/Experimental/PhysicsControl/Source/PhysicsControl/Private/PhysicsControlComponent.cpp
+++ b/Engine/Plugins/Experimental/PhysicsControl/Source/PhysicsControl/Private/PhysicsControlComponent.cpp
@@ -157,12 +157,16 @@ void UPhysicsControlComponent::TickComponent(
 			Implementation->ApplyKinematicTarget(BodyModifier);
 			break;
 		case EPhysicsMovementType::Simulated:
-			BodyInstance->SetInstanceSimulatePhysics(true, false, true);
+			BodyInstance->SetInstanceSimulatePhysics(true, true, true);
 			break;
 		default:
 			UE_LOG(LogPhysicsControlComponent, Warning, TEXT("Invalid movement type %d"), BodyModifier.MovementType);
 			break;
 		}
+		// We always overwrite the physics blend weight, since the functions above can still modify
+		// it (even though they all use the "maintain physics blending" option), since there is an
+		// expectation that zero blend weight means to disable physics.
+		BodyInstance->PhysicsBlendWeight = BodyModifier.PhysicsBlendWeight;
 
 		UBodySetup* BodySetup = BodyInstance->GetBodySetup();
 		if (BodySetup)
@@ -1289,11 +1293,13 @@ FName UPhysicsControlComponent::CreateBodyModifier(
 	const EPhysicsMovementType    MovementType,
 	const ECollisionEnabled::Type CollisionType,
 	const float                   GravityMultiplier,
+	const float                   PhysicsBlendWeight,
 	const bool                    bUseSkeletalAnimation)
 {
 	const FName Name = Implementation->GetUniqueBodyModifierName(BoneName);
 	if (CreateNamedBodyModifier(
-		Name, MeshComponent, BoneName, Set, MovementType, CollisionType, GravityMultiplier, bUseSkeletalAnimation))
+		Name, MeshComponent, BoneName, Set, MovementType, CollisionType, 
+		GravityMultiplier, PhysicsBlendWeight, bUseSkeletalAnimation))
 	{
 		return Name;
 	}
@@ -1309,6 +1315,7 @@ bool UPhysicsControlComponent::CreateNamedBodyModifier(
 	const EPhysicsMovementType    MovementType,
 	const ECollisionEnabled::Type CollisionType,
 	const float                   GravityMultiplier,
+	const float                   PhysicsBlendWeight,
 	const bool                    bUseSkeletalAnimation)
 {
 	if (Implementation->FindBodyModifier(Name))
@@ -1325,7 +1332,8 @@ bool UPhysicsControlComponent::CreateNamedBodyModifier(
 
 	FPhysicsBodyModifier& Modifier = Implementation->PhysicsBodyModifiers.Add(
 		Name, FPhysicsBodyModifier(
-			MeshComponent, BoneName, MovementType, CollisionType, GravityMultiplier, bUseSkeletalAnimation));
+			MeshComponent, BoneName, MovementType, CollisionType, 
+			GravityMultiplier, PhysicsBlendWeight, bUseSkeletalAnimation));
 
 	USkeletalMeshComponent* SkeletalMeshComponent = Cast<USkeletalMeshComponent>(MeshComponent);
 	if (SkeletalMeshComponent)
@@ -1348,6 +1356,7 @@ TArray<FName> UPhysicsControlComponent::CreateBodyModifiersFromSkeletalMeshBelow
 	const EPhysicsMovementType    MovementType,
 	const ECollisionEnabled::Type CollisionType,
 	const float                   GravityMultiplier,
+	const float                   PhysicsBlendWeight,
 	const bool                    bUseSkeletalAnimation)
 {
 	TArray<FName> Result;
@@ -1361,7 +1370,7 @@ TArray<FName> UPhysicsControlComponent::CreateBodyModifiersFromSkeletalMeshBelow
 		BoneName, bIncludeSelf, /*bSkipCustomType=*/false,
 		[
 			this, PhysicsAsset, SkeletalMeshComponent, Set, MovementType, CollisionType,
-			GravityMultiplier, bUseSkeletalAnimation, &Result
+			GravityMultiplier, PhysicsBlendWeight, bUseSkeletalAnimation, &Result
 		](const FBodyInstance* BI)
 		{
 			if (USkeletalBodySetup* BodySetup = Cast<USkeletalBodySetup>(BI->BodySetup.Get()))
@@ -1369,7 +1378,7 @@ TArray<FName> UPhysicsControlComponent::CreateBodyModifiersFromSkeletalMeshBelow
 				const FName BoneName = PhysicsAsset->SkeletalBodySetups[BI->InstanceBodyIndex]->BoneName;
 				const FName BodyModifierName = CreateBodyModifier(
 					SkeletalMeshComponent, BoneName, Set, MovementType, 
-					CollisionType, GravityMultiplier, bUseSkeletalAnimation);
+					CollisionType, GravityMultiplier, PhysicsBlendWeight, bUseSkeletalAnimation);
 				Result.Add(BodyModifierName);
 			}
 		});
@@ -1384,6 +1393,7 @@ TMap<FName, FPhysicsControlNames> UPhysicsControlComponent::CreateBodyModifiersF
 	const EPhysicsMovementType                   MovementType,
 	const ECollisionEnabled::Type                CollisionType,
 	const float                                  GravityMultiplier,
+	const float                                  PhysicsBlendWeight,
 	const bool                                   bUseSkeletalAnimation)
 {
 	TMap<FName, FPhysicsControlNames> Result;
@@ -1410,7 +1420,7 @@ TMap<FName, FPhysicsControlNames> UPhysicsControlComponent::CreateBodyModifiersF
 		{
 			const FName BodyModifierName = CreateBodyModifier(
 				BonesInLimb.SkeletalMeshComponent, BoneName, LimbName, MovementType, 
-				CollisionType, GravityMultiplier, bUseSkeletalAnimation);
+				CollisionType, GravityMultiplier, PhysicsBlendWeight, bUseSkeletalAnimation);
 			if (!BodyModifierName.IsNone())
 			{
 				LimbResult.Names.Add(BodyModifierName);
@@ -1568,6 +1578,39 @@ void UPhysicsControlComponent::SetBodyModifiersInSetGravityMultiplier(
 	SetBodyModifiersGravityMultiplier(GetBodyModifierNamesInSet(SetName), GravityMultiplier);
 }
 
+//======================================================================================================================
+bool UPhysicsControlComponent::SetBodyModifierPhysicsBlendWeight(
+	const FName Name,
+	const float PhysicsBlendWeight)
+{
+	FPhysicsBodyModifier* PhysicsBodyModifier = Implementation->FindBodyModifier(Name);
+	if (PhysicsBodyModifier)
+	{
+		PhysicsBodyModifier->PhysicsBlendWeight = PhysicsBlendWeight;
+		return true;
+	}
+	return false;
+}
+
+//======================================================================================================================
+void UPhysicsControlComponent::SetBodyModifiersPhysicsBlendWeight(
+	const TArray<FName>& Names,
+	const float          PhysicsBlendWeight)
+{
+	for (FName Name : Names)
+	{
+		SetBodyModifierPhysicsBlendWeight(Name, PhysicsBlendWeight);
+	}
+}
+
+//======================================================================================================================
+void UPhysicsControlComponent::SetBodyModifiersInSetPhysicsBlendWeight(
+	const FName SetName,
+	const float PhysicsBlendWeight)
+{
+	SetBodyModifiersPhysicsBlendWeight(GetBodyModifierNamesInSet(SetName), PhysicsBlendWeight);
+}
+
 //======================================================================================================================
 bool UPhysicsControlComponent::SetBodyModifierUseSkeletalAnimation(
 	const FName Name,
@@ -1624,7 +1667,8 @@ void UPhysicsControlComponent::CreateControlsAndBodyModifiersFromLimbBones(
 	const FPhysicsControlSettings               ParentSpaceControlSettings,
 	const bool                                  bEnableParentSpaceControls,
 	const EPhysicsMovementType                  PhysicsMovementType,
-	const float                                 GravityMultiplier)
+	const float                                 GravityMultiplier,
+	const float                                 PhysicsBlendWeight)
 {
 	TMap<FName, FPhysicsControlLimbBones> LimbBones = 
 		GetLimbBonesFromSkeletalMesh(SkeletalMeshComponent, LimbSetupData);
@@ -1638,7 +1682,8 @@ void UPhysicsControlComponent::CreateControlsAndBodyModifiersFromLimbBones(
 		ParentSpaceControlData, ParentSpaceControlSettings, bEnableParentSpaceControls);
 
 	LimbBodyModifiers = CreateBodyModifiersFromLimbBones(
-		AllBodyModifiers, LimbBones, PhysicsMovementType, ECollisionEnabled::QueryAndPhysics, GravityMultiplier, true);
+		AllBodyModifiers, LimbBones, PhysicsMovementType, ECollisionEnabled::QueryAndPhysics, 
+		GravityMultiplier, PhysicsBlendWeight, true);
 }
 
 //======================================================================================================================
@@ -1850,12 +1895,13 @@ void UPhysicsControlComponent::DebugDraw(FPrimitiveDrawInterface* PDI) const
 				FString ComponentName = Record.MeshComponent ? Record.MeshComponent->GetName() : TEXT("None");
 
 				FString Text = FString::Printf(
-					TEXT("%s: %s: %s %s GravityMultiplier %f"),
+					TEXT("%s: %s: %s %s GravityMultiplier %f BlendWeight %f"),
 					*Name.ToString(),
 					*ComponentName,
 					*UEnum::GetValueAsString(Record.MovementType),
 					*UEnum::GetValueAsString(Record.CollisionType),
-					Record.GravityMultiplier);
+					Record.GravityMultiplier,
+					Record.PhysicsBlendWeight);
 
 				GEngine->AddOnScreenDebugMessage(-1, 0.0f, 
 					Record.MovementType == EPhysicsMovementType::Simulated ? FColor::Green : FColor::Red, Text);
diff --git a/Engine/Plugins/Experimental/PhysicsControl/Source/PhysicsControl/Private/PhysicsControlRecord.h b/Engine/Plugins/Experimental/PhysicsControl/Source/PhysicsControl/Private/PhysicsControlRecord.h
--- a/Engine/Plugins/Experimental/PhysicsControl/Source/PhysicsControl/Private/PhysicsControlRecord.h
+++ b/Engine/Plugins/Experimental/PhysicsControl/Source/PhysicsControl/Private/PhysicsControlRecord.h
@@ -62,12 +62,14 @@ struct FPhysicsBodyModifier
 		EPhysicsMovementType       InMovementType,
 		ECollisionEnabled::Type    InCollisionType,
 		float                      InGravityMultiplier,
+		float                      InPhysicsBlendWeight,
 		bool                       InUseSkeletalAnimation)
 		: MeshComponent(InMeshComponent)
 		, BoneName(InBoneName)
 		, MovementType(InMovementType)
 		, CollisionType(InCollisionType)
 		, GravityMultiplier(InGravityMultiplier)
+		, PhysicsBlendWeight(InPhysicsBlendWeight)
 		, KinematicTargetPosition(FVector::ZeroVector)
 		, KinematicTargetOrientation(FQuat::Identity)
 		, bUseSkeletalAnimation(InUseSkeletalAnimation)
@@ -91,6 +93,11 @@ struct FPhysicsBodyModifier
 	 */
 	float GravityMultiplier = 1.0f;
 
+	/**
+	 * Blend weight (between 0 and 1) that is used to set/override the one in the body instance
+	 */
+	float PhysicsBlendWeight = 1.0f;
+
 	/** 
 	 * The target position when kinematic. Note that this is applied on top of any animation 
 	 * target if bUseSkeletalAnimation is set. 
diff --git a/Engine/Plugins/Experimental/PhysicsControl/Source/PhysicsControl/Public/PhysicsControlComponent.h b/Engine/Plugins/Experimental/PhysicsControl/Source/PhysicsControl/Public/PhysicsControlComponent.h
--- a/Engine/Plugins/Experimental/PhysicsControl/Source/PhysicsControl/Public/PhysicsControlComponent.h
+++ b/Engine/Plugins/Experimental/PhysicsControl/Source/PhysicsControl/Public/PhysicsControlComponent.h
@@ -762,6 +762,7 @@ public:
 		const EPhysicsMovementType    MovementType = EPhysicsMovementType::Simulated,
 		const ECollisionEnabled::Type CollisionType = ECollisionEnabled::QueryAndPhysics,
 		const float                   GravityMultiplier = 1.0f,
+		const float                   PhysicsBlendWeight = 1.0f,
 		const bool                    bUseSkeletalAnimation = true);
 
 	/**
@@ -788,6 +789,7 @@ public:
 		const EPhysicsMovementType    MovementType = EPhysicsMovementType::Simulated,
 		const ECollisionEnabled::Type CollisionType = ECollisionEnabled::QueryAndPhysics,
 		const float                   GravityMultiplier = 1.0f,
+		const float                   PhysicsBlendWeight = 1.0f,
 		const bool                    bUseSkeletalAnimation = true);
 
 	/**
@@ -814,6 +816,7 @@ public:
 		const EPhysicsMovementType    MovementType = EPhysicsMovementType::Simulated,
 		const ECollisionEnabled::Type CollisionType = ECollisionEnabled::QueryAndPhysics,
 		const float                   GravityMultiplier = 1.0f,
+		const float                   PhysicsBlendWeight = 1.0f,
 		const bool                    bUseSkeletalAnimation = true);
 
 
@@ -839,6 +842,7 @@ public:
 		const EPhysicsMovementType                   MovementType = EPhysicsMovementType::Simulated,
 		const ECollisionEnabled::Type                CollisionType = ECollisionEnabled::QueryAndPhysics,
 		const float                                  GravityMultiplier = 1.0f,
+		const float                                  PhysicsBlendWeight = 1.0f,
 		const bool                                   bUseSkeletalAnimation = true);
 
 	/**
@@ -994,6 +998,42 @@ public:
 		const FName Set,
 		const float GravityMultiplier = 1.0f);
 
+	/**
+	 * Sets the physics blend weight for a body modifier
+	 *
+	 * @param Name The name of the body modifier to access. 
+	 * @param PhysicsBlendWeight The amount of gravity to apply when simulating
+	 * @return true if the body modifier was found, false if not
+	 */
+	UFUNCTION(BlueprintCallable, Category = PhysicsControl)
+	bool SetBodyModifierPhysicsBlendWeight(
+		const FName Name,
+		const float PhysicsBlendWeight = 1.0f);
+
+	/**
+	 * Sets the physics blend weight for body modifiers
+	 *
+	 * @param Names The names of the body modifiers to access. Note that if you have these in a FPhysicsControlNameArray
+	 *              then it can be split.
+	 * @param PhysicsBlendWeight The amount of gravity to apply when simulating
+	 */
+	UFUNCTION(BlueprintCallable, Category = PhysicsControl)
+	void SetBodyModifiersPhysicsBlendWeight(
+		const TArray<FName>& Names,
+		const float          PhysicsBlendWeight = 1.0f);
+
+	/**
+	 * Sets the physics blend weight for body modifiers
+	 *
+	 * @param Set The set of body modifiers to modify. Standard sets will include "All" and things like
+	 *        "ArmLeft", depending on how body modifiers have been created.
+	 * @param PhysicsBlendWeight The amount of gravity to apply when simulating
+	 */
+	UFUNCTION(BlueprintCallable, Category = PhysicsControl)
+	void SetBodyModifiersInSetPhysicsBlendWeight(
+		const FName Set,
+		const float PhysicsBlendWeight = 1.0f);
+
 	/**
 	 * Sets whether a body modifier should use skeletal animation for its kinematic targets
 	 *
@@ -1065,8 +1105,9 @@ public:
 		const FPhysicsControlSettings               ParentSpaceControlSettings,
 		const bool                                  bEnableParentSpaceControls,
 		const EPhysicsMovementType                  PhysicsMovementType = EPhysicsMovementType::Static,
-		const float                                 GravityMultiplier = 1.0f
-	);
+		const float                                 GravityMultiplier = 1.0f,
+		const float                                 PhysicsBlendWeight = 1.0f
+		);
 
 	/**
 	 * Adds a Control to a Set. This will add a new set if necessary. For example, you might
diff --git a/Engine/Source/Runtime/Engine/Classes/Components/SkeletalMeshComponent.h b/Engine/Source/Runtime/Engine/Classes/Components/SkeletalMeshComponent.h
--- a/Engine/Source/Runtime/Engine/Classes/Components/SkeletalMeshComponent.h
+++ b/Engine/Source/Runtime/Engine/Classes/Components/SkeletalMeshComponent.h
@@ -585,7 +585,7 @@ public:
 	uint8 bHasValidBodies:1;
 
 
-	/** Enables blending in of physics bodies whether Simulate or not*/
+	/** Forces use of the physics bodies irrespective of whether they are simulated or using the physics blend weight */
 	UPROPERTY(transient)
 	uint8 bBlendPhysics:1;
 
diff --git a/Engine/Source/Runtime/Engine/Private/PhysicsEngine/PhysAnim.cpp b/Engine/Source/Runtime/Engine/Private/PhysicsEngine/PhysAnim.cpp
--- a/Engine/Source/Runtime/Engine/Private/PhysicsEngine/PhysAnim.cpp
+++ b/Engine/Source/Runtime/Engine/Private/PhysicsEngine/PhysAnim.cpp
@@ -195,8 +195,6 @@ void USkeletalMeshComponent::PerformBlendPhysicsBones(const TArray<FBoneIndexTyp
 
 			// See if this is a physics bone..
 			int32 BodyIndex = PhysicsAsset->FindBodyIndex(GetSkeletalMeshAsset()->GetRefSkeleton().GetBoneName(BoneIndex));
-			// need to update back to physX so that physX knows where it was after blending
-			FBodyInstance* PhysicsAssetBodyInstance = nullptr;
 
 			// Gets set to true if we have a valid body and the skeletal mesh option has been set to
 			// be driven by kinematic body parts.
@@ -216,14 +214,14 @@ void USkeletalMeshComponent::PerformBlendPhysicsBones(const TArray<FBoneIndexTyp
 					continue;
 				}
 #endif
-				PhysicsAssetBodyInstance = Bodies[BodyIndex];
+				FBodyInstance* PhysicsAssetBodyInstance = Bodies[BodyIndex];
 
 				bDriveMeshWhenKinematic =
 					bUpdateMeshWhenKinematic &&
 					PhysicsAssetBodyInstance->IsValidBodyInstance();
 
 				//if simulated body copy back and blend with animation
-				if(PhysicsAssetBodyInstance->IsInstanceSimulatingPhysics() || bDriveMeshWhenKinematic)
+				if (PhysicsAssetBodyInstance->IsInstanceSimulatingPhysics() || bDriveMeshWhenKinematic)
 				{
 					FTransform PhysTM = PhysicsAssetBodyInstance->GetUnrealWorldTransform_AssumesLocked();
 
@@ -231,25 +229,28 @@ void USkeletalMeshComponent::PerformBlendPhysicsBones(const TArray<FBoneIndexTyp
 					WorldBoneTMs[BoneIndex].TM = PhysTM;
 					WorldBoneTMs[BoneIndex].bUpToDate = true;
 
-					float UsePhysWeight = (bBlendPhysics || bDriveMeshWhenKinematic)? 1.f : PhysicsAssetBodyInstance->PhysicsBlendWeight;
-
-					// if the body instance is disabled, then we want to use the animation transform and ignore the physics one
-					if (PhysicsAssetBodyInstance->IsPhysicsDisabled() && !bDriveMeshWhenKinematic)
+					if (PhysicsAssetBodyInstance->IsPhysicsDisabled())
 					{
-						UsePhysWeight = 0.0f;
+						continue;
 					}
 
-					// Find this bones parent matrix.
-					FTransform ParentWorldTM;
+					// Note that when using bDriveMeshWhenKinematic you should still set the
+					// PhysicsBlendWeight. This allows you to use intermediate values for the blend
+					// weight when kinematic.
+					float PhysicsBlendWeight = bBlendPhysics ? 1.f : PhysicsAssetBodyInstance->PhysicsBlendWeight;
 
-					// if we want 'full weight' we just find 
-					if(UsePhysWeight > 0.f)
+					// if the body instance is disabled, then we want to use the animation transform
+					// and ignore the physics one, but still allow picking up transforms from
+					// kinematic movement.
+					if (PhysicsBlendWeight > 0.f)
 					{
 						if (!(ensure(InOutBoneSpaceTransforms.Num())))
 						{
 							continue;
 						}
 
+						// Find this bones parent matrix.
+						FTransform ParentWorldTM;
 						if(BoneIndex == 0)
 						{
 							ParentWorldTM = LocalToWorldTM;
@@ -271,7 +272,7 @@ void USkeletalMeshComponent::PerformBlendPhysicsBones(const TArray<FBoneIndexTyp
 						FTransform PhysAtom = FTransform(RelRot, RelPos, InOutBoneSpaceTransforms[BoneIndex].GetScale3D());
 
 						// Now blend in this atom. See if we are forcing this bone to always be blended in
-						InOutBoneSpaceTransforms[BoneIndex].Blend( InOutBoneSpaceTransforms[BoneIndex], PhysAtom, UsePhysWeight );
+						InOutBoneSpaceTransforms[BoneIndex].Blend( InOutBoneSpaceTransforms[BoneIndex], PhysAtom, PhysicsBlendWeight );
 
 						if (!bSetParentScale)
 						{
