Work-in-progress Physics Control Profile Editor

diff --git a/Engine/Plugins/Experimental/PhysicsControl/PhysicsControl.uplugin b/Engine/Plugins/Experimental/PhysicsControl/PhysicsControl.uplugin
--- a/Engine/Plugins/Experimental/PhysicsControl/PhysicsControl.uplugin
+++ b/Engine/Plugins/Experimental/PhysicsControl/PhysicsControl.uplugin
@@ -22,7 +22,7 @@
 			"LoadingPhase": "Default"
 		},
 		{
-			"Name": "PhysicsControlVisualizer",
+			"Name": "PhysicsControlEditor",
 			"Type": "Editor",
 			"LoadingPhase": "PostEngineInit"
 		}
diff --git a/Engine/Plugins/Experimental/PhysicsControl/Source/PhysicsControl/Private/PhysicsControlProfileAsset.cpp b/Engine/Plugins/Experimental/PhysicsControl/Source/PhysicsControl/Private/PhysicsControlProfileAsset.cpp
new file mode 100644
--- /dev/null
+++ b/Engine/Plugins/Experimental/PhysicsControl/Source/PhysicsControl/Private/PhysicsControlProfileAsset.cpp
@@ -0,0 +1,36 @@
+// Copyright Epic Games, Inc. All Rights Reserved.
+
+#include "PhysicsControlProfileAsset.h"
+#include "Engine/SkeletalMesh.h"
+
+//======================================================================================================================
+UPhysicsControlProfileAsset::UPhysicsControlProfileAsset()
+	: TestValue(1.0f)
+{
+}
+
+//======================================================================================================================
+void UPhysicsControlProfileAsset::Log()
+{
+	UE_LOG(LogTemp, Log, TEXT("%f"), TestValue)
+}
+
+#if WITH_EDITOR
+//======================================================================================================================
+const FName UPhysicsControlProfileAsset::GetPreviewMeshPropertyName()
+{
+	return GET_MEMBER_NAME_STRING_CHECKED(UPhysicsControlProfileAsset, PreviewSkeletalMesh);
+};
+#endif
+
+//======================================================================================================================
+void UPhysicsControlProfileAsset::SetPreviewMesh(USkeletalMesh* PreviewMesh, bool bMarkAsDirty)
+{
+	PreviewSkeletalMesh = PreviewMesh;
+}
+
+//======================================================================================================================
+USkeletalMesh* UPhysicsControlProfileAsset::GetPreviewMesh() const
+{
+	return PreviewSkeletalMesh.LoadSynchronous();
+}
diff --git a/Engine/Plugins/Experimental/PhysicsControl/Source/PhysicsControl/Public/PhysicsControlProfileAsset.h b/Engine/Plugins/Experimental/PhysicsControl/Source/PhysicsControl/Public/PhysicsControlProfileAsset.h
new file mode 100644
--- /dev/null
+++ b/Engine/Plugins/Experimental/PhysicsControl/Source/PhysicsControl/Public/PhysicsControlProfileAsset.h
@@ -0,0 +1,53 @@
+// Copyright Epic Games, Inc. All Rights Reserved.
+
+#pragma once
+
+#include "Interfaces/Interface_PreviewMeshProvider.h"
+
+#include "PhysicsControlProfileAsset.generated.h"
+
+class USkeletalMesh;
+
+/**
+ * Asset for storing Physics Control Profiles. These will contain data that define:
+ * - Controls and body modifiers to be created on a mesh
+ * - Sets referencing those controls and body modifiers
+ * - Full profiles containing settings for all the controls/modifiers
+ * - Sparse profiles containing partial sets of settings for specific controls/modifiers
+ * 
+ * It will also be desirable to support "inheritance" - so a generic profile can be made, and then 
+ * customized for certain characters or scenarios.
+ */
+UCLASS(BlueprintType)
+class PHYSICSCONTROL_API UPhysicsControlProfileAsset : public UObject, public IInterface_PreviewMeshProvider
+{
+	GENERATED_BODY()
+public:
+	UPhysicsControlProfileAsset();
+
+	/** Placeholder for testing/developing */
+	UFUNCTION(CallInEditor)
+	void Log();
+
+	/** IInterface_PreviewMeshProvider interface */
+	virtual void SetPreviewMesh(USkeletalMesh* PreviewMesh, bool bMarkAsDirty = true) override;
+	virtual USkeletalMesh* GetPreviewMesh() const override;
+	/** END IInterface_PreviewMeshProvider interface */
+
+#if WITH_EDITOR
+/* Get name of Preview Mesh property */
+	static const FName GetPreviewMeshPropertyName();
+#endif
+
+public:
+	/** 
+	* The skeletal mesh to use for generating controls and previewing. If it turns out this 
+	* doesn't need to be stored in the asset, it will get moved out of here.
+	*/
+	UPROPERTY(AssetRegistrySearchable, EditAnywhere, Category = PreviewMesh)
+	TSoftObjectPtr<USkeletalMesh> PreviewSkeletalMesh;
+
+	/** Placeholder for testing/developing */
+	UPROPERTY(EditAnywhere, Category=PhysicsControl)
+	float TestValue;
+};
diff --git a/Engine/Plugins/Experimental/PhysicsControl/Source/PhysicsControlEditor/PhysicsControlEditor.Build.cs b/Engine/Plugins/Experimental/PhysicsControl/Source/PhysicsControlEditor/PhysicsControlEditor.Build.cs
new file mode 100644
--- /dev/null
+++ b/Engine/Plugins/Experimental/PhysicsControl/Source/PhysicsControlEditor/PhysicsControlEditor.Build.cs
@@ -0,0 +1,60 @@
+// Copyright Epic Games, Inc. All Rights Reserved.
+
+using UnrealBuildTool;
+public class PhysicsControlEditor : ModuleRules
+{
+	public PhysicsControlEditor(ReadOnlyTargetRules Target) : base(Target)
+	{
+		PCHUsage = ModuleRules.PCHUsageMode.UseExplicitOrSharedPCHs;
+		UnsafeTypeCastWarningLevel = WarningLevel.Warning;
+
+		PublicIncludePaths.AddRange(
+			new string[] {
+				// ... add public include paths required here ...
+			}
+			);
+				
+		
+		PrivateIncludePaths.AddRange(
+			new string[] 
+			{
+				// ... add other private include paths required here ...
+			}
+			);
+			
+		
+		PublicDependencyModuleNames.AddRange(
+			new string[]
+			{
+				"Core"
+				// ... add other public dependencies that you statically link with here ...
+			}
+			);
+			
+		
+		PrivateDependencyModuleNames.AddRange(
+			new string[]
+			{
+				"CoreUObject",
+				"EditorFramework",
+				"Engine",
+				"Slate",
+				"SlateCore",
+				"UnrealEd",
+				"Persona",
+				"AnimationEditMode",
+				"ComponentVisualizers",
+				"AnimGraph",
+				"PhysicsControl"
+			}
+			);
+		
+		
+		DynamicallyLoadedModuleNames.AddRange(
+			new string[]
+			{
+				// ... add any modules that your module loads dynamically here ...
+			}
+			);
+	}
+}
diff --git a/Engine/Plugins/Experimental/PhysicsControl/Source/PhysicsControlEditor/Private/PhysicsControlComponentVisualizer.cpp b/Engine/Plugins/Experimental/PhysicsControl/Source/PhysicsControlEditor/Private/PhysicsControlComponentVisualizer.cpp
new file mode 100644
--- /dev/null
+++ b/Engine/Plugins/Experimental/PhysicsControl/Source/PhysicsControlEditor/Private/PhysicsControlComponentVisualizer.cpp
@@ -0,0 +1,16 @@
+// Copyright Epic Games, Inc. All Rights Reserved.
+
+#include "PhysicsControlComponentVisualizer.h"
+#include "PhysicsControlComponent.h"
+
+//======================================================================================================================
+void FPhysicsControlComponentVisualizer::DrawVisualization(
+	const UActorComponent*   Component,
+	const FSceneView*        View,
+	FPrimitiveDrawInterface* PDI)
+{
+	if (const UPhysicsControlComponent* PCC = Cast<const UPhysicsControlComponent>(Component))
+	{
+		PCC->DebugDraw(PDI);
+	}
+}
diff --git a/Engine/Plugins/Experimental/PhysicsControl/Source/PhysicsControlEditor/Private/PhysicsControlComponentVisualizer.h b/Engine/Plugins/Experimental/PhysicsControl/Source/PhysicsControlEditor/Private/PhysicsControlComponentVisualizer.h
new file mode 100644
--- /dev/null
+++ b/Engine/Plugins/Experimental/PhysicsControl/Source/PhysicsControlEditor/Private/PhysicsControlComponentVisualizer.h
@@ -0,0 +1,13 @@
+// Copyright Epic Games, Inc. All Rights Reserved.
+
+#pragma once
+#include "ComponentVisualizer.h"
+
+class PHYSICSCONTROLEDITOR_API FPhysicsControlComponentVisualizer : public FComponentVisualizer
+{
+private:
+	virtual void DrawVisualization(
+		const UActorComponent*   Component, 
+		const FSceneView*        View, 
+		FPrimitiveDrawInterface* PDI) override;
+};
\ No newline at end of file
diff --git a/Engine/Plugins/Experimental/PhysicsControl/Source/PhysicsControlEditor/Private/PhysicsControlEditorModule.cpp b/Engine/Plugins/Experimental/PhysicsControl/Source/PhysicsControlEditor/Private/PhysicsControlEditorModule.cpp
new file mode 100644
--- /dev/null
+++ b/Engine/Plugins/Experimental/PhysicsControl/Source/PhysicsControlEditor/Private/PhysicsControlEditorModule.cpp
@@ -0,0 +1,68 @@
+// Copyright Epic Games, Inc. All Rights Reserved.
+
+#include "PhysicsControlEditorModule.h"
+#include "PhysicsControlComponent.h"
+#include "PhysicsControlComponentVisualizer.h"
+#include "PhysicsControlProfileAssetActions.h"
+#include "PhysicsControlProfileEditorMode.h"
+#include "UnrealEdGlobals.h"
+#include "Editor/UnrealEdEngine.h"
+
+#define LOCTEXT_NAMESPACE "PhysicsControlModule"
+
+//#define ENABLE_PHYSICS_CONTROL_PROFILE_EDITOR
+
+//======================================================================================================================
+void FPhysicsControlEditorModule::StartupModule()
+{
+	// Physics Control Profile editor/asset is disabled for now
+#ifdef ENABLE_PHYSICS_CONTROL_PROFILE_EDITOR
+	IAssetTools& AssetTools = FModuleManager::LoadModuleChecked<FAssetToolsModule>("AssetTools").Get();
+	PhysicsControlProfileAssetActions = MakeShared<FPhysicsControlProfileAssetActions>();
+	AssetTools.RegisterAssetTypeActions(PhysicsControlProfileAssetActions.ToSharedRef());
+
+	FEditorModeRegistry::Get().RegisterMode<FPhysicsControlProfileEditorMode>(
+		FPhysicsControlProfileEditorMode::ModeName, 
+		LOCTEXT("PhysicsControlProfileEditorMode", "PhysicsControlProfile"), 
+		FSlateIcon(), false);
+#endif
+
+	if (GUnrealEd)
+	{
+		TSharedPtr<FPhysicsControlComponentVisualizer> Visualizer = MakeShared<FPhysicsControlComponentVisualizer>();
+		GUnrealEd->RegisterComponentVisualizer(UPhysicsControlComponent::StaticClass()->GetFName(), Visualizer);
+		// This call should maybe be inside the RegisterComponentVisualizer call above, but since it's not,
+		// we'll put it here.
+		Visualizer->OnRegister();
+		VisualizersToUnregisterOnShutdown.Add(UPhysicsControlComponent::StaticClass()->GetFName());
+	}
+}
+
+//======================================================================================================================
+void FPhysicsControlEditorModule::ShutdownModule()
+{
+	// Physics Control Profile editor/asset is disabled for now
+#ifdef ENABLE_PHYSICS_CONTROL_PROFILE_EDITOR
+	FEditorModeRegistry::Get().UnregisterMode(FPhysicsControlProfileEditorMode::ModeName);
+
+	if (FModuleManager::Get().IsModuleLoaded("AssetTools"))
+	{
+		FAssetToolsModule::GetModule().Get().UnregisterAssetTypeActions(
+			PhysicsControlProfileAssetActions.ToSharedRef());
+	}
+#endif
+
+	if (GUnrealEd)
+	{
+		for (const FName& Name : VisualizersToUnregisterOnShutdown)
+		{
+			GUnrealEd->UnregisterComponentVisualizer(Name);
+		}
+	}
+}
+
+#undef LOCTEXT_NAMESPACE
+	
+IMPLEMENT_MODULE(FPhysicsControlEditorModule, PhysicsControl)
+
+
diff --git a/Engine/Plugins/Experimental/PhysicsControl/Source/PhysicsControlEditor/Private/PhysicsControlProfileApplicationMode.cpp b/Engine/Plugins/Experimental/PhysicsControl/Source/PhysicsControlEditor/Private/PhysicsControlProfileApplicationMode.cpp
new file mode 100644
--- /dev/null
+++ b/Engine/Plugins/Experimental/PhysicsControl/Source/PhysicsControlEditor/Private/PhysicsControlProfileApplicationMode.cpp
@@ -0,0 +1,89 @@
+// Copyright Epic Games, Inc. All Rights Reserved.
+
+#include "PhysicsControlProfileApplicationMode.h"
+#include "PhysicsControlProfileEditorToolkit.h"
+#include "PersonaModule.h"
+#include "Modules/ModuleManager.h"
+#include "PersonaTabs.h"
+
+#define LOCTEXT_NAMESPACE "PhysicsControlProfileApplicationMode"
+
+FName FPhysicsControlProfileApplicationMode::ModeName("PhysicsControlProfileAssetEditMode");
+
+//======================================================================================================================
+FPhysicsControlProfileApplicationMode::FPhysicsControlProfileApplicationMode(
+	TSharedRef<FWorkflowCentricApplication> InHostingApp,
+	TSharedRef<IPersonaPreviewScene>        InPreviewScene)
+	: 
+	FApplicationMode(PhysicsControlProfileEditorModes::Editor)
+{
+	EditorToolkit = StaticCastSharedRef<FPhysicsControlProfileEditorToolkit>(InHostingApp);
+	TSharedRef<FPhysicsControlProfileEditorToolkit> PhysicsControlProfileEditor = 
+		StaticCastSharedRef<FPhysicsControlProfileEditorToolkit>(InHostingApp);
+
+	FPersonaViewportArgs ViewportArgs(InPreviewScene);
+	ViewportArgs.bAlwaysShowTransformToolbar = true;
+	ViewportArgs.bShowStats = false;
+	ViewportArgs.bShowTimeline = false;
+	ViewportArgs.bShowLODMenu = true;
+	ViewportArgs.bShowPlaySpeedMenu = false;
+	ViewportArgs.bShowPhysicsMenu = false;
+	ViewportArgs.ContextName = TEXT("PhysicsControlProfileEditor.Viewport");
+	ViewportArgs.OnViewportCreated = FOnViewportCreated::CreateSP(
+		PhysicsControlProfileEditor, &FPhysicsControlProfileEditorToolkit::HandleViewportCreated);
+
+	// Register Persona tabs.
+	FPersonaModule& PersonaModule = FModuleManager::LoadModuleChecked<FPersonaModule>("Persona");
+	TabFactories.RegisterFactory(PersonaModule.CreatePersonaViewportTabFactory(InHostingApp, ViewportArgs));
+	TabFactories.RegisterFactory(PersonaModule.CreateDetailsTabFactory(InHostingApp, FOnDetailsCreated::CreateSP(&PhysicsControlProfileEditor.Get(), &FPhysicsControlProfileEditorToolkit::HandleDetailsCreated)));
+
+	// Create tab layout.
+	TabLayout = FTabManager::NewLayout("Standalone_PhysicsControlProfileEditor_Layout_v0.001")
+		->AddArea
+		(
+			FTabManager::NewPrimaryArea()
+			->SetOrientation(Orient_Vertical)
+			->Split
+			(
+				FTabManager::NewSplitter()
+				->SetSizeCoefficient(0.32f)
+				->SetOrientation(Orient_Horizontal)
+				->Split
+				(
+					FTabManager::NewSplitter()
+					->SetSizeCoefficient(1.0f)
+					->SetOrientation(Orient_Vertical)
+					->Split
+					(
+						FTabManager::NewStack()
+						->SetSizeCoefficient(0.7f)
+						->SetHideTabWell(true)
+						->AddTab(FPersonaTabs::PreviewViewportID, ETabState::OpenedTab)
+					)
+				)
+				->Split
+				(
+					FTabManager::NewStack()
+					->SetSizeCoefficient(0.3f)
+					->AddTab(FPersonaTabs::DetailsID, ETabState::OpenedTab)
+					->SetForegroundTab(FPersonaTabs::DetailsID)
+				)
+			)
+		);
+
+	PersonaModule.OnRegisterTabs().Broadcast(TabFactories, InHostingApp);
+	LayoutExtender = MakeShared<FLayoutExtender>();
+	PersonaModule.OnRegisterLayoutExtensions().Broadcast(*LayoutExtender.Get());
+	TabLayout->ProcessExtensions(*LayoutExtender.Get());
+}
+
+//======================================================================================================================
+void FPhysicsControlProfileApplicationMode::RegisterTabFactories(TSharedPtr<FTabManager> InTabManager)
+{
+	TSharedPtr<FPhysicsControlProfileEditorToolkit> Editor = EditorToolkit.Pin();
+	Editor->RegisterTabSpawners(InTabManager.ToSharedRef());
+	Editor->PushTabFactories(TabFactories);
+	FApplicationMode::RegisterTabFactories(InTabManager);
+}
+
+#undef LOCTEXT_NAMESPACE
diff --git a/Engine/Plugins/Experimental/PhysicsControl/Source/PhysicsControlEditor/Private/PhysicsControlProfileAssetActions.cpp b/Engine/Plugins/Experimental/PhysicsControl/Source/PhysicsControlEditor/Private/PhysicsControlProfileAssetActions.cpp
new file mode 100644
--- /dev/null
+++ b/Engine/Plugins/Experimental/PhysicsControl/Source/PhysicsControlEditor/Private/PhysicsControlProfileAssetActions.cpp
@@ -0,0 +1,47 @@
+// Copyright Epic Games, Inc. All Rights Reserved.
+
+#include "PhysicsControlProfileAssetActions.h"
+#include "PhysicsControlProfileAsset.h"
+#include "PhysicsControlProfileEditorToolkit.h"
+
+//======================================================================================================================
+UClass* FPhysicsControlProfileAssetActions::GetSupportedClass() const
+{
+	return UPhysicsControlProfileAsset::StaticClass();
+}
+
+//======================================================================================================================
+FText FPhysicsControlProfileAssetActions::GetName() const
+{
+	return INVTEXT("Physics Control Profile");
+}
+
+//======================================================================================================================
+FColor FPhysicsControlProfileAssetActions::GetTypeColor() const
+{
+	// Match the "standard" physics color - they tend to be variations around this value
+	return FColor(255, 192, 128);
+}
+
+//======================================================================================================================
+uint32 FPhysicsControlProfileAssetActions::GetCategories()
+{
+	return EAssetTypeCategories::Physics;
+}
+
+//======================================================================================================================
+void FPhysicsControlProfileAssetActions::OpenAssetEditor(
+	const TArray<UObject*>& InObjects, TSharedPtr<IToolkitHost> EditWithinLevelEditor)
+{
+	const EToolkitMode::Type Mode = 
+		EditWithinLevelEditor.IsValid() ? EToolkitMode::WorldCentric : EToolkitMode::Standalone;
+
+	for (UObject* Object : InObjects)
+	{
+		if (UPhysicsControlProfileAsset* Asset = Cast<UPhysicsControlProfileAsset>(Object))
+		{
+			TSharedRef<FPhysicsControlProfileEditorToolkit> NewEditor(new FPhysicsControlProfileEditorToolkit());
+			NewEditor->InitAssetEditor(Mode, EditWithinLevelEditor, Asset);
+		}
+	}
+}
diff --git a/Engine/Plugins/Experimental/PhysicsControl/Source/PhysicsControlEditor/Private/PhysicsControlProfileAssetActions.h b/Engine/Plugins/Experimental/PhysicsControl/Source/PhysicsControlEditor/Private/PhysicsControlProfileAssetActions.h
new file mode 100644
--- /dev/null
+++ b/Engine/Plugins/Experimental/PhysicsControl/Source/PhysicsControlEditor/Private/PhysicsControlProfileAssetActions.h
@@ -0,0 +1,20 @@
+// Copyright Epic Games, Inc. All Rights Reserved.
+
+#pragma once
+
+#include "CoreMinimal.h"
+#include "AssetTypeActions_Base.h"
+
+class FPhysicsControlProfileAssetActions : public FAssetTypeActions_Base
+{
+public:
+	UClass* GetSupportedClass() const override;
+	FText GetName() const override;
+	FColor GetTypeColor() const override;
+	uint32 GetCategories() override;
+
+	void OpenAssetEditor(
+		const TArray<UObject*>& InObjects, 
+		TSharedPtr<class IToolkitHost> EditWithinLevelEditor = TSharedPtr<IToolkitHost>()) override;
+
+};
diff --git a/Engine/Plugins/Experimental/PhysicsControl/Source/PhysicsControlEditor/Private/PhysicsControlProfileAssetFactory.cpp b/Engine/Plugins/Experimental/PhysicsControl/Source/PhysicsControlEditor/Private/PhysicsControlProfileAssetFactory.cpp
new file mode 100644
--- /dev/null
+++ b/Engine/Plugins/Experimental/PhysicsControl/Source/PhysicsControlEditor/Private/PhysicsControlProfileAssetFactory.cpp
@@ -0,0 +1,19 @@
+// Copyright Epic Games, Inc. All Rights Reserved.
+
+#include "PhysicsControlProfileAssetFactory.h"
+#include "PhysicsControlProfileAsset.h"
+
+//======================================================================================================================
+UPhysicsControlProfileAssetFactory::UPhysicsControlProfileAssetFactory()
+{
+	SupportedClass = UPhysicsControlProfileAsset::StaticClass();
+	bCreateNew = true;
+	bEditAfterNew = false;
+}
+
+//======================================================================================================================
+UObject* UPhysicsControlProfileAssetFactory::FactoryCreateNew(
+	UClass* Class, UObject* InParent, FName Name, EObjectFlags Flags, UObject* Context, FFeedbackContext* Warn)
+{
+	return NewObject<UPhysicsControlProfileAsset>(InParent, Class, Name, Flags, Context);
+}
diff --git a/Engine/Plugins/Experimental/PhysicsControl/Source/PhysicsControlEditor/Private/PhysicsControlProfileAssetFactory.h b/Engine/Plugins/Experimental/PhysicsControl/Source/PhysicsControlEditor/Private/PhysicsControlProfileAssetFactory.h
new file mode 100644
--- /dev/null
+++ b/Engine/Plugins/Experimental/PhysicsControl/Source/PhysicsControlEditor/Private/PhysicsControlProfileAssetFactory.h
@@ -0,0 +1,17 @@
+// Copyright Epic Games, Inc. All Rights Reserved.
+
+#pragma once
+
+#include "CoreMinimal.h"
+#include "Factories/Factory.h"
+#include "PhysicsControlProfileAssetFactory.generated.h"
+
+UCLASS()
+class UPhysicsControlProfileAssetFactory : public UFactory
+{
+	GENERATED_BODY()
+public:
+	UPhysicsControlProfileAssetFactory();
+	UObject* FactoryCreateNew(
+		UClass* Class, UObject* InParent, FName Name, EObjectFlags Flags, UObject* Context, FFeedbackContext* Warn);
+};
diff --git a/Engine/Plugins/Experimental/PhysicsControl/Source/PhysicsControlEditor/Private/PhysicsControlProfileEditorData.cpp b/Engine/Plugins/Experimental/PhysicsControl/Source/PhysicsControlEditor/Private/PhysicsControlProfileEditorData.cpp
new file mode 100644
--- /dev/null
+++ b/Engine/Plugins/Experimental/PhysicsControl/Source/PhysicsControlEditor/Private/PhysicsControlProfileEditorData.cpp
@@ -0,0 +1,60 @@
+#include "PhysicsControlProfileEditorData.h"
+#include "PhysicsControlProfileAsset.h"
+#include "Engine/SkeletalMesh.h"
+#include "Misc/MessageDialog.h"
+
+#define LOCTEXT_NAMESPACE "PhysicsControlProfileEditorData"
+
+//======================================================================================================================
+FPhysicsControlProfileEditorData::FPhysicsControlProfileEditorData()
+{
+}
+
+//======================================================================================================================
+void FPhysicsControlProfileEditorData::Initialize(const TSharedRef<IPersonaPreviewScene>& InPreviewScene)
+{
+	PreviewScene = InPreviewScene;
+
+	EditorSkelComp = nullptr;
+	PhysicsControlComponent = nullptr;
+	FSoftObjectPath PreviewMeshStringRef = PhysicsControlProfileAsset->PreviewSkeletalMesh.ToSoftObjectPath();
+
+	// Support undo/redo
+	PhysicsControlProfileAsset->SetFlags(RF_Transactional);
+}
+
+//======================================================================================================================
+void FPhysicsControlProfileEditorData::CachePreviewMesh()
+{
+	USkeletalMesh* PreviewMesh = PhysicsControlProfileAsset->PreviewSkeletalMesh.LoadSynchronous();
+
+	if (PreviewMesh == nullptr)
+	{
+		// Fall back to the default skeletal mesh in the EngineMeshes package.
+		// This is statically loaded as the package is likely not fully loaded
+		// (otherwise, it would have been found in the above iteration).
+		PreviewMesh = (USkeletalMesh*)StaticLoadObject(USkeletalMesh::StaticClass(), NULL, TEXT("/Engine/EngineMeshes/SkeletalCube.SkeletalCube"), NULL, LOAD_None, NULL);
+		check(PreviewMesh);
+
+		PhysicsControlProfileAsset->PreviewSkeletalMesh = PreviewMesh;
+
+		FMessageDialog::Open(EAppMsgType::Ok, FText::Format(
+			LOCTEXT("Error_PhysicsControlProfileAssetHasNoSkelMesh", "Warning: Physics Control Profile Asset has no skeletal mesh assigned.\nFor now, a simple default skeletal mesh ({0}) will be used.\nYou can fix this by opening the asset and choosing another skeletal mesh from the toolbar."),
+			FText::FromString(PreviewMesh->GetFullName())));
+	}
+	else if (PreviewMesh->GetSkeleton() == nullptr)
+	{
+		// Fall back in the case of a deleted skeleton
+		PreviewMesh = (USkeletalMesh*)StaticLoadObject(USkeletalMesh::StaticClass(), NULL, TEXT("/Engine/EngineMeshes/SkeletalCube.SkeletalCube"), NULL, LOAD_None, NULL);
+		check(PreviewMesh);
+
+		PhysicsControlProfileAsset->PreviewSkeletalMesh = PreviewMesh;
+
+		FMessageDialog::Open(EAppMsgType::Ok, FText::Format(
+			LOCTEXT("Error_PhysicsControlProfileAssetHasNoSkelMeshSkeleton", "Warning: Physics Control Profile Asset has a skeletal mesh with no skeleton assigned.\nFor now, a simple default skeletal mesh ({0}) will be used.\nYou can fix this by opening the asset and choosing another skeletal mesh from the toolbar, or repairing the skeleton."),
+			FText::FromString(PreviewMesh->GetFullName())));
+	}
+}
+
+
+#undef LOCTEXT_NAMESPACE
diff --git a/Engine/Plugins/Experimental/PhysicsControl/Source/PhysicsControlEditor/Private/PhysicsControlProfileEditorData.h b/Engine/Plugins/Experimental/PhysicsControl/Source/PhysicsControlEditor/Private/PhysicsControlProfileEditorData.h
new file mode 100644
--- /dev/null
+++ b/Engine/Plugins/Experimental/PhysicsControl/Source/PhysicsControlEditor/Private/PhysicsControlProfileEditorData.h
@@ -0,0 +1,38 @@
+// Copyright Epic Games, Inc. All Rights Reserved.
+
+#pragma once
+
+#include "CoreMinimal.h"
+
+class IPersonaPreviewScene;
+class UPhysicsControlProfileAsset;
+class UPhysicsControlProfileAssetEditorSkeletalMeshComponent;
+
+/**
+ * Helper/container for data used by the Physics Control Profile Editor 
+ */
+class FPhysicsControlProfileEditorData
+{
+public:
+	FPhysicsControlProfileEditorData();
+
+	/** Initializes members */
+	void Initialize(const TSharedRef<IPersonaPreviewScene>& InPreviewScene);
+
+	/** Caches a preview mesh. Sets us to a default mesh if none is set yet (or if an older one got deleted) */
+	void CachePreviewMesh();
+
+public:
+	/** The asset being inspected */
+	UPhysicsControlProfileAsset* PhysicsControlProfileAsset;
+
+	/** Skeletal mesh component specialized for this asset editor */
+	UPhysicsControlProfileAssetEditorSkeletalMeshComponent* EditorSkelComp;
+
+	/** The physics control component used for testing/simulating on the character */
+	class UPhysicsControlComponent* PhysicsControlComponent;
+
+	/** Preview scene */
+	TWeakPtr<IPersonaPreviewScene> PreviewScene;
+
+};
diff --git a/Engine/Plugins/Experimental/PhysicsControl/Source/PhysicsControlEditor/Private/PhysicsControlProfileEditorMode.cpp b/Engine/Plugins/Experimental/PhysicsControl/Source/PhysicsControlEditor/Private/PhysicsControlProfileEditorMode.cpp
new file mode 100644
--- /dev/null
+++ b/Engine/Plugins/Experimental/PhysicsControl/Source/PhysicsControlEditor/Private/PhysicsControlProfileEditorMode.cpp
@@ -0,0 +1,38 @@
+// Copyright Epic Games, Inc. All Rights Reserved.
+
+#include "PhysicsControlProfileEditorMode.h"
+#include "PhysicsControlEditorModule.h"
+#include "PhysicsControlProfileEditorToolkit.h"
+#include "AssetEditorModeManager.h"
+
+FName FPhysicsControlProfileEditorMode::ModeName("PhysicsControlProfileAssetEditMode");
+
+//======================================================================================================================
+bool FPhysicsControlProfileEditorMode::GetCameraTarget(FSphere& OutTarget) const
+{
+	return false;
+}
+
+//======================================================================================================================
+IPersonaPreviewScene& FPhysicsControlProfileEditorMode::GetAnimPreviewScene() const
+{
+	return *static_cast<IPersonaPreviewScene*>(static_cast<FAssetEditorModeManager*>(Owner)->GetPreviewScene());
+}
+
+//======================================================================================================================
+void FPhysicsControlProfileEditorMode::Render(const FSceneView* View, FViewport* Viewport, FPrimitiveDrawInterface* PDI)
+{
+	FEdMode::Render(View, Viewport, PDI);
+}
+
+//======================================================================================================================
+void FPhysicsControlProfileEditorMode::Tick(FEditorViewportClient* ViewportClient, float DeltaTime)
+{
+	FEdMode::Tick(ViewportClient, DeltaTime);
+}
+
+//======================================================================================================================
+void FPhysicsControlProfileEditorMode::DrawHUD(FEditorViewportClient* ViewportClient, FViewport* Viewport, const FSceneView* View, FCanvas* Canvas)
+{
+	FEdMode::DrawHUD(ViewportClient, Viewport, View, Canvas);
+}
diff --git a/Engine/Plugins/Experimental/PhysicsControl/Source/PhysicsControlEditor/Private/PhysicsControlProfileEditorToolkit.cpp b/Engine/Plugins/Experimental/PhysicsControl/Source/PhysicsControlEditor/Private/PhysicsControlProfileEditorToolkit.cpp
new file mode 100644
--- /dev/null
+++ b/Engine/Plugins/Experimental/PhysicsControl/Source/PhysicsControlEditor/Private/PhysicsControlProfileEditorToolkit.cpp
@@ -0,0 +1,221 @@
+// Copyright Epic Games, Inc. All Rights Reserved.
+
+#include "PhysicsControlProfileEditorToolkit.h"
+
+#include "AnimationEditorPreviewActor.h"
+#include "EditorModeManager.h"
+#include "IDetailsView.h"
+#include "IPersonaToolkit.h"
+#include "Modules/ModuleManager.h"
+#include "PersonaModule.h"
+#include "PhysicsControlProfileApplicationMode.h"
+#include "PhysicsControlProfileAsset.h"
+#include "PhysicsControlProfileEditorMode.h"
+#include "Animation/DebugSkelMeshComponent.h"
+#include "AnimPreviewInstance.h"
+#include "PhysicsControlProfileEditorData.h"
+
+#define LOCTEXT_NAMESPACE "PhysicsControlProfileEditorToolkit"
+
+const FName PhysicsControlProfileEditorModes::Editor("PhysicsControlProfileEditorMode");
+const FName PhysicsControlProfileEditorAppName = FName(TEXT("PhysicsControlProfileEditorApp"));
+
+//======================================================================================================================
+void FPhysicsControlProfileEditorToolkit::InitAssetEditor(
+	const EToolkitMode::Type        Mode,
+	const TSharedPtr<IToolkitHost>& InitToolkitHost,
+	UPhysicsControlProfileAsset*    InPhysicsControlProfileAsset)
+{
+	bIsInitialized = false;
+
+	// Initialise EditorData
+	{
+		EditorData = MakeShared<FPhysicsControlProfileEditorData>();
+		EditorData->PhysicsControlProfileAsset = InPhysicsControlProfileAsset;
+		EditorData->CachePreviewMesh();
+	}
+
+	// Create Persona toolkit
+	{
+		FPersonaToolkitArgs PersonaToolkitArgs;
+		PersonaToolkitArgs.OnPreviewSceneCreated = FOnPreviewSceneCreated::FDelegate::CreateSP(
+			this, &FPhysicsControlProfileEditorToolkit::HandlePreviewSceneCreated);
+		FPersonaModule& PersonaModule = FModuleManager::LoadModuleChecked<FPersonaModule>("Persona");
+		PersonaToolkit = PersonaModule.CreatePersonaToolkit(InPhysicsControlProfileAsset, PersonaToolkitArgs);
+		PersonaModule.RecordAssetOpened(FAssetData(InPhysicsControlProfileAsset));
+	}
+
+	// Note - we might want to make a custom skeleton tree view here, based on showing either the
+	// animation or the physics bones/skeleton. See FPhysicsAssetEditor::InitPhysicsAssetEditor
+
+	GEditor->RegisterForUndo(this);
+
+	// Initialise the asset editor
+	const bool bCreateDefaultStandaloneMenu = true;
+	const bool bCreateDefaultToolbar = true;
+	FAssetEditorToolkit::InitAssetEditor(
+		Mode,
+		InitToolkitHost,
+		PhysicsControlProfileEditorAppName,
+		FTabManager::FLayout::NullLayout,
+		bCreateDefaultStandaloneMenu,
+		bCreateDefaultToolbar,
+		InPhysicsControlProfileAsset);
+
+	// Create and set the application mode.
+	ApplicationMode = new FPhysicsControlProfileApplicationMode(SharedThis(this), PersonaToolkit->GetPreviewScene());
+	AddApplicationMode(PhysicsControlProfileEditorModes::Editor, MakeShareable(ApplicationMode));
+	SetCurrentMode(PhysicsControlProfileEditorModes::Editor);
+
+	// Activate the editor mode.
+	GetEditorModeManager().SetDefaultMode(FPhysicsControlProfileEditorMode::ModeName);
+	GetEditorModeManager().ActivateMode(FPhysicsControlProfileEditorMode::ModeName);
+
+	FPhysicsControlProfileEditorMode* EditorMode = GetEditorModeManager().
+		GetActiveModeTyped<FPhysicsControlProfileEditorMode>(FPhysicsControlProfileEditorMode::ModeName);
+	EditorMode->SetEditorToolkit(this);
+
+	bIsInitialized = true;
+}
+
+
+//======================================================================================================================
+void FPhysicsControlProfileEditorToolkit::RegisterTabSpawners(const TSharedRef<FTabManager>& InTabManager)
+{
+	WorkspaceMenuCategory = InTabManager->AddLocalWorkspaceMenuCategory(LOCTEXT("WorkspaceMenu_MLDeformerEditor", "ML Deformer Editor"));
+	auto WorkspaceMenuCategoryRef = WorkspaceMenuCategory.ToSharedRef();
+
+	FAssetEditorToolkit::RegisterTabSpawners(InTabManager);
+}
+
+//======================================================================================================================
+void FPhysicsControlProfileEditorToolkit::UnregisterTabSpawners(const TSharedRef<FTabManager>& InTabManager)
+{
+	FAssetEditorToolkit::UnregisterTabSpawners(InTabManager);
+}
+
+//======================================================================================================================
+FName FPhysicsControlProfileEditorToolkit::GetToolkitFName() const
+{
+	return FName("PhysicsControlProfileEditor");
+}
+
+//======================================================================================================================
+FText FPhysicsControlProfileEditorToolkit::GetBaseToolkitName() const
+{
+	return LOCTEXT("PhysicsControlProfilekEditorAppLabel", "Physics Control Profile Editor");
+}
+
+//======================================================================================================================
+FText FPhysicsControlProfileEditorToolkit::GetToolkitName() const
+{
+	FFormatNamedArguments Args;
+	Args.Add(TEXT("AssetName"), FText::FromString(EditorData->PhysicsControlProfileAsset->GetName()));
+	return FText::Format(LOCTEXT("PhysicsControlProfileEditorToolkitName", "{AssetName}"), Args);
+}
+
+//======================================================================================================================
+FLinearColor FPhysicsControlProfileEditorToolkit::GetWorldCentricTabColorScale() const
+{
+	return FLinearColor::White;
+}
+
+//======================================================================================================================
+FString FPhysicsControlProfileEditorToolkit::GetWorldCentricTabPrefix() const
+{
+	return TEXT("PhysicsControlProfileEditor");
+}
+
+//======================================================================================================================
+void FPhysicsControlProfileEditorToolkit::AddReferencedObjects(FReferenceCollector& Collector)
+{
+	Collector.AddReferencedObject(EditorData->PhysicsControlProfileAsset);
+}
+
+//======================================================================================================================
+TStatId FPhysicsControlProfileEditorToolkit::GetStatId() const
+{
+	RETURN_QUICK_DECLARE_CYCLE_STAT(FPhysicsControlProfileEditorToolkit, STATGROUP_Tickables);
+}
+
+//======================================================================================================================
+TSharedRef<IPersonaToolkit> FPhysicsControlProfileEditorToolkit::GetPersonaToolkit() const
+{
+	return PersonaToolkit.ToSharedRef();
+}
+
+//======================================================================================================================
+// For inspiration here, see:
+// - PhysicsAssetEditor
+// - MLDeformerEditorToolkit
+// - IKRigToolkit
+// though note that they all do things differently!
+void FPhysicsControlProfileEditorToolkit::HandlePreviewSceneCreated(
+	const TSharedRef<IPersonaPreviewScene>& InPersonaPreviewScene)
+{
+	EditorData->PreviewScene = InPersonaPreviewScene;
+
+	AAnimationEditorPreviewActor* Actor = InPersonaPreviewScene->GetWorld()->SpawnActor<AAnimationEditorPreviewActor>(
+		AAnimationEditorPreviewActor::StaticClass(), FTransform::Identity);
+	Actor->SetFlags(RF_Transient);
+	InPersonaPreviewScene->SetActor(Actor);
+
+	// Create the preview skeletal mesh component
+	ViewportSkeletalMeshComponent = NewObject<UDebugSkelMeshComponent>(Actor);
+	// turn off default bone rendering
+	//ViewportSkeletalMeshComponent->SkeletonDrawMode = ESkeletonDrawMode::Hidden;
+
+	// Setup and apply an anim instance to the skeletal mesh component
+	ViewportAnimInstance = NewObject<UAnimPreviewInstance>(
+		ViewportSkeletalMeshComponent, TEXT("PhysicsControlProfileEditorAnimInstance"));
+	ViewportSkeletalMeshComponent->PreviewInstance = ViewportAnimInstance;
+	//ViewportAnimInstance->InitializeAnimation();
+
+	// Set the skeletal mesh on the component, using the asset. Note that this will change if/when
+	// the asset doesn't hold a mesh.
+	USkeletalMesh* Mesh = EditorData->PhysicsControlProfileAsset->PreviewSkeletalMesh.Get();
+	ViewportSkeletalMeshComponent->SetSkeletalMesh(Mesh);
+
+	// apply mesh to the preview scene
+	InPersonaPreviewScene->SetPreviewMeshComponent(ViewportSkeletalMeshComponent);
+	InPersonaPreviewScene->SetAllowMeshHitProxies(false);
+	InPersonaPreviewScene->SetAdditionalMeshesSelectable(false);
+	ViewportSkeletalMeshComponent->bSelectable = false;
+	InPersonaPreviewScene->SetPreviewMesh(Mesh);
+	InPersonaPreviewScene->AddComponent(ViewportSkeletalMeshComponent, FTransform::Identity);
+}
+
+//======================================================================================================================
+void FPhysicsControlProfileEditorToolkit::HandleViewportCreated(const TSharedRef<IPersonaViewport>& InPersonaViewport)
+{
+	PersonaViewport = InPersonaViewport;
+}
+
+//======================================================================================================================
+void FPhysicsControlProfileEditorToolkit::ShowEmptyDetails() const
+{
+	DetailsView->SetObject(EditorData->PhysicsControlProfileAsset);
+}
+
+//======================================================================================================================
+void FPhysicsControlProfileEditorToolkit::HandleDetailsCreated(const TSharedRef<class IDetailsView>& InDetailsView)
+{
+	DetailsView = InDetailsView;
+	DetailsView->OnFinishedChangingProperties().AddSP(
+		this, &FPhysicsControlProfileEditorToolkit::OnFinishedChangingDetails);
+	ShowEmptyDetails();
+}
+
+//======================================================================================================================
+void FPhysicsControlProfileEditorToolkit::OnFinishedChangingDetails(const FPropertyChangedEvent& PropertyChangedEvent)
+{
+	const bool bPreviewMeshChanged = PropertyChangedEvent.GetPropertyName() == UPhysicsControlProfileAsset::GetPreviewMeshPropertyName();
+	if (bPreviewMeshChanged)
+	{
+		USkeletalMesh* Mesh = EditorData->PhysicsControlProfileAsset->PreviewSkeletalMesh.LoadSynchronous();
+		ViewportSkeletalMeshComponent->SetSkeletalMesh(Mesh);
+		EditorData->CachePreviewMesh();
+	}
+}
+
+
diff --git a/Engine/Plugins/Experimental/PhysicsControl/Source/PhysicsControlEditor/Public/PhysicsControlEditorModule.h b/Engine/Plugins/Experimental/PhysicsControl/Source/PhysicsControlEditor/Public/PhysicsControlEditorModule.h
new file mode 100644
--- /dev/null
+++ b/Engine/Plugins/Experimental/PhysicsControl/Source/PhysicsControlEditor/Public/PhysicsControlEditorModule.h
@@ -0,0 +1,25 @@
+// Copyright Epic Games, Inc. All Rights Reserved.
+
+#pragma once
+
+#include "CoreMinimal.h"
+#include "Modules/ModuleInterface.h"
+
+class FPhysicsControlProfileAssetActions;
+
+/**
+ * Module handles the Physics Control Profile editor, and also the Physics Control Visualizer
+ */
+class FPhysicsControlEditorModule : public IModuleInterface
+{
+public:
+
+	/** IModuleInterface implementation */
+	void StartupModule() override;
+	void ShutdownModule() override;
+
+private:
+	TArray<FName> VisualizersToUnregisterOnShutdown;
+	TSharedPtr<FPhysicsControlProfileAssetActions> PhysicsControlProfileAssetActions;
+};
+
diff --git a/Engine/Plugins/Experimental/PhysicsControl/Source/PhysicsControlEditor/Public/PhysicsControlProfileApplicationMode.h b/Engine/Plugins/Experimental/PhysicsControl/Source/PhysicsControlEditor/Public/PhysicsControlProfileApplicationMode.h
new file mode 100644
--- /dev/null
+++ b/Engine/Plugins/Experimental/PhysicsControl/Source/PhysicsControlEditor/Public/PhysicsControlProfileApplicationMode.h
@@ -0,0 +1,37 @@
+// Copyright Epic Games, Inc. All Rights Reserved.
+
+#pragma once
+
+#include "WorkflowOrientedApp/ApplicationMode.h"
+#include "WorkflowOrientedApp/WorkflowTabManager.h"
+
+class FWorkflowCentricApplication;
+class IPersonaPreviewScene;
+
+class FPhysicsControlProfileEditorToolkit;
+
+/**
+ * The application mode for the Physics Control Profile Editor
+ * This defines the layout of the UI. It basically spawns all the tabs, such as the viewport, details panels, etc.
+ */
+class PHYSICSCONTROLEDITOR_API FPhysicsControlProfileApplicationMode : public FApplicationMode
+{
+public:
+	/** The name of this mode. */
+	static FName ModeName;
+
+	FPhysicsControlProfileApplicationMode(
+		TSharedRef<FWorkflowCentricApplication> InHostingApp, 
+		TSharedRef<IPersonaPreviewScene>        InPreviewScene);
+
+	// FApplicationMode overrides.
+	virtual void RegisterTabFactories(TSharedPtr<FTabManager> InTabManager) override;
+	// ~END FApplicationMode overrides.
+
+protected:
+	/** The hosting app. */
+	TWeakPtr<FPhysicsControlProfileEditorToolkit> EditorToolkit = nullptr;
+
+	/** The tab factories we support. */
+	FWorkflowAllowedTabSet TabFactories;
+};
diff --git a/Engine/Plugins/Experimental/PhysicsControl/Source/PhysicsControlEditor/Public/PhysicsControlProfileEditorMode.h b/Engine/Plugins/Experimental/PhysicsControl/Source/PhysicsControlEditor/Public/PhysicsControlProfileEditorMode.h
new file mode 100644
--- /dev/null
+++ b/Engine/Plugins/Experimental/PhysicsControl/Source/PhysicsControlEditor/Public/PhysicsControlProfileEditorMode.h
@@ -0,0 +1,40 @@
+// Copyright Epic Games, Inc. All Rights Reserved.
+
+#pragma once
+
+#include "CoreMinimal.h"
+#include "IPersonaEditMode.h"
+
+class FPhysicsControlProfileEditorToolkit;
+
+/**
+ * The Physics Control Profile Persona editor mode.
+ */
+class PHYSICSCONTROLEDITOR_API FPhysicsControlProfileEditorMode	: public IPersonaEditMode
+{
+public:
+	/** The name of the mode. */
+	static FName ModeName;
+
+	void SetEditorToolkit(FPhysicsControlProfileEditorToolkit* InToolkit) { PhysicsControlProfileEditorToolkit = InToolkit; }
+
+	// IPersonaEditMode overrides.
+	virtual bool GetCameraTarget(FSphere& OutTarget) const override;
+	virtual class IPersonaPreviewScene& GetAnimPreviewScene() const override;
+	virtual void GetOnScreenDebugInfo(TArray<FText>& OutDebugInfo) const override {}
+	// ~END IPersonaEditMode overrides.
+
+	// FEdMode overrides.
+	virtual void Tick(FEditorViewportClient* ViewportClient, float DeltaTime) override;
+	virtual void Render(const FSceneView* View, FViewport* Viewport, FPrimitiveDrawInterface* PDI) override;
+	virtual void DrawHUD(FEditorViewportClient* ViewportClient, FViewport* Viewport, const FSceneView* View, FCanvas* Canvas) override;
+	virtual bool IsCompatibleWith(FEditorModeID OtherModeID) const override { return true; }
+	virtual bool AllowWidgetMove() override { return false; }
+	virtual bool ShouldDrawWidget() const override { return false; }
+	virtual bool UsesTransformWidget() const override { return false; }
+	virtual bool UsesTransformWidget(UE::Widget::EWidgetMode CheckMode) const override { return false; }
+	// ~END FEdMode overrides.
+
+protected:
+	FPhysicsControlProfileEditorToolkit* PhysicsControlProfileEditorToolkit = nullptr;
+};
diff --git a/Engine/Plugins/Experimental/PhysicsControl/Source/PhysicsControlEditor/Public/PhysicsControlProfileEditorToolkit.h b/Engine/Plugins/Experimental/PhysicsControl/Source/PhysicsControlEditor/Public/PhysicsControlProfileEditorToolkit.h
new file mode 100644
--- /dev/null
+++ b/Engine/Plugins/Experimental/PhysicsControl/Source/PhysicsControlEditor/Public/PhysicsControlProfileEditorToolkit.h
@@ -0,0 +1,102 @@
+// Copyright Epic Games, Inc. All Rights Reserved.
+
+#pragma once
+
+#include "CoreMinimal.h"
+#include "EditorUndoClient.h"
+#include "IHasPersonaToolkit.h"
+#include "IPersonaPreviewScene.h"
+#include "IPersonaViewport.h"
+#include "PersonaAssetEditorToolkit.h"
+
+class FPhysicsControlProfileApplicationMode;
+class FPhysicsControlProfileEditorData;
+class UAnimPreviewInstance;
+class UPhysicsControlProfileAsset;
+
+namespace PhysicsControlProfileEditorModes
+{
+	extern const FName Editor;
+}
+
+/**
+ * The main toolkit/editor for working with Physics Control Profile assets
+ */
+class PHYSICSCONTROLEDITOR_API FPhysicsControlProfileEditorToolkit :
+	public FPersonaAssetEditorToolkit,
+	public IHasPersonaToolkit,
+	public FGCObject,
+	public FEditorUndoClient,
+	public FTickableEditorObject
+{
+public:
+	friend class FPhysicsControlProfileApplicationMode;
+
+public:
+
+	/** Initialize the asset editor. This will register the application mode, init the preview scene, etc. */
+	void InitAssetEditor(
+		const EToolkitMode::Type        Mode,
+		const TSharedPtr<IToolkitHost>& InitToolkitHost,
+		UPhysicsControlProfileAsset*    InPhysicsControlProfileAsset);
+
+	// FAssetEditorToolkit overrides.
+	virtual void RegisterTabSpawners(const TSharedRef<FTabManager>& InTabManager) override;
+	virtual void UnregisterTabSpawners(const TSharedRef<FTabManager>& InTabManager) override;
+	virtual FName GetToolkitFName() const override;
+	virtual FText GetBaseToolkitName() const override;
+	virtual FText GetToolkitName() const override;
+	virtual FLinearColor GetWorldCentricTabColorScale() const override;
+	virtual FString GetWorldCentricTabPrefix() const override;
+	// ~END FAssetEditorToolkit overrides.
+
+	// FGCObject overrides.
+	virtual FString GetReferencerName() const override { return TEXT("FPhysicsControlProfileEditorToolkit"); }
+	virtual void AddReferencedObjects(FReferenceCollector& Collector) override;
+	// ~END FGCObject overrides.
+
+	// FTickableEditorObject overrides.
+	virtual void Tick(float DeltaTime) override {};
+	virtual ETickableTickType GetTickableTickType() const override { return ETickableTickType::Always; }
+	virtual TStatId GetStatId() const override;
+	// ~END FTickableEditorObject overrides.
+
+	// IHasPersonaToolkit overrides.
+	virtual TSharedRef<IPersonaToolkit> GetPersonaToolkit() const override;
+	// ~END IHasPersonaToolkit overrides.
+
+	IPersonaToolkit* GetPersonaToolkitPointer() const { return PersonaToolkit.Get(); }
+
+private:
+	/** Preview scene setup. */
+	void HandlePreviewSceneCreated(const TSharedRef<IPersonaPreviewScene>& InPersonaPreviewScene);
+	void HandleViewportCreated(const TSharedRef<IPersonaViewport>& InPersonaViewport);
+	void HandleDetailsCreated(const TSharedRef<class IDetailsView>& InDetailsView);
+	void OnFinishedChangingDetails(const FPropertyChangedEvent& PropertyChangedEvent);
+	void ShowEmptyDetails() const;
+
+private:
+	/** The persona toolkit. */
+	TSharedPtr<IPersonaToolkit> PersonaToolkit = nullptr;
+
+	// Persona viewport.
+	TSharedPtr<IPersonaViewport> PersonaViewport = nullptr;
+
+	/** Data and methods shared across multiple classes */
+	TSharedPtr<FPhysicsControlProfileEditorData> EditorData;
+
+	// Asset properties tab 
+	TSharedPtr<IDetailsView> DetailsView;
+
+	FPhysicsControlProfileApplicationMode* ApplicationMode = nullptr;
+
+	// viewport anim instance 
+	UPROPERTY(transient, NonTransactional)
+	TObjectPtr<UAnimPreviewInstance> ViewportAnimInstance;
+
+	// viewport skeletal mesh 
+	UDebugSkelMeshComponent* ViewportSkeletalMeshComponent;
+
+	/** Has the asset editor been initialized? */
+	bool bIsInitialized = false;
+};
diff --git a/Engine/Plugins/Experimental/PhysicsControl/Source/PhysicsControlVisualizer/PhysicsControlVisualizer.Build.cs b/Engine/Plugins/Experimental/PhysicsControl/Source/PhysicsControlVisualizer/PhysicsControlVisualizer.Build.cs
deleted file mode 100644
diff --git a/Engine/Plugins/Experimental/PhysicsControl/Source/PhysicsControlVisualizer/Private/PhysicsControlComponentVisualizer.cpp b/Engine/Plugins/Experimental/PhysicsControl/Source/PhysicsControlVisualizer/Private/PhysicsControlComponentVisualizer.cpp
deleted file mode 100644
diff --git a/Engine/Plugins/Experimental/PhysicsControl/Source/PhysicsControlVisualizer/Private/PhysicsControlComponentVisualizer.h b/Engine/Plugins/Experimental/PhysicsControl/Source/PhysicsControlVisualizer/Private/PhysicsControlComponentVisualizer.h
deleted file mode 100644
diff --git a/Engine/Plugins/Experimental/PhysicsControl/Source/PhysicsControlVisualizer/Private/PhysicsControlVisualizerModule.cpp b/Engine/Plugins/Experimental/PhysicsControl/Source/PhysicsControlVisualizer/Private/PhysicsControlVisualizerModule.cpp
deleted file mode 100644
diff --git a/Engine/Plugins/Experimental/PhysicsControl/Source/PhysicsControlVisualizer/Public/PhysicsControlVisualizerModule.h b/Engine/Plugins/Experimental/PhysicsControl/Source/PhysicsControlVisualizer/Public/PhysicsControlVisualizerModule.h
deleted file mode 100644
