Build fix - missing undef LOCTEXT_NAMESPACE

diff --git a/Engine/Plugins/Experimental/PhysicsControl/Source/PhysicsControlEditor/Private/PhysicsControlProfileEditorToolkit.cpp b/Engine/Plugins/Experimental/PhysicsControl/Source/PhysicsControlEditor/Private/PhysicsControlProfileEditorToolkit.cpp
--- a/Engine/Plugins/Experimental/PhysicsControl/Source/PhysicsControlEditor/Private/PhysicsControlProfileEditorToolkit.cpp
+++ b/Engine/Plugins/Experimental/PhysicsControl/Source/PhysicsControlEditor/Private/PhysicsControlProfileEditorToolkit.cpp
@@ -218,4 +218,5 @@ void FPhysicsControlProfileEditorToolkit::OnFinishedChangingDetails(const FPrope
 	}
 }
 
+#undef LOCTEXT_NAMESPACE
 
