Fix some casting warnings

diff --git a/Engine/Plugins/Experimental/PhysicsControl/Source/PhysicsControl/Private/PhysicsControlComponentHelpers.cpp b/Engine/Plugins/Experimental/PhysicsControl/Source/PhysicsControl/Private/PhysicsControlComponentHelpers.cpp
--- a/Engine/Plugins/Experimental/PhysicsControl/Source/PhysicsControl/Private/PhysicsControlComponentHelpers.cpp
+++ b/Engine/Plugins/Experimental/PhysicsControl/Source/PhysicsControl/Private/PhysicsControlComponentHelpers.cpp
@@ -94,8 +94,7 @@ void ConvertConstraintProfileToControlData(
 		GetLinearDriveStiffness(InProfileProperties.LinearDrive), 
 		GetLinearDriveDamping(InProfileProperties.LinearDrive));
 
-	OutControlData.MaxForce = FloatCastChecked<float>(
-		GetLinearDriveMaxForce(InProfileProperties.LinearDrive), UE::LWC::DefaultFloatPrecision);
+	OutControlData.MaxForce = float(GetLinearDriveMaxForce(InProfileProperties.LinearDrive));
 
 	ConvertSpringToStrengthParams(
 		OutControlData.AngularStrength,
@@ -104,8 +103,7 @@ void ConvertConstraintProfileToControlData(
 		GetAngularDriveStiffness(InProfileProperties.AngularDrive),
 		GetAngularDriveDamping(InProfileProperties.AngularDrive));
 
-	OutControlData.MaxTorque = FloatCastChecked<float>(
-		GetAngularDriveMaxTorque(InProfileProperties.AngularDrive), UE::LWC::DefaultFloatPrecision);
+	OutControlData.MaxTorque = float(GetAngularDriveMaxTorque(InProfileProperties.AngularDrive));
 }
 
 //======================================================================================================================
diff --git a/Engine/Plugins/Experimental/PhysicsControl/Source/PhysicsControl/Private/PhysicsControlComponentHelpers.h b/Engine/Plugins/Experimental/PhysicsControl/Source/PhysicsControl/Private/PhysicsControlComponentHelpers.h
--- a/Engine/Plugins/Experimental/PhysicsControl/Source/PhysicsControl/Private/PhysicsControlComponentHelpers.h
+++ b/Engine/Plugins/Experimental/PhysicsControl/Source/PhysicsControl/Private/PhysicsControlComponentHelpers.h
@@ -54,11 +54,11 @@ void ConvertStrengthToSpringParams(
 	TOut& OutSpring, TOut& OutDamping,
 	double InStrength, double InDampingRatio, double InExtraDamping)
 {
-	TOut AngularFrequency = FloatCastChecked<TOut>(InStrength * UE_DOUBLE_TWO_PI, UE::LWC::DefaultFloatPrecision);
+	TOut AngularFrequency = TOut(InStrength * UE_DOUBLE_TWO_PI);
 	TOut Stiffness = AngularFrequency * AngularFrequency;
 
 	OutSpring = Stiffness;
-	OutDamping = FloatCastChecked<TOut>(InExtraDamping + 2 * InDampingRatio * AngularFrequency, UE::LWC::DefaultFloatPrecision);
+	OutDamping = TOut(InExtraDamping + 2 * InDampingRatio * AngularFrequency);
 }
 
 
@@ -72,27 +72,27 @@ static void ConvertSpringToStrengthParams(
 	const double InSpring, const double InDamping)
 {
 	// Simple calculation to get the strength
-	TOut AngularFrequency = FloatCastChecked<TOut>(FMath::Sqrt(InSpring), UE::LWC::DefaultFloatPrecision);
+	TOut AngularFrequency = TOut(FMath::Sqrt(InSpring));
 	OutStrength = AngularFrequency / UE_TWO_PI;
 
 	// For damping, try to put as much into the damping ratio as possible, up to a max DR of 1. Then
 	// the rest goes into extra damping.
 	OutDampingRatio = 1;
-	TOut ImpliedDamping = FloatCastChecked<TOut>(2.0 * OutDampingRatio * AngularFrequency, UE::LWC::DefaultFloatPrecision);
+	TOut ImpliedDamping = TOut(2.0 * OutDampingRatio * AngularFrequency);
 
 	if (ImpliedDamping < InDamping)
 	{
-		OutExtraDamping = FloatCastChecked<TOut>(InDamping, UE::LWC::DefaultFloatPrecision) - ImpliedDamping;
+		OutExtraDamping = TOut(InDamping) - ImpliedDamping;
 	}
 	else if (AngularFrequency > 0)
 	{
 		OutExtraDamping = 0;
-		OutDampingRatio = FloatCastChecked<TOut>(InDamping, UE::LWC::DefaultFloatPrecision) / (2 * AngularFrequency);
+		OutDampingRatio = TOut(InDamping) / (2 * AngularFrequency);
 	}
 	else
 	{
 		OutDampingRatio = 1;
-		OutExtraDamping = FloatCastChecked<TOut>(InDamping, UE::LWC::DefaultFloatPrecision);
+		OutExtraDamping = TOut(InDamping);
 	}
 }
 
diff --git a/Engine/Plugins/Experimental/PhysicsControl/Source/PhysicsControl/Public/PhysicsControlComponent.h b/Engine/Plugins/Experimental/PhysicsControl/Source/PhysicsControl/Public/PhysicsControlComponent.h
--- a/Engine/Plugins/Experimental/PhysicsControl/Source/PhysicsControl/Public/PhysicsControlComponent.h
+++ b/Engine/Plugins/Experimental/PhysicsControl/Source/PhysicsControl/Public/PhysicsControlComponent.h
@@ -750,6 +750,8 @@ public:
 	 * @param MovementType Whether to enable/disable simulation on the body
 	 * @param CollisionType Collision type to set on the body
 	 * @param GravityMultiplier The amount of gravity to apply when simulating
+	 * @param PhysicsBlendWeight The blend weight between the body transform coming from
+	 *        animation and that coming from simulation.
 	 * @param bUseSkeletalAnimation Whether the kinematic target is specified in the frame of the
 	 *                              skeletal animation, rather than world space. Only relevant if the
 	 *                              body is part of a skeletal mesh.
@@ -776,6 +778,8 @@ public:
 	 * @param MovementType Whether to enable/disable simulation on the body
 	 * @param CollisionType Collision type to set on the body
 	 * @param GravityMultiplier The amount of gravity to apply when simulating
+	 * @param PhysicsBlendWeight The blend weight between the body transform coming from
+	 *        animation and that coming from simulation.
 	 * @param bUseSkeletalAnimation Whether the kinematic target is specified in the frame of the
 	 *                              skeletal animation, rather than world space. Only relevant if the
 	 *                              body is part of a skeletal mesh.
@@ -803,6 +807,8 @@ public:
 	 * @param MovementType Whether to enable/disable simulation on the body
 	 * @param CollisionType Collision type to set on the body
 	 * @param GravityMultiplier The amount of gravity to apply when simulating
+	 * @param PhysicsBlendWeight The blend weight between the body transform coming from
+	 *        animation and that coming from simulation.
 	 * @param bUseSkeletalAnimation Whether the kinematic target is specified in the frame of the
 	 *                              skeletal animation, rather than world space. Only relevant if the
 	 *                              body is part of a skeletal mesh.
@@ -829,6 +835,8 @@ public:
 	 * @param MovementType Whether to enable/disable simulation on the body
 	 * @param CollisionType Collision type to set on the body
 	 * @param GravityMultiplier The amount of gravity to apply when simulating
+	 * @param PhysicsBlendWeight The blend weight between the body transform coming from
+	 *        animation and that coming from simulation.
 	 * @param bUseSkeletalAnimation Whether the kinematic target is specified in the frame of the
 	 *                              skeletal animation, rather than world space. Only relevant if the
 	 *                              body is part of a skeletal mesh.
@@ -1002,7 +1010,8 @@ public:
 	 * Sets the physics blend weight for a body modifier
 	 *
 	 * @param Name The name of the body modifier to access. 
-	 * @param PhysicsBlendWeight The amount of gravity to apply when simulating
+	 * @param PhysicsBlendWeight The blend weight between the body transform coming from 
+	 *        animation and that coming from simulation.
 	 * @return true if the body modifier was found, false if not
 	 */
 	UFUNCTION(BlueprintCallable, Category = PhysicsControl)
@@ -1014,8 +1023,9 @@ public:
 	 * Sets the physics blend weight for body modifiers
 	 *
 	 * @param Names The names of the body modifiers to access. Note that if you have these in a FPhysicsControlNameArray
-	 *              then it can be split.
-	 * @param PhysicsBlendWeight The amount of gravity to apply when simulating
+	 *        then it can be split.
+	 * @param PhysicsBlendWeight The blend weight between the body transform coming from
+	 *        animation and that coming from simulation.
 	 */
 	UFUNCTION(BlueprintCallable, Category = PhysicsControl)
 	void SetBodyModifiersPhysicsBlendWeight(
@@ -1027,7 +1037,8 @@ public:
 	 *
 	 * @param Set The set of body modifiers to modify. Standard sets will include "All" and things like
 	 *        "ArmLeft", depending on how body modifiers have been created.
-	 * @param PhysicsBlendWeight The amount of gravity to apply when simulating
+	 * @param PhysicsBlendWeight The blend weight between the body transform coming from
+	 *        animation and that coming from simulation.
 	 */
 	UFUNCTION(BlueprintCallable, Category = PhysicsControl)
 	void SetBodyModifiersInSetPhysicsBlendWeight(
diff --git a/Engine/Source/Editor/Persona/Private/PersonaBlendSpaceAnalysis.cpp b/Engine/Source/Editor/Persona/Private/PersonaBlendSpaceAnalysis.cpp
--- a/Engine/Source/Editor/Persona/Private/PersonaBlendSpaceAnalysis.cpp
+++ b/Engine/Source/Editor/Persona/Private/PersonaBlendSpaceAnalysis.cpp
@@ -302,7 +302,7 @@ FVector BlendSpaceAnalysis::CalculateSampleValue(const UBlendSpace& BlendSpace,
 		const UAnalysisProperties* AnalysisProperties = BlendSpace.AnalysisProperties[Index].Get();
 		for (const IBlendSpaceAnalysisFeature* Feature : ModularFeatures)
 		{
-			float NewPosition = FloatCastChecked<float>(AdjustedPosition[Index], UE::LWC::DefaultFloatPrecision);
+			float NewPosition = float(AdjustedPosition[Index]);
 			bAnalyzed[Index] = Feature->CalculateSampleValue(
 				NewPosition, BlendSpace, AnalysisProperties, Animation, RateScale);
 			if (bAnalyzed[Index])
diff --git a/Engine/Source/Editor/Persona/Public/BlendSpaceAnalysis.h b/Engine/Source/Editor/Persona/Public/BlendSpaceAnalysis.h
--- a/Engine/Source/Editor/Persona/Public/BlendSpaceAnalysis.h
+++ b/Engine/Source/Editor/Persona/Public/BlendSpaceAnalysis.h
@@ -283,7 +283,7 @@ static bool CalculateComponentSampleValue(
 {
 	double DoubleResult = Result;
 	bool bResult = CalculateComponentSampleValue(DoubleResult, Fn, BlendSpace, AnalysisProperties, Animation, RateScale);
-	Result = FloatCastChecked<float>(DoubleResult, UE::LWC::DefaultFloatPrecision);
+	Result = float(DoubleResult);
 	return bResult;
 }
 
