Set Physics Blend Weight to 1 when using Physics Control Component - ensures that using a mixture of kinematic and dynamic bodies works with skeletal meshes.

diff --git a/Engine/Plugins/Experimental/PhysicsControl/Source/PhysicsControl/Private/PhysicsControlComponentImpl.cpp b/Engine/Plugins/Experimental/PhysicsControl/Source/PhysicsControl/Private/PhysicsControlComponentImpl.cpp
--- a/Engine/Plugins/Experimental/PhysicsControl/Source/PhysicsControl/Private/PhysicsControlComponentImpl.cpp
+++ b/Engine/Plugins/Experimental/PhysicsControl/Source/PhysicsControl/Private/PhysicsControlComponentImpl.cpp
@@ -306,6 +306,17 @@ void FPhysicsControlComponentImpl::AddSkeletalMeshReferenceForModifier(
 	Data.OriginalKinematicBonesUpdateType = InSkeletalMeshComponent->KinematicBonesUpdateType;
 	InSkeletalMeshComponent->bUpdateMeshWhenKinematic = true;
 	InSkeletalMeshComponent->KinematicBonesUpdateType = EKinematicBonesUpdateToPhysics::SkipAllBones;
+	// By default, kinematic bodies will have their blend weight set to zero. This is a problem for us since:
+	// 1. We expect there will be lots of cases where only part of the character is dynamic, and other 
+	//    parts are kinematic
+	// 2. If those parts are towards the root of the character, then if their physics blend weight is zero, 
+	//    they are unable to "move away" from the component - e.g. if the component itself is moved by the 
+	//    movement component
+	// 3. We want to support users using the physics blend weight, so we can't simply force a physics blend 
+	//    weight of 1 in the skeletal mesh component (PhysAnim.cpp).
+	// So, we set all the bodies to have a blend weight of 1, noting that any under the control of a BodyModifier
+	// will get updated each tick.
+	InSkeletalMeshComponent->SetAllBodiesPhysicsBlendWeight(1.0f);
 }
 
 //======================================================================================================================
