PhysicsControlComponent: Allow getting current/target transforms. Allow manual updates. Initialise control strengths to 0 instead of the (arbitrary) 1. Expose damping ratio in the PhysicsControlMultiplier struct. Allow bulk setting of target positions/orientations etc. Allow resetting to the cached targets. Add profiling markup.

diff --git a/Engine/Plugins/Experimental/PhysicsControl/Source/PhysicsControl/Private/PhysicsControlComponent.cpp b/Engine/Plugins/Experimental/PhysicsControl/Source/PhysicsControl/Private/PhysicsControlComponent.cpp
--- a/Engine/Plugins/Experimental/PhysicsControl/Source/PhysicsControl/Private/PhysicsControlComponent.cpp
+++ b/Engine/Plugins/Experimental/PhysicsControl/Source/PhysicsControl/Private/PhysicsControlComponent.cpp
@@ -24,7 +24,8 @@
 
 #include "SceneManagement.h"
 
-#include UE_INLINE_GENERATED_CPP_BY_NAME(PhysicsControlComponent)
+DECLARE_CYCLE_STAT(TEXT("PhysicsControl UpdateTargetCaches"), STAT_PhysicsControl_UpdateTargetCaches, STATGROUP_Anim);
+DECLARE_CYCLE_STAT(TEXT("PhysicsControl UpdateControls"), STAT_PhysicsControl_UpdateControls, STATGROUP_Anim);
 
 //======================================================================================================================
 // This is used, rather than UEnum::GetValueAsString, so that we have more control over the string returned, which 
@@ -89,28 +90,25 @@ void UPhysicsControlComponent::BeginDestroy()
 }
 
 //======================================================================================================================
-void UPhysicsControlComponent::TickComponent(
-	float                        DeltaTime, 
-	enum ELevelTick              TickType, 
-	FActorComponentTickFunction* ThisTickFunction)
+void UPhysicsControlComponent::UpdateTargetCaches(float DeltaTime)
 {
-	Super::TickComponent(DeltaTime, TickType, ThisTickFunction);
+	SCOPE_CYCLE_COUNTER(STAT_PhysicsControl_UpdateTargetCaches);
+	TRACE_CPUPROFILER_EVENT_SCOPE(UPhysicsControlComponent::UpdateTargetCaches);
 
-	// We only want to continue the update if this is a "real" tick that corresponds to updating the
-	// world. We certainly don't want to tick during a pause, because part of the processing involves 
-	// (optionally) calculating target velocities based on target positions in previous ticks etc.
-	if (TickType != LEVELTICK_All)
-	{
-		return;
-	}
+	// Update the skeletal mesh caches
+	Implementation->UpdateCachedSkeletalBoneData(DeltaTime);
+}
+
+//======================================================================================================================
+void UPhysicsControlComponent::UpdateControls(float DeltaTime)
+{
+	SCOPE_CYCLE_COUNTER(STAT_PhysicsControl_UpdateControls);
+	TRACE_CPUPROFILER_EVENT_SCOPE(UPhysicsControlComponent::UpdateControls);
 
 	// Tidy up
 	Implementation->PhysicsControlRecords.Compact();
 	Implementation->PhysicsBodyModifiers.Compact();
 
-	// Update the skeletal mesh caches
-	Implementation->UpdateCachedSkeletalBoneData(DeltaTime);
-
 	for (TPair<FName, FPhysicsControlRecord>& RecordPair : Implementation->PhysicsControlRecords)
 	{
 		// New constraint requested when one doesn't exist
@@ -168,6 +166,13 @@ void UPhysicsControlComponent::TickComponent(
 		// expectation that zero blend weight means to disable physics.
 		BodyInstance->PhysicsBlendWeight = BodyModifier.PhysicsBlendWeight;
 
+		if (BodyModifier.bResetToCachedTarget)
+		{
+			BodyModifier.bResetToCachedTarget = false;
+			Implementation->ResetToCachedTarget(BodyModifier);
+		}
+
+
 		UBodySetup* BodySetup = BodyInstance->GetBodySetup();
 		if (BodySetup)
 		{
@@ -199,6 +204,27 @@ void UPhysicsControlComponent::TickComponent(
 	}
 }
 
+//======================================================================================================================
+void UPhysicsControlComponent::TickComponent(
+	float                        DeltaTime,
+	enum ELevelTick              TickType,
+	FActorComponentTickFunction* ThisTickFunction)
+{
+	Super::TickComponent(DeltaTime, TickType, ThisTickFunction);
+
+	// We only want to continue the update if this is a "real" tick that corresponds to updating the
+	// world. We certainly don't want to tick during a pause, because part of the processing involves 
+	// (optionally) calculating target velocities based on target positions in previous ticks etc.
+	if (TickType != LEVELTICK_All)
+	{
+		return;
+	}
+
+	UpdateTargetCaches(DeltaTime);
+
+	UpdateControls(DeltaTime);
+}
+
 //======================================================================================================================
 TMap<FName, FPhysicsControlLimbBones> UPhysicsControlComponent::GetLimbBonesFromSkeletalMesh(
 	USkeletalMeshComponent*                     SkeletalMeshComponent,
@@ -960,7 +986,7 @@ bool UPhysicsControlComponent::ResetControlPoint(const FName Name)
 
 //======================================================================================================================
 bool UPhysicsControlComponent::SetControlTarget(
-	const FName           Name, 
+	const FName                 Name, 
 	const FPhysicsControlTarget ControlTarget, 
 	const bool                  bEnableControl)
 {
@@ -1016,7 +1042,35 @@ bool UPhysicsControlComponent::SetControlTargetPositionAndOrientation(
 			Name, Orientation, VelocityDeltaTime, bEnableControl, bApplyControlPointToTarget);
 	}
 	return false;
+}
 
+//======================================================================================================================
+void UPhysicsControlComponent::SetControlTargetPositionsAndOrientations(
+	const TArray<FName>& Names,
+	const FVector        Position,
+	const FRotator       Orientation,
+	const float          VelocityDeltaTime,
+	const bool           bEnableControl,
+	const bool           bApplyControlPointToTarget)
+{
+	for (FName Name : Names)
+	{
+		SetControlTargetPositionAndOrientation(
+			Name, Position, Orientation, VelocityDeltaTime, bEnableControl, bApplyControlPointToTarget);
+	}
+}
+
+//======================================================================================================================
+void UPhysicsControlComponent::SetControlTargetPositionsAndOrientationsInSet(
+	const FName          SetName,
+	const FVector        Position,
+	const FRotator       Orientation,
+	const float          VelocityDeltaTime,
+	const bool           bEnableControl,
+	const bool           bApplyControlPointToTarget)
+{
+	SetControlTargetPositionsAndOrientations(GetControlNamesInSet(SetName), 
+		Position, Orientation, VelocityDeltaTime, bEnableControl, bApplyControlPointToTarget);
 }
 
 //======================================================================================================================
@@ -1050,6 +1104,33 @@ bool UPhysicsControlComponent::SetControlTargetPosition(
 	return false;
 }
 
+//======================================================================================================================
+void UPhysicsControlComponent::SetControlTargetPositions(
+	const TArray<FName>& Names,
+	const FVector        Position,
+	const float          VelocityDeltaTime,
+	const bool           bEnableControl,
+	const bool           bApplyControlPointToTarget)
+{
+	for (FName Name : Names)
+	{
+		SetControlTargetPosition(
+			Name, Position, VelocityDeltaTime, bEnableControl, bApplyControlPointToTarget);
+	}
+}
+
+//======================================================================================================================
+void UPhysicsControlComponent::SetControlTargetPositionsInSet(
+	const FName   SetName,
+	const FVector Position,
+	const float   VelocityDeltaTime,
+	const bool    bEnableControl,
+	const bool    bApplyControlPointToTarget)
+{
+	SetControlTargetPositions(GetControlNamesInSet(SetName),
+		Position, VelocityDeltaTime, bEnableControl, bApplyControlPointToTarget);
+}
+
 //======================================================================================================================
 bool UPhysicsControlComponent::SetControlTargetOrientation(
 	const FName    Name, 
@@ -1086,6 +1167,33 @@ bool UPhysicsControlComponent::SetControlTargetOrientation(
 	return false;
 }
 
+//======================================================================================================================
+void UPhysicsControlComponent::SetControlTargetOrientations(
+	const TArray<FName>& Names,
+	const FRotator Orientation,
+	const float    AngularVelocityDeltaTime,
+	const bool     bEnableControl,
+	const bool     bApplyControlPointToTarget)
+{
+	for (FName Name : Names)
+	{
+		SetControlTargetOrientation(
+			Name, Orientation, AngularVelocityDeltaTime, bEnableControl, bApplyControlPointToTarget);
+	}
+}
+
+//======================================================================================================================
+void UPhysicsControlComponent::SetControlTargetOrientationsInSet(
+	const FName    SetName,
+	const FRotator Orientation,
+	const float    AngularVelocityDeltaTime,
+	const bool     bEnableControl,
+	const bool     bApplyControlPointToTarget)
+{
+	SetControlTargetOrientations(GetControlNamesInSet(SetName),
+		Orientation, AngularVelocityDeltaTime, bEnableControl, bApplyControlPointToTarget);
+}
+
 //======================================================================================================================
 bool UPhysicsControlComponent::SetControlTargetPoses(
 	const FName    Name,
@@ -1785,6 +1893,180 @@ TArray<FName> UPhysicsControlComponent::GetSetsContainingBodyModifier(const FNam
 	return Result;
 }
 
+//======================================================================================================================
+TArray<FTransform> UPhysicsControlComponent::GetCachedTargetBoneTransforms(
+	const USkeletalMeshComponent* SkeletalMeshComponent, 
+	const TArray<FName>&          BoneNames)
+{
+	TArray<FTransform> Result;
+	Result.Reserve(BoneNames.Num());
+	FCachedSkeletalMeshData::FBoneData BoneData;
+
+	for (const FName& BoneName : BoneNames)
+	{
+		if (Implementation->GetBoneData(BoneData, SkeletalMeshComponent, BoneName))
+		{
+			FTransform BoneTransform(BoneData.Orientation, BoneData.Position);
+			Result.Add(BoneTransform);
+		}
+		else
+		{
+			Result.Add(FTransform::Identity);
+		}
+	}
+
+	return Result;
+}
+
+//======================================================================================================================
+TArray<FVector> UPhysicsControlComponent::GetCachedTargetBoneVelocities(
+	const USkeletalMeshComponent* SkeletalMeshComponent,
+	const TArray<FName>&          BoneNames)
+{
+	TArray<FVector> Result;
+	Result.Reserve(BoneNames.Num());
+	FCachedSkeletalMeshData::FBoneData BoneData;
+
+	for (const FName& BoneName : BoneNames)
+	{
+		if (Implementation->GetBoneData(BoneData, SkeletalMeshComponent, BoneName))
+		{
+			Result.Add(BoneData.Velocity);
+		}
+		else
+		{
+			Result.Add(FVector::ZeroVector);
+		}
+	}
+
+	return Result;
+}
+
+//======================================================================================================================
+TArray<FVector> UPhysicsControlComponent::GetCachedTargetBoneAngularVelocities(
+	const USkeletalMeshComponent* SkeletalMeshComponent,
+	const TArray<FName>&          BoneNames)
+{
+	TArray<FVector> Result;
+	Result.Reserve(BoneNames.Num());
+	FCachedSkeletalMeshData::FBoneData BoneData;
+
+	for (const FName& BoneName : BoneNames)
+	{
+		if (Implementation->GetBoneData(BoneData, SkeletalMeshComponent, BoneName))
+		{
+			Result.Add(BoneData.AngularVelocity);
+		}
+		else
+		{
+			Result.Add(FVector::ZeroVector);
+		}
+	}
+
+	return Result;
+}
+
+//======================================================================================================================
+FTransform UPhysicsControlComponent::GetCachedTargetBoneTransform(
+	const USkeletalMeshComponent* SkeletalMeshComponent,
+	const FName                   BoneName)
+{
+	FCachedSkeletalMeshData::FBoneData BoneData;
+	if (Implementation->GetBoneData(BoneData, SkeletalMeshComponent, BoneName))
+	{
+		return FTransform(BoneData.Orientation, BoneData.Position);
+	}
+	return FTransform();
+}
+
+//======================================================================================================================
+FVector UPhysicsControlComponent::GetCachedTargetBoneVelocity(
+	const USkeletalMeshComponent* SkeletalMeshComponent,
+	const FName                   BoneName)
+{
+	FCachedSkeletalMeshData::FBoneData BoneData;
+	if (Implementation->GetBoneData(BoneData, SkeletalMeshComponent, BoneName))
+	{
+		return BoneData.Velocity;
+	}
+	return FVector::Zero();
+}
+
+//======================================================================================================================
+FVector UPhysicsControlComponent::GetCachedTargetBoneAngularVelocity(
+	const USkeletalMeshComponent* SkeletalMeshComponent,
+	const FName                   BoneName)
+{
+	FCachedSkeletalMeshData::FBoneData BoneData;
+	if (Implementation->GetBoneData(BoneData, SkeletalMeshComponent, BoneName))
+	{
+		return BoneData.AngularVelocity;
+	}
+	return FVector::Zero();
+}
+
+//======================================================================================================================
+bool UPhysicsControlComponent::SetCachedTargetBoneData(
+	const USkeletalMeshComponent* SkeletalMeshComponent,
+	const FName                   BoneName,
+	const FTransform&             TM,
+	const FVector                 Velocity,
+	const FVector                 AngularVelocity)
+{
+	FCachedSkeletalMeshData::FBoneData* BoneData;
+	if (Implementation->GetModifiableBoneData(BoneData, SkeletalMeshComponent, BoneName))
+	{
+		BoneData->Position = TM.GetLocation();
+		BoneData->Orientation = TM.GetRotation();
+		BoneData->Velocity = Velocity;
+		BoneData->AngularVelocity = AngularVelocity;
+		return true;
+	}
+	return false;
+}
+
+//======================================================================================================================
+bool UPhysicsControlComponent::ResetBodyModifierToCachedTarget(
+	const FName                        Name,
+	const EResetToCachedTargetBehavior Behavior)
+{
+	FPhysicsBodyModifier* PhysicsBodyModifier = Implementation->FindBodyModifier(Name);
+	if (PhysicsBodyModifier)
+	{
+		if (Behavior == EResetToCachedTargetBehavior::ResetImmediately)
+		{
+			Implementation->ResetToCachedTarget(*PhysicsBodyModifier);
+		}
+		else
+		{
+			PhysicsBodyModifier->bResetToCachedTarget = true;
+		}
+		return true;
+	}
+	return false;
+
+}
+
+//======================================================================================================================
+void UPhysicsControlComponent::ResetBodyModifiersToCachedTargets(
+	const TArray<FName>&               Names,
+	const EResetToCachedTargetBehavior Behavior)
+{
+	for (FName Name : Names)
+	{
+		ResetBodyModifierToCachedTarget(Name, Behavior);
+	}
+}
+
+//======================================================================================================================
+void UPhysicsControlComponent::ResetBodyModifiersInSetToCachedTargets(
+	const FName                        SetName,
+	const EResetToCachedTargetBehavior Behavior)
+{
+	ResetBodyModifiersToCachedTargets(GetBodyModifierNamesInSet(SetName), Behavior);
+}
+
+
 #if WITH_EDITOR
 
 //======================================================================================================================
diff --git a/Engine/Plugins/Experimental/PhysicsControl/Source/PhysicsControl/Private/PhysicsControlComponentHelpers.cpp b/Engine/Plugins/Experimental/PhysicsControl/Source/PhysicsControl/Private/PhysicsControlComponentHelpers.cpp
--- a/Engine/Plugins/Experimental/PhysicsControl/Source/PhysicsControl/Private/PhysicsControlComponentHelpers.cpp
+++ b/Engine/Plugins/Experimental/PhysicsControl/Source/PhysicsControl/Private/PhysicsControlComponentHelpers.cpp
@@ -15,11 +15,11 @@ namespace PhysicsControlComponent
 //======================================================================================================================
 void ConvertStrengthToSpringParams(
 	FVector& OutSpring, FVector& OutDamping, 
-	const FVector& InStrength, float InDampingRatio, const FVector& InExtraDamping)
+	const FVector& InStrength, const FVector& InDampingRatio, const FVector& InExtraDamping)
 {
-	ConvertStrengthToSpringParams(OutSpring.X, OutDamping.X, InStrength.X, InDampingRatio, InExtraDamping.X);
-	ConvertStrengthToSpringParams(OutSpring.Y, OutDamping.Y, InStrength.Y, InDampingRatio, InExtraDamping.Y);
-	ConvertStrengthToSpringParams(OutSpring.Z, OutDamping.Z, InStrength.Z, InDampingRatio, InExtraDamping.Z);
+	ConvertStrengthToSpringParams(OutSpring.X, OutDamping.X, InStrength.X, InDampingRatio.X, InExtraDamping.X);
+	ConvertStrengthToSpringParams(OutSpring.Y, OutDamping.Y, InStrength.Y, InDampingRatio.Y, InExtraDamping.Y);
+	ConvertStrengthToSpringParams(OutSpring.Z, OutDamping.Z, InStrength.Z, InDampingRatio.Z, InExtraDamping.Z);
 }
 
 inline double GetLinearDriveStiffness(const FLinearDriveConstraint& InDrive)
diff --git a/Engine/Plugins/Experimental/PhysicsControl/Source/PhysicsControl/Private/PhysicsControlComponentHelpers.h b/Engine/Plugins/Experimental/PhysicsControl/Source/PhysicsControl/Private/PhysicsControlComponentHelpers.h
--- a/Engine/Plugins/Experimental/PhysicsControl/Source/PhysicsControl/Private/PhysicsControlComponentHelpers.h
+++ b/Engine/Plugins/Experimental/PhysicsControl/Source/PhysicsControl/Private/PhysicsControlComponentHelpers.h
@@ -22,7 +22,7 @@ namespace PhysicsControlComponent
  */
 void ConvertStrengthToSpringParams(
 	FVector& OutSpring, FVector& OutDamping, 
-	const FVector& InStrength, float InDampingRatio, const FVector& InExtraDamping);
+	const FVector& InStrength, const FVector& InDampingRatio, const FVector& InExtraDamping);
 
 /**
  * Converts the drive settings from the constraint profile into the control data strength/damping etc. 
diff --git a/Engine/Plugins/Experimental/PhysicsControl/Source/PhysicsControl/Private/PhysicsControlComponentImpl.cpp b/Engine/Plugins/Experimental/PhysicsControl/Source/PhysicsControl/Private/PhysicsControlComponentImpl.cpp
--- a/Engine/Plugins/Experimental/PhysicsControl/Source/PhysicsControl/Private/PhysicsControlComponentImpl.cpp
+++ b/Engine/Plugins/Experimental/PhysicsControl/Source/PhysicsControl/Private/PhysicsControlComponentImpl.cpp
@@ -81,6 +81,40 @@ bool FPhysicsControlComponentImpl::GetBoneData(
 	return false;
 }
 
+//======================================================================================================================
+bool FPhysicsControlComponentImpl::GetModifiableBoneData(
+	FCachedSkeletalMeshData::FBoneData*& OutBoneData,
+	const USkeletalMeshComponent*        InSkeletalMeshComponent,
+	const FName                          InBoneName)
+{
+	check(InSkeletalMeshComponent);
+	const FReferenceSkeleton& RefSkeleton = InSkeletalMeshComponent->GetSkeletalMeshAsset()->GetRefSkeleton();
+	const int32 BoneIndex = RefSkeleton.FindBoneIndex(InBoneName);
+
+	if (BoneIndex == INDEX_NONE)
+	{
+		UE_LOG(LogPhysicsControlComponent, Warning, TEXT("Failed to find BoneIndex for %s"), *InBoneName.ToString());
+		return false;
+	}
+
+	FCachedSkeletalMeshData* CachedSkeletalMeshData = CachedSkeletalMeshDatas.Find(InSkeletalMeshComponent);
+	if (CachedSkeletalMeshData &&
+		CachedSkeletalMeshData->ReferenceCount > 0 &&
+		!CachedSkeletalMeshData->BoneData.IsEmpty())
+	{
+		if (BoneIndex < CachedSkeletalMeshData->BoneData.Num())
+		{
+			OutBoneData = &CachedSkeletalMeshData->BoneData[BoneIndex];
+			return true;
+		}
+		UE_LOG(LogPhysicsControlComponent, Warning, TEXT("BoneIndex is out of range"));
+
+	}
+	UE_LOG(LogPhysicsControlComponent, Warning, TEXT("Failed to find bone data for %s"), *InBoneName.ToString());
+	return false;
+}
+
+
 //======================================================================================================================
 FPhysicsControlRecord* FPhysicsControlComponentImpl::FindControlRecord(const FName Name)
 {
@@ -148,18 +182,31 @@ void FPhysicsControlComponentImpl::UpdateCachedSkeletalBoneData(float Dt)
 		{
 			FTransform ComponentTM = SkeletalMesh->GetComponentToWorld();
 			const TArray<FTransform>& TMs = SkeletalMesh->GetEditableComponentSpaceTransforms();
-			if (TMs.Num() == CachedSkeletalMeshData.BoneData.Num() &&
+			const int32 NumTMs = TMs.Num();
+			if (NumTMs == CachedSkeletalMeshData.BoneData.Num() &&
 				!DetectTeleport(CachedSkeletalMeshData.ComponentTM, ComponentTM))
 			{
-				for (int32 Index = 0; Index != TMs.Num(); ++Index)
+				// Avoid the Dt test on every bone
+				if (Dt > 0)
+				{
+					for (int32 Index = 0; Index != NumTMs; ++Index)
+					{
+						FTransform TM = TMs[Index] * ComponentTM;
+						CachedSkeletalMeshData.BoneData[Index].Update(TM.GetTranslation(), TM.GetRotation(), Dt);
+					}
+				}
+				else
 				{
-					FTransform TM = TMs[Index] * ComponentTM;
-					CachedSkeletalMeshData.BoneData[Index].Update(TM.GetTranslation(), TM.GetRotation(), Dt);
+					for (int32 Index = 0; Index != NumTMs; ++Index)
+					{
+						FTransform TM = TMs[Index] * ComponentTM;
+						CachedSkeletalMeshData.BoneData[Index].Update(TM.GetTranslation(), TM.GetRotation());
+					}
 				}
 			}
 			else
 			{
-				CachedSkeletalMeshData.BoneData.Empty(TMs.Num());
+				CachedSkeletalMeshData.BoneData.Empty(NumTMs);
 				for (const FTransform& BoneTM : TMs)
 				{
 					FTransform TM = BoneTM * ComponentTM;
@@ -203,7 +250,7 @@ void FPhysicsControlComponentImpl::ApplyKinematicTarget(const FPhysicsBodyModifi
 			return;
 		}
 
-		FTransform TM = BodyInstance->GetUnrealWorldTransform();
+		FTransform TM = BodyInstance->GetUnrealWorldTransform(); // Preserve scale
 		FTransform KinematicTarget = TM;
 		KinematicTarget.SetRotation(BodyModifier.KinematicTargetOrientation);
 		KinematicTarget.SetTranslation(BodyModifier.KinematicTargetPosition);
@@ -232,6 +279,32 @@ void FPhysicsControlComponentImpl::ApplyKinematicTarget(const FPhysicsBodyModifi
 	}
 }
 
+//======================================================================================================================
+void FPhysicsControlComponentImpl::ResetToCachedTarget(const FPhysicsBodyModifier& BodyModifier) const
+{
+	FBodyInstance* BodyInstance = UE::PhysicsControlComponent::GetBodyInstance(
+		BodyModifier.MeshComponent, BodyModifier.BoneName);
+	if (!BodyInstance)
+	{
+		return;
+	}
+
+	if (USkeletalMeshComponent* SkeletalMeshComponent = GetValidSkeletalMeshComponentFromBodyModifier(BodyModifier))
+	{
+		FCachedSkeletalMeshData::FBoneData BoneData;
+		if (GetBoneData(BoneData, SkeletalMeshComponent, BodyModifier.BoneName))
+		{
+			FTransform BoneTM = BodyInstance->GetUnrealWorldTransform(); // Preserve scale
+			BoneTM.SetLocation(BoneData.Position);
+			BoneTM.SetRotation(BoneData.Orientation);
+
+			BodyInstance->SetBodyTransform(BoneTM, ETeleportType::TeleportPhysics);
+			BodyInstance->SetLinearVelocity(BoneData.Velocity, false);
+			BodyInstance->SetAngularVelocityInRadians(BoneData.AngularVelocity, false);
+		}
+	}
+}
+
 //======================================================================================================================
 void FPhysicsControlComponentImpl::AddSkeletalMeshReferenceForCaching(
 	USkeletalMeshComponent* InSkeletalMeshComponent)
@@ -589,12 +662,12 @@ bool FPhysicsControlComponentImpl::ApplyControlStrengths(
 	UE::PhysicsControlComponent::ConvertStrengthToSpringParams(
 		AngularSpring, AngularDamping,
 		Data.AngularStrength * Multiplier.AngularStrengthMultiplier,
-		Data.AngularDampingRatio,
+		Data.AngularDampingRatio * Multiplier.AngularDampingRatioMultiplier,
 		Data.AngularExtraDamping * Multiplier.AngularExtraDampingMultiplier);
 	UE::PhysicsControlComponent::ConvertStrengthToSpringParams(
 		LinearSpring, LinearDamping,
 		Data.LinearStrength * Multiplier.LinearStrengthMultiplier,
-		Data.LinearDampingRatio,
+		Data.LinearDampingRatio * Multiplier.LinearDampingRatioMultiplier,
 		Data.LinearExtraDamping * Multiplier.LinearExtraDampingMultiplier);
 
 	if (Multiplier.MaxTorqueMultiplier <= 0)
diff --git a/Engine/Plugins/Experimental/PhysicsControl/Source/PhysicsControl/Private/PhysicsControlComponentImpl.h b/Engine/Plugins/Experimental/PhysicsControl/Source/PhysicsControl/Private/PhysicsControlComponentImpl.h
--- a/Engine/Plugins/Experimental/PhysicsControl/Source/PhysicsControl/Private/PhysicsControlComponentImpl.h
+++ b/Engine/Plugins/Experimental/PhysicsControl/Source/PhysicsControl/Private/PhysicsControlComponentImpl.h
@@ -34,6 +34,19 @@ public:
 		const USkeletalMeshComponent*       InSkeletalMeshComponent,
 		const FName                         InBoneName) const;
 
+	/*
+	 * Retrieves the bone data for the specified bone given the skeletal mesh component, for modification
+	 *
+	 * @param OutBoneData If successful, this will point to the bone data.
+	 * @param InSkeletalMeshComponent Required to be a valid pointer to a skeletal mesh component
+	 * @param InBoneName The name of the bone to retrieve data for
+	 * @return true if the bone and data were found, false if not (in which case warnings will be logged)
+	 */
+	bool GetModifiableBoneData(
+		FCachedSkeletalMeshData::FBoneData*& OutBoneData,
+		const USkeletalMeshComponent*        InSkeletalMeshComponent,
+		const FName                          InBoneName);
+
 	/**
 	 * Retrieves the control for the name. Note that if Name is blank then the first control will be returned,
 	 * assuming there is one.
@@ -114,6 +127,12 @@ public:
 	 */
 	void ApplyKinematicTarget(const FPhysicsBodyModifier& BodyModifier) const;
 
+	/**
+	 * Sets the body (simulated or kinematic) to have the position/velocity etc state that has been cached. 
+	 * Has no effect if there is no cached target data.
+	 */
+	void ResetToCachedTarget(const FPhysicsBodyModifier& BodyModifier) const;
+
 	/**
 	 * Retrieves the body modifier for the name. Note that if Name is blank then the first modifier will be returned,
 	 * assuming there is one.
diff --git a/Engine/Plugins/Experimental/PhysicsControl/Source/PhysicsControl/Private/PhysicsControlRecord.cpp b/Engine/Plugins/Experimental/PhysicsControl/Source/PhysicsControl/Private/PhysicsControlRecord.cpp
--- a/Engine/Plugins/Experimental/PhysicsControl/Source/PhysicsControl/Private/PhysicsControlRecord.cpp
+++ b/Engine/Plugins/Experimental/PhysicsControl/Source/PhysicsControl/Private/PhysicsControlRecord.cpp
@@ -94,24 +94,3 @@ void FPhysicsControlRecord::ResetControlPoint()
 	UpdateConstraintControlPoint();
 }
 
-//======================================================================================================================
-void FCachedSkeletalMeshData::FBoneData::Update(const FVector& InPosition, const FQuat& InOrientation, float Dt)
-{
-	if (Dt > 0)
-	{
-		Velocity = (InPosition - Position) / Dt;
-		Orientation.EnforceShortestArcWith(InOrientation);
-		// Note that quats multiply in the opposite order to TMs
-		FQuat DeltaQ = InOrientation * Orientation.Inverse();
-		AngularVelocity = DeltaQ.ToRotationVector() / Dt;
-	}
-	else
-	{
-		Velocity = FVector::ZeroVector;
-		AngularVelocity = FVector::ZeroVector;
-	}
-	Position = InPosition;
-	Orientation = InOrientation;
-}
-
-
diff --git a/Engine/Plugins/Experimental/PhysicsControl/Source/PhysicsControl/Private/PhysicsControlRecord.h b/Engine/Plugins/Experimental/PhysicsControl/Source/PhysicsControl/Private/PhysicsControlRecord.h
--- a/Engine/Plugins/Experimental/PhysicsControl/Source/PhysicsControl/Private/PhysicsControlRecord.h
+++ b/Engine/Plugins/Experimental/PhysicsControl/Source/PhysicsControl/Private/PhysicsControlRecord.h
@@ -73,6 +73,7 @@ struct FPhysicsBodyModifier
 		, KinematicTargetPosition(FVector::ZeroVector)
 		, KinematicTargetOrientation(FQuat::Identity)
 		, bUseSkeletalAnimation(InUseSkeletalAnimation)
+		, bResetToCachedTarget(false)
 	{}
 
 	/**  The mesh that will be modified. */
@@ -112,6 +113,12 @@ struct FPhysicsBodyModifier
 
 	/** If true then the target will be applied on top of the skeletal animation (if there is any) */
 	uint8 bUseSkeletalAnimation:1;
+
+	/** 
+	 * If true then the body will be set to the transform/velocity stored in any cached target (if that
+	 * exists), and then this flag will be cleared.
+	 */
+	uint8 bResetToCachedTarget:1;
 };
 
 /**
@@ -157,9 +164,15 @@ public:
 			Velocity(FVector::ZeroVector), AngularVelocity(FVector::ZeroVector) {}
 
 		/**
-		 * Sets position and velocity, and calculates velocity if Dt > 0 (otherwise sets it to zero)
+		 * Sets position/orientation and calculates velocity/angular velocity - requires Dt > 0 
 		 */
 		void Update(const FVector& InPosition, const FQuat& InOrientation, float Dt);
+
+		/**
+		 * Sets position/orientation, and sets velocity to zero
+		 */
+		void Update(const FVector& InPosition, const FQuat& InOrientation);
+
 		FTransform GetTM() const { return FTransform(Orientation, Position); }
 
 		FVector Position;
@@ -187,3 +200,25 @@ public:
 	int32 ReferenceCount;
 };
 
+//======================================================================================================================
+inline void FCachedSkeletalMeshData::FBoneData::Update(const FVector& InPosition, const FQuat& InOrientation, float Dt)
+{
+	check(Dt > 0);
+
+	Velocity = (InPosition - Position) / Dt;
+	Orientation.EnforceShortestArcWith(InOrientation);
+	// Note that quats multiply in the opposite order to TMs
+	const FQuat DeltaQ = InOrientation * Orientation.Inverse();
+	AngularVelocity = DeltaQ.ToRotationVector() / Dt;
+	Position = InPosition;
+	Orientation = InOrientation;
+}
+
+//======================================================================================================================
+inline void FCachedSkeletalMeshData::FBoneData::Update(const FVector& InPosition, const FQuat& InOrientation)
+{
+	Position = InPosition;
+	Orientation = InOrientation;
+	Velocity = FVector::ZeroVector;
+	AngularVelocity = FVector::ZeroVector;
+}
diff --git a/Engine/Plugins/Experimental/PhysicsControl/Source/PhysicsControl/Public/PhysicsControlComponent.h b/Engine/Plugins/Experimental/PhysicsControl/Source/PhysicsControl/Public/PhysicsControlComponent.h
--- a/Engine/Plugins/Experimental/PhysicsControl/Source/PhysicsControl/Public/PhysicsControlComponent.h
+++ b/Engine/Plugins/Experimental/PhysicsControl/Source/PhysicsControl/Public/PhysicsControlComponent.h
@@ -32,6 +32,26 @@ enum class EPhysicsControlType : uint8
 	ParentSpace,
 };
 
+/**
+ * Specifies how any reset to cached target should work. 
+ */
+UENUM(BlueprintType)
+enum class EResetToCachedTargetBehavior : uint8
+{
+	/** 
+	 * Reset of the associated physics bodies is done immediately, to whatever transforms are in the cache. These
+	 * will reflect the previous or the upcoming animation targets, depending on what stage of the tick this 
+	 * is called.
+	 */
+	ResetImmediately,
+	/**
+	 * Reset of the associated physics bodies will be done during the next Tick (or UpdateControls). Note that
+	 * this will reset the bodies to the updated animation targets, and then there will be a subsequent physics 
+	 * update, which will result in the final transforms.
+	 */
+	ResetDuringUpdateControls
+};
+
 /**
  * This is the main Physics Control Component class which manages Controls and Body Modifiers associated 
  * with one or more static or skeletal meshes. You can add this as a component to an actor containing a 
@@ -52,6 +72,31 @@ class PHYSICSCONTROL_API UPhysicsControlComponent : public USceneComponent
 	GENERATED_UCLASS_BODY()
 
 public:
+
+	/**
+	* Allows manual ticking so that your code can run in between updating the target caches and updating 
+	* the controls. This allows you to read the targets coming from animation and use those values to 
+	* create your own controls etc.
+	* 
+	* To use this function, you should disable ticking of the Physics Control Component, and ensure that the 
+	* relevant Skeletal Mesh Component (if being used) has ticked, using a tick prerequisite. Then explicitly 
+	* call (in order) UpdateTargetCaches and UpdateControls as you process your tick.
+	*/
+	UFUNCTION(BlueprintCallable, Category = PhysicsControl)
+	void UpdateTargetCaches(float DeltaTime);
+
+	/**
+	* Allows manual ticking so that your code can run in between updating the target caches and updating
+	* the controls and body modifiers. This allows you to read the targets coming from animation and use 
+	* those values to create your own controls etc.
+	*
+	* To use this function, you should disable ticking of the Physics Control Component, and ensure that the
+	* relevant Skeletal Mesh Component (if being used) has ticked, using a tick prerequisite. Then explicitly
+	* call (in order) UpdateTargetCaches and UpdateControls as you process your tick.
+	*/
+	UFUNCTION(BlueprintCallable, Category = PhysicsControl)
+	void UpdateControls(float DeltaTime);
+
 	/**
 	 * Creates a new control for mesh components
 	 * 
@@ -503,6 +548,30 @@ public:
 		const bool     bEnableControl = true,
 		const bool     bApplyControlPointToTarget = false);
 
+	/**
+	 * Calls SetControlTargetPositionAndOrientation for each of the control names
+	 */
+	UFUNCTION(BlueprintCallable, Category = PhysicsControl)
+	void SetControlTargetPositionsAndOrientations(
+		const TArray<FName>& Names,
+		const FVector        Position,
+		const FRotator       Orientation,
+		const float          VelocityDeltaTime,
+		const bool           bEnableControl = true,
+		const bool           bApplyControlPointToTarget = false);
+
+	/**
+	 * Calls SetControlTargetPositionAndOrientation for each control in the set
+	 */
+	UFUNCTION(BlueprintCallable, Category = PhysicsControl)
+	void SetControlTargetPositionsAndOrientationsInSet(
+		const FName          SetName,
+		const FVector        Position,
+		const FRotator       Orientation,
+		const float          VelocityDeltaTime,
+		const bool           bEnableControl = true,
+		const bool           bApplyControlPointToTarget = false);
+
 	/**
 	 * Modifies an existing control target - i.e. what it is driving towards, relative to the parent object
 	 *
@@ -528,6 +597,28 @@ public:
 		const bool    bEnableControl = true,
 		const bool    bApplyControlPointToTarget = false);
 
+	/**
+	 * Calls SetControlTargetPosition for each of the control names
+	 */
+	UFUNCTION(BlueprintCallable, Category = PhysicsControl)
+	void SetControlTargetPositions(
+		const TArray<FName>& Names,
+		const FVector        Position, 
+		const float          VelocityDeltaTime, 
+		const bool           bEnableControl = true,
+		const bool           bApplyControlPointToTarget = false);
+
+	/**
+	 * Calls SetControlTargetPosition for each of the controls in the set
+	 */
+	UFUNCTION(BlueprintCallable, Category = PhysicsControl)
+	void SetControlTargetPositionsInSet(
+		const FName   SetName,
+		const FVector Position, 
+		const float   VelocityDeltaTime, 
+		const bool    bEnableControl = true,
+		const bool    bApplyControlPointToTarget = false);
+
 	/**
 	 * Modifies an existing control target - i.e. what it is driving towards, relative to the parent object
 	 *
@@ -553,6 +644,28 @@ public:
 		const bool     bEnableControl = true,
 		const bool     bApplyControlPointToTarget = false);
 
+	/**
+	 * Calls SetControlTargetOrientation for each of the control names
+	 */
+	UFUNCTION(BlueprintCallable, Category = PhysicsControl)
+	void SetControlTargetOrientations(
+		const TArray<FName>& Names,
+		const FRotator Orientation,
+		const float    AngularVelocityDeltaTime, 
+		const bool     bEnableControl = true,
+		const bool     bApplyControlPointToTarget = false);
+
+	/**
+	 * Calls SetControlTargetOrientation for each of the controls in the set
+	 */
+	UFUNCTION(BlueprintCallable, Category = PhysicsControl)
+	void SetControlTargetOrientationsInSet(
+		const FName    SetName,
+		const FRotator Orientation, 
+		const float    AngularVelocityDeltaTime, 
+		const bool     bEnableControl = true,
+		const bool     bApplyControlPointToTarget = false);
+
 	/**
 	 * Calculates and sets an existing control target. This takes the "virtual" position/orientation of the parent 
 	 * and child and calculates the relative control. Note that this will set bApplyControlPointToTarget to true.
@@ -1196,6 +1309,124 @@ public:
 	UFUNCTION(BlueprintCallable, Category = PhysicsControl)
 	TArray<FName> GetSetsContainingBodyModifier(const FName Control) const;
 
+	/**
+	 * Gets the transforms of the requested bones that will be used as targets (in world space). Targets for bones 
+	 * that are not found will be set to Identity. Note that these targets will have been calculated and cached 
+	 * at the start of the Physics Control Component, so if using the built in tick, may be too old to be useful. 
+	 * If you manually update the component then you can access these target transforms prior to applying your 
+	 * own targets.
+	 */
+	UFUNCTION(BlueprintCallable, Category = PhysicsControl)
+	TArray<FTransform> GetCachedTargetBoneTransforms(
+		const USkeletalMeshComponent* SkeletalMeshComponent,
+		const TArray<FName>&          BoneNames);
+
+	/**
+	 * Gets the linear velocities of the requested bones that will be used as targets (in world space). Target 
+	 * velocities for bones that are not found will be set to zero. Note that these targets will have been 
+	 * calculated and cached at the start of the Physics Control Component, so if using the built in tick, 
+	 * may be too old to be useful. If you manually update the component then you can access these target 
+	 * velocities prior to applying your own targets.
+	 */
+	UFUNCTION(BlueprintCallable, Category = PhysicsControl)
+	TArray<FVector> GetCachedTargetBoneVelocities(
+		const USkeletalMeshComponent* SkeletalMeshComponent,
+		const TArray<FName>&          BoneNames);
+
+	/**
+	 * Gets the angular velocities of the requested bones that will be used as targets (in world space). Target 
+	 * velocities for bones that are not found will be set to zero. Note that these targets will have been 
+	 * calculated and cached at the start of the Physics Control Component, so if using the built in tick, 
+	 * may be too old to be useful. If you manually update the component then you can access these target 
+	 * velocities prior to applying your own targets.
+	 */
+	UFUNCTION(BlueprintCallable, Category = PhysicsControl)
+	TArray<FVector> GetCachedTargetBoneAngularVelocities(
+		const USkeletalMeshComponent* SkeletalMeshComponent,
+		const TArray<FName>&          BoneNames);
+
+	/**
+	 * Gets the transforms of the requested bone that will be used as a target (in world space). Targets for bones
+	 * that are not found will be set to Identity. Note that these targets will have been calculated and cached
+	 * at the start of the Physics Control Component, so if using the built in tick, may be too old to be useful.
+	 * If you manually update the component then you can access these target transforms prior to applying your
+	 * own targets.
+	 */
+	UFUNCTION(BlueprintCallable, Category = PhysicsControl)
+	FTransform GetCachedTargetBoneTransform(
+		const USkeletalMeshComponent* SkeletalMeshComponent,
+		const FName                   BoneName);
+
+	/**
+	 * Gets the linear velocity of the requested bone that will be used as a target (in world space). Target 
+	 * velocities for bones that are not found will be set to zero. Note that these targets will have been 
+	 * calculated and cached at the start of the Physics Control Component, so if using the built in tick, 
+	 * may be too old to be useful. If you manually update the component then you can access these target 
+	 * velocities prior to applying your own targets.
+	 */
+	UFUNCTION(BlueprintCallable, Category = PhysicsControl)
+	FVector GetCachedTargetBoneVelocity(
+		const USkeletalMeshComponent* SkeletalMeshComponent,
+		const FName                   BoneName);
+
+	/**
+	 * Gets the angular velocity of the requested bone that will be used as a target (in world space). Target 
+	 * velocities for bones that are not found will be set to zero. Note that these targets will have been 
+	 * calculated and cached at the start of the Physics Control Component update, so if using the built in tick, 
+	 * may be too old to be useful. If you manually update the component then you can access these target 
+	 * velocities prior to applying your own targets.
+	 */
+	UFUNCTION(BlueprintCallable, Category = PhysicsControl)
+	FVector GetCachedTargetBoneAngularVelocity(
+		const USkeletalMeshComponent* SkeletalMeshComponent,
+		const FName                   BoneName);
+
+	/**
+	 * This allows the caller to override the target that will have been calculated and cached at the start of 
+	 * the Physics Control Component update. This is unlikely to be useful when using the built in tick, 
+	 * but if you are manually updating the component then you may wish to call this after UpdateTargetCaches 
+	 * but before UpdateControls. 
+	 * 
+	 * @return true if successful, and false if no cached target can be found for the bone.
+	 */
+	UFUNCTION(BlueprintCallable, Category = PhysicsControl)
+	bool SetCachedTargetBoneData(
+		const USkeletalMeshComponent* SkeletalMeshComponent,
+		const FName                   BoneName, 
+		const FTransform&             TM,
+		const FVector                 Velocity,
+		const FVector                 AngularVelocity);
+
+	/**
+	 * This flags the body associated with the modifier to set (using teleport) its position and velocity to 
+	 * the cached animation target. This will only affect skeletal mesh component bodies. 
+	 * 
+	 * @param Name     The name of the body modifier to use to identify the body to reset.
+	 * @param Behavior When the reset should happen.
+	 * 
+	 * @return true if the body modifier is found (even if no cached target is found), and false otherwise.
+	 */
+	UFUNCTION(BlueprintCallable, Category = PhysicsControl)
+	bool ResetBodyModifierToCachedTarget(
+		const FName                        Name,
+		const EResetToCachedTargetBehavior Behavior = EResetToCachedTargetBehavior::ResetImmediately);
+
+	/**
+	 * Calls ResetBodyModifierToCachedTarget for each of the body modifiers
+	 */
+	UFUNCTION(BlueprintCallable, Category = PhysicsControl)
+	void ResetBodyModifiersToCachedTargets(
+		const TArray<FName>&               Names,
+		const EResetToCachedTargetBehavior Behavior = EResetToCachedTargetBehavior::ResetImmediately);
+
+	/**
+	 * Calls ResetBodyModifierToCachedTarget for each of the body modifiers in the set
+	 */
+	UFUNCTION(BlueprintCallable, Category = PhysicsControl)
+	void ResetBodyModifiersInSetToCachedTargets(
+		const FName                        SetName,
+		const EResetToCachedTargetBehavior Behavior = EResetToCachedTargetBehavior::ResetImmediately);
+
 public:
 
 	/**
diff --git a/Engine/Plugins/Experimental/PhysicsControl/Source/PhysicsControl/Public/PhysicsControlData.h b/Engine/Plugins/Experimental/PhysicsControl/Source/PhysicsControl/Public/PhysicsControlData.h
--- a/Engine/Plugins/Experimental/PhysicsControl/Source/PhysicsControl/Public/PhysicsControlData.h
+++ b/Engine/Plugins/Experimental/PhysicsControl/Source/PhysicsControl/Public/PhysicsControlData.h
@@ -35,11 +35,11 @@ struct PHYSICSCONTROL_API FPhysicsControlData
 	GENERATED_BODY()
 
 	FPhysicsControlData()
-		: LinearStrength(1.0f)
+		: LinearStrength(0.0f)
 		, LinearDampingRatio(1.0f)
 		, LinearExtraDamping(0.0f)
 		, MaxForce(0.0f)
-		, AngularStrength(1.0f)
+		, AngularStrength(0.0f)
 		, AngularDampingRatio(1.0f)
 		, AngularExtraDamping(0.0f)
 		, MaxTorque(0.0f)
@@ -108,9 +108,11 @@ struct PHYSICSCONTROL_API FPhysicsControlMultiplier
 
 	FPhysicsControlMultiplier()
 		: LinearStrengthMultiplier(1.0)
+		, LinearDampingRatioMultiplier(1.0)
 		, LinearExtraDampingMultiplier(1.0)
 		, MaxForceMultiplier(1.0)
 		, AngularStrengthMultiplier(1.0)
+		, AngularDampingRatioMultiplier(1.0)
 		, AngularExtraDampingMultiplier(1.0)
 		, MaxTorqueMultiplier(1.0)
 	{
@@ -120,6 +122,10 @@ struct PHYSICSCONTROL_API FPhysicsControlMultiplier
 	UPROPERTY(EditAnywhere, BlueprintReadWrite, Category = PhysicsControl, meta = (ClampMin = "0.0"))
 	FVector LinearStrengthMultiplier;
 
+	// Per-direction multiplier on the linear damping ratio.
+	UPROPERTY(EditAnywhere, BlueprintReadWrite, Category = PhysicsControl, meta = (ClampMin = "0.0"))
+	FVector LinearDampingRatioMultiplier;
+
 	// Per-direction multiplier on the linear extra damping.
 	UPROPERTY(EditAnywhere, BlueprintReadWrite, Category = PhysicsControl, meta = (ClampMin = "0.0"))
 	FVector LinearExtraDampingMultiplier;
@@ -132,6 +138,10 @@ struct PHYSICSCONTROL_API FPhysicsControlMultiplier
 	UPROPERTY(EditAnywhere, BlueprintReadWrite, Category = PhysicsControl, meta = (ClampMin = "0.0"))
 	float AngularStrengthMultiplier;
 
+	// Multiplier on the angular damping ratio.
+	UPROPERTY(EditAnywhere, BlueprintReadWrite, Category = PhysicsControl, meta = (ClampMin = "0.0"))
+	float AngularDampingRatioMultiplier;
+
 	// Multiplier on the angular extra damping.
 	UPROPERTY(EditAnywhere, BlueprintReadWrite, Category = PhysicsControl, meta = (ClampMin = "0.0"))
 	float AngularExtraDampingMultiplier;
