Additional functions for getting the cached bone data, and for setting targets via arrays. Also some renaming to things like "CachedBoneTransform" rather than "CachedTargetBoneTransform"

diff --git a/Engine/Plugins/Experimental/PhysicsControl/Source/PhysicsControl/Private/PhysicsControlComponent.cpp b/Engine/Plugins/Experimental/PhysicsControl/Source/PhysicsControl/Private/PhysicsControlComponent.cpp
--- a/Engine/Plugins/Experimental/PhysicsControl/Source/PhysicsControl/Private/PhysicsControlComponent.cpp
+++ b/Engine/Plugins/Experimental/PhysicsControl/Source/PhysicsControl/Private/PhysicsControlComponent.cpp
@@ -1194,6 +1194,74 @@ void UPhysicsControlComponent::SetControlTargetOrientationsInSet(
 		Orientation, AngularVelocityDeltaTime, bEnableControl, bApplyControlPointToTarget);
 }
 
+//======================================================================================================================
+bool UPhysicsControlComponent::SetControlTargetPositionsFromArray(
+	const TArray<FName>&   Names,
+	const TArray<FVector>& Positions,
+	const float            VelocityDeltaTime,
+	const bool             bEnableControl,
+	const bool             bApplyControlPointToTarget)
+{
+	int32 NumControlNames = Names.Num();
+	int32 NumPositions = Positions.Num();
+	if (NumControlNames != NumPositions)
+	{
+		return false;
+	}
+	for (int32 Index = 0 ; Index != NumControlNames ; ++Index)
+	{
+		SetControlTargetPosition(
+			Names[Index], Positions[Index], VelocityDeltaTime, bEnableControl, bApplyControlPointToTarget);
+	}
+	return true;
+}
+
+//======================================================================================================================
+bool UPhysicsControlComponent::SetControlTargetOrientationsFromArray(
+	const TArray<FName>&    Names,
+	const TArray<FRotator>& Orientations,
+	const float             VelocityDeltaTime,
+	const bool              bEnableControl,
+	const bool              bApplyControlPointToTarget)
+{
+	int32 NumControlNames = Names.Num();
+	int32 NumOrientations = Orientations.Num();
+	if (NumControlNames != NumOrientations)
+	{
+		return false;
+	}
+	for (int32 Index = 0 ; Index != NumControlNames ; ++Index)
+	{
+		SetControlTargetOrientation(
+			Names[Index], Orientations[Index], VelocityDeltaTime, bEnableControl, bApplyControlPointToTarget);
+	}
+	return true;
+}
+
+//======================================================================================================================
+bool UPhysicsControlComponent::SetControlTargetPositionsAndOrientationsFromArray(
+	const TArray<FName>&    Names,
+	const TArray<FVector>&  Positions,
+	const TArray<FRotator>& Orientations,
+	const float             VelocityDeltaTime,
+	const bool              bEnableControl,
+	const bool              bApplyControlPointToTarget)
+{
+	int32 NumControlNames = Names.Num();
+	int32 NumPositions = Positions.Num();
+	int32 NumOrientations = Orientations.Num();
+	if (NumControlNames != NumPositions || NumControlNames != NumOrientations)
+	{
+		return false;
+	}
+	for (int32 Index = 0; Index != NumControlNames; ++Index)
+	{
+		SetControlTargetPositionAndOrientation(
+			Names[Index], Positions[Index], Orientations[Index], VelocityDeltaTime, bEnableControl, bApplyControlPointToTarget);
+	}
+	return true;
+}
+
 //======================================================================================================================
 bool UPhysicsControlComponent::SetControlTargetPoses(
 	const FName    Name,
@@ -1894,7 +1962,7 @@ TArray<FName> UPhysicsControlComponent::GetSetsContainingBodyModifier(const FNam
 }
 
 //======================================================================================================================
-TArray<FTransform> UPhysicsControlComponent::GetCachedTargetBoneTransforms(
+TArray<FTransform> UPhysicsControlComponent::GetCachedBoneTransforms(
 	const USkeletalMeshComponent* SkeletalMeshComponent, 
 	const TArray<FName>&          BoneNames)
 {
@@ -1914,12 +1982,57 @@ TArray<FTransform> UPhysicsControlComponent::GetCachedTargetBoneTransforms(
 			Result.Add(FTransform::Identity);
 		}
 	}
+	return Result;
+}
+
+//======================================================================================================================
+TArray<FVector> UPhysicsControlComponent::GetCachedBonePositions(
+	const USkeletalMeshComponent* SkeletalMeshComponent,
+	const TArray<FName>&          BoneNames)
+{
+	TArray<FVector> Result;
+	Result.Reserve(BoneNames.Num());
+	FCachedSkeletalMeshData::FBoneData BoneData;
+
+	for (const FName& BoneName : BoneNames)
+	{
+		if (Implementation->GetBoneData(BoneData, SkeletalMeshComponent, BoneName))
+		{
+			Result.Add(BoneData.Position);
+		}
+		else
+		{
+			Result.Add(FVector::ZeroVector);
+		}
+	}
+	return Result;
+}
 
+//======================================================================================================================
+TArray<FRotator> UPhysicsControlComponent::GetCachedBoneOrientations(
+	const USkeletalMeshComponent* SkeletalMeshComponent,
+	const TArray<FName>&          BoneNames)
+{
+	TArray<FRotator> Result;
+	Result.Reserve(BoneNames.Num());
+	FCachedSkeletalMeshData::FBoneData BoneData;
+
+	for (const FName& BoneName : BoneNames)
+	{
+		if (Implementation->GetBoneData(BoneData, SkeletalMeshComponent, BoneName))
+		{
+			Result.Add(BoneData.Orientation.Rotator());
+		}
+		else
+		{
+			Result.Add(FRotator::ZeroRotator);
+		}
+	}
 	return Result;
 }
 
 //======================================================================================================================
-TArray<FVector> UPhysicsControlComponent::GetCachedTargetBoneVelocities(
+TArray<FVector> UPhysicsControlComponent::GetCachedBoneVelocities(
 	const USkeletalMeshComponent* SkeletalMeshComponent,
 	const TArray<FName>&          BoneNames)
 {
@@ -1943,7 +2056,7 @@ TArray<FVector> UPhysicsControlComponent::GetCachedTargetBoneVelocities(
 }
 
 //======================================================================================================================
-TArray<FVector> UPhysicsControlComponent::GetCachedTargetBoneAngularVelocities(
+TArray<FVector> UPhysicsControlComponent::GetCachedBoneAngularVelocities(
 	const USkeletalMeshComponent* SkeletalMeshComponent,
 	const TArray<FName>&          BoneNames)
 {
@@ -1967,7 +2080,7 @@ TArray<FVector> UPhysicsControlComponent::GetCachedTargetBoneAngularVelocities(
 }
 
 //======================================================================================================================
-FTransform UPhysicsControlComponent::GetCachedTargetBoneTransform(
+FTransform UPhysicsControlComponent::GetCachedBoneTransform(
 	const USkeletalMeshComponent* SkeletalMeshComponent,
 	const FName                   BoneName)
 {
@@ -1980,7 +2093,33 @@ FTransform UPhysicsControlComponent::GetCachedTargetBoneTransform(
 }
 
 //======================================================================================================================
-FVector UPhysicsControlComponent::GetCachedTargetBoneVelocity(
+FVector UPhysicsControlComponent::GetCachedBonePosition(
+	const USkeletalMeshComponent* SkeletalMeshComponent,
+	const FName                   BoneName)
+{
+	FCachedSkeletalMeshData::FBoneData BoneData;
+	if (Implementation->GetBoneData(BoneData, SkeletalMeshComponent, BoneName))
+	{
+		return BoneData.Position;
+	}
+	return FVector::ZeroVector;
+}
+
+//======================================================================================================================
+FRotator UPhysicsControlComponent::GetCachedBoneOrientation(
+	const USkeletalMeshComponent* SkeletalMeshComponent,
+	const FName                   BoneName)
+{
+	FCachedSkeletalMeshData::FBoneData BoneData;
+	if (Implementation->GetBoneData(BoneData, SkeletalMeshComponent, BoneName))
+	{
+		return BoneData.Orientation.Rotator();
+	}
+	return FRotator::ZeroRotator;
+}
+
+//======================================================================================================================
+FVector UPhysicsControlComponent::GetCachedBoneVelocity(
 	const USkeletalMeshComponent* SkeletalMeshComponent,
 	const FName                   BoneName)
 {
@@ -1993,7 +2132,7 @@ FVector UPhysicsControlComponent::GetCachedTargetBoneVelocity(
 }
 
 //======================================================================================================================
-FVector UPhysicsControlComponent::GetCachedTargetBoneAngularVelocity(
+FVector UPhysicsControlComponent::GetCachedBoneAngularVelocity(
 	const USkeletalMeshComponent* SkeletalMeshComponent,
 	const FName                   BoneName)
 {
@@ -2006,7 +2145,7 @@ FVector UPhysicsControlComponent::GetCachedTargetBoneAngularVelocity(
 }
 
 //======================================================================================================================
-bool UPhysicsControlComponent::SetCachedTargetBoneData(
+bool UPhysicsControlComponent::SetCachedBoneData(
 	const USkeletalMeshComponent* SkeletalMeshComponent,
 	const FName                   BoneName,
 	const FTransform&             TM,
@@ -2026,7 +2165,7 @@ bool UPhysicsControlComponent::SetCachedTargetBoneData(
 }
 
 //======================================================================================================================
-bool UPhysicsControlComponent::ResetBodyModifierToCachedTarget(
+bool UPhysicsControlComponent::ResetBodyModifierToCachedBoneTransform(
 	const FName                        Name,
 	const EResetToCachedTargetBehavior Behavior)
 {
@@ -2048,22 +2187,22 @@ bool UPhysicsControlComponent::ResetBodyModifierToCachedTarget(
 }
 
 //======================================================================================================================
-void UPhysicsControlComponent::ResetBodyModifiersToCachedTargets(
+void UPhysicsControlComponent::ResetBodyModifiersToCachedBoneTransforms(
 	const TArray<FName>&               Names,
 	const EResetToCachedTargetBehavior Behavior)
 {
 	for (FName Name : Names)
 	{
-		ResetBodyModifierToCachedTarget(Name, Behavior);
+		ResetBodyModifierToCachedBoneTransform(Name, Behavior);
 	}
 }
 
 //======================================================================================================================
-void UPhysicsControlComponent::ResetBodyModifiersInSetToCachedTargets(
+void UPhysicsControlComponent::ResetBodyModifiersInSetToCachedBoneTransforms(
 	const FName                        SetName,
 	const EResetToCachedTargetBehavior Behavior)
 {
-	ResetBodyModifiersToCachedTargets(GetBodyModifierNamesInSet(SetName), Behavior);
+	ResetBodyModifiersToCachedBoneTransforms(GetBodyModifierNamesInSet(SetName), Behavior);
 }
 
 
diff --git a/Engine/Plugins/Experimental/PhysicsControl/Source/PhysicsControl/Public/PhysicsControlComponent.h b/Engine/Plugins/Experimental/PhysicsControl/Source/PhysicsControl/Public/PhysicsControlComponent.h
--- a/Engine/Plugins/Experimental/PhysicsControl/Source/PhysicsControl/Public/PhysicsControlComponent.h
+++ b/Engine/Plugins/Experimental/PhysicsControl/Source/PhysicsControl/Public/PhysicsControlComponent.h
@@ -666,6 +666,46 @@ public:
 		const bool     bEnableControl = true,
 		const bool     bApplyControlPointToTarget = false);
 
+	/**
+	 * Calls SetControlTargetPosition for each element of the control names and positions. These array should match
+	 * in size.
+	 * @return true if the control/position arrays match, false if they don't.
+	 */
+	UFUNCTION(BlueprintCallable, Category = PhysicsControl)
+	bool SetControlTargetPositionsFromArray(
+		const TArray<FName>&   Names,
+		const TArray<FVector>& Positions, 
+		const float            VelocityDeltaTime, 
+		const bool             bEnableControl = true,
+		const bool             bApplyControlPointToTarget = false);
+
+	/**
+	 * Calls SetControlTargetPosition for each element of the control names and positions. These array should match
+	 * in size.
+	 * @return true if the control/position arrays match, false if they don't.
+	 */
+	UFUNCTION(BlueprintCallable, Category = PhysicsControl)
+	bool SetControlTargetOrientationsFromArray(
+		const TArray<FName>&    Names,
+		const TArray<FRotator>& Orientations, 
+		const float             VelocityDeltaTime, 
+		const bool              bEnableControl = true,
+		const bool              bApplyControlPointToTarget = false);
+
+	/**
+	 * Calls SetControlTargetPositionAndOrientation for each element of the control names, positions and orientations. 
+	 * These array should match in size.
+	 * @return true if the control/position/orientation arrays match, false if they don't.
+	 */
+	UFUNCTION(BlueprintCallable, Category = PhysicsControl)
+	bool SetControlTargetPositionsAndOrientationsFromArray(
+		const TArray<FName>&    Names,
+		const TArray<FVector>&  Positions,
+		const TArray<FRotator>& Orientations,
+		const float             VelocityDeltaTime, 
+		const bool              bEnableControl = true,
+		const bool              bApplyControlPointToTarget = false);
+
 	/**
 	 * Calculates and sets an existing control target. This takes the "virtual" position/orientation of the parent 
 	 * and child and calculates the relative control. Note that this will set bApplyControlPointToTarget to true.
@@ -1317,7 +1357,31 @@ public:
 	 * own targets.
 	 */
 	UFUNCTION(BlueprintCallable, Category = PhysicsControl)
-	TArray<FTransform> GetCachedTargetBoneTransforms(
+	TArray<FTransform> GetCachedBoneTransforms(
+		const USkeletalMeshComponent* SkeletalMeshComponent,
+		const TArray<FName>&          BoneNames);
+
+	/**
+	 * Gets the positions of the requested bones that will be used as targets (in world space). Targets for bones 
+	 * that are not found will be set to zero. Note that these targets will have been calculated and cached 
+	 * at the start of the Physics Control Component, so if using the built in tick, may be too old to be useful. 
+	 * If you manually update the component then you can access these target transforms prior to applying your 
+	 * own targets.
+	 */
+	UFUNCTION(BlueprintCallable, Category = PhysicsControl)
+	TArray<FVector> GetCachedBonePositions(
+		const USkeletalMeshComponent* SkeletalMeshComponent,
+		const TArray<FName>&          BoneNames);
+
+	/**
+	 * Gets the orientations of the requested bones that will be used as targets (in world space). Targets for bones 
+	 * that are not found will be set to identity. Note that these targets will have been calculated and cached 
+	 * at the start of the Physics Control Component, so if using the built in tick, may be too old to be useful. 
+	 * If you manually update the component then you can access these target transforms prior to applying your 
+	 * own targets.
+	 */
+	UFUNCTION(BlueprintCallable, Category = PhysicsControl)
+	TArray<FRotator> GetCachedBoneOrientations(
 		const USkeletalMeshComponent* SkeletalMeshComponent,
 		const TArray<FName>&          BoneNames);
 
@@ -1329,7 +1393,7 @@ public:
 	 * velocities prior to applying your own targets.
 	 */
 	UFUNCTION(BlueprintCallable, Category = PhysicsControl)
-	TArray<FVector> GetCachedTargetBoneVelocities(
+	TArray<FVector> GetCachedBoneVelocities(
 		const USkeletalMeshComponent* SkeletalMeshComponent,
 		const TArray<FName>&          BoneNames);
 
@@ -1341,7 +1405,7 @@ public:
 	 * velocities prior to applying your own targets.
 	 */
 	UFUNCTION(BlueprintCallable, Category = PhysicsControl)
-	TArray<FVector> GetCachedTargetBoneAngularVelocities(
+	TArray<FVector> GetCachedBoneAngularVelocities(
 		const USkeletalMeshComponent* SkeletalMeshComponent,
 		const TArray<FName>&          BoneNames);
 
@@ -1353,7 +1417,31 @@ public:
 	 * own targets.
 	 */
 	UFUNCTION(BlueprintCallable, Category = PhysicsControl)
-	FTransform GetCachedTargetBoneTransform(
+	FTransform GetCachedBoneTransform(
+		const USkeletalMeshComponent* SkeletalMeshComponent,
+		const FName                   BoneName);
+
+	/**
+	 * Gets the position of the requested bone that will be used as a target (in world space). Targets for bones
+	 * that are not found will be set to zero. Note that these targets will have been calculated and cached
+	 * at the start of the Physics Control Component, so if using the built in tick, may be too old to be useful.
+	 * If you manually update the component then you can access these target transforms prior to applying your
+	 * own targets.
+	 */
+	UFUNCTION(BlueprintCallable, Category = PhysicsControl)
+	FVector GetCachedBonePosition(
+		const USkeletalMeshComponent* SkeletalMeshComponent,
+		const FName                   BoneName);
+
+	/**
+	 * Gets the orientation of the requested bone that will be used as a target (in world space). Targets for bones
+	 * that are not found will be set to Identity. Note that these targets will have been calculated and cached
+	 * at the start of the Physics Control Component, so if using the built in tick, may be too old to be useful.
+	 * If you manually update the component then you can access these target transforms prior to applying your
+	 * own targets.
+	 */
+	UFUNCTION(BlueprintCallable, Category = PhysicsControl)
+	FRotator GetCachedBoneOrientation(
 		const USkeletalMeshComponent* SkeletalMeshComponent,
 		const FName                   BoneName);
 
@@ -1365,7 +1453,7 @@ public:
 	 * velocities prior to applying your own targets.
 	 */
 	UFUNCTION(BlueprintCallable, Category = PhysicsControl)
-	FVector GetCachedTargetBoneVelocity(
+	FVector GetCachedBoneVelocity(
 		const USkeletalMeshComponent* SkeletalMeshComponent,
 		const FName                   BoneName);
 
@@ -1377,7 +1465,7 @@ public:
 	 * velocities prior to applying your own targets.
 	 */
 	UFUNCTION(BlueprintCallable, Category = PhysicsControl)
-	FVector GetCachedTargetBoneAngularVelocity(
+	FVector GetCachedBoneAngularVelocity(
 		const USkeletalMeshComponent* SkeletalMeshComponent,
 		const FName                   BoneName);
 
@@ -1390,7 +1478,7 @@ public:
 	 * @return true if successful, and false if no cached target can be found for the bone.
 	 */
 	UFUNCTION(BlueprintCallable, Category = PhysicsControl)
-	bool SetCachedTargetBoneData(
+	bool SetCachedBoneData(
 		const USkeletalMeshComponent* SkeletalMeshComponent,
 		const FName                   BoneName, 
 		const FTransform&             TM,
@@ -1407,7 +1495,7 @@ public:
 	 * @return true if the body modifier is found (even if no cached target is found), and false otherwise.
 	 */
 	UFUNCTION(BlueprintCallable, Category = PhysicsControl)
-	bool ResetBodyModifierToCachedTarget(
+	bool ResetBodyModifierToCachedBoneTransform(
 		const FName                        Name,
 		const EResetToCachedTargetBehavior Behavior = EResetToCachedTargetBehavior::ResetImmediately);
 
@@ -1415,7 +1503,7 @@ public:
 	 * Calls ResetBodyModifierToCachedTarget for each of the body modifiers
 	 */
 	UFUNCTION(BlueprintCallable, Category = PhysicsControl)
-	void ResetBodyModifiersToCachedTargets(
+	void ResetBodyModifiersToCachedBoneTransforms(
 		const TArray<FName>&               Names,
 		const EResetToCachedTargetBehavior Behavior = EResetToCachedTargetBehavior::ResetImmediately);
 
@@ -1423,7 +1511,7 @@ public:
 	 * Calls ResetBodyModifierToCachedTarget for each of the body modifiers in the set
 	 */
 	UFUNCTION(BlueprintCallable, Category = PhysicsControl)
-	void ResetBodyModifiersInSetToCachedTargets(
+	void ResetBodyModifiersInSetToCachedBoneTransforms(
 		const FName                        SetName,
 		const EResetToCachedTargetBehavior Behavior = EResetToCachedTargetBehavior::ResetImmediately);
 
