Fix for updating kinematic shapes with Physics Control Component

diff --git a/Engine/Plugins/Experimental/PhysicsControl/Source/PhysicsControl/Private/PhysicsControlComponentImpl.cpp b/Engine/Plugins/Experimental/PhysicsControl/Source/PhysicsControl/Private/PhysicsControlComponentImpl.cpp
--- a/Engine/Plugins/Experimental/PhysicsControl/Source/PhysicsControl/Private/PhysicsControlComponentImpl.cpp
+++ b/Engine/Plugins/Experimental/PhysicsControl/Source/PhysicsControl/Private/PhysicsControlComponentImpl.cpp
@@ -378,7 +378,6 @@ void FPhysicsControlComponentImpl::AddSkeletalMeshReferenceForModifier(
 	Data.bOriginalUpdateMeshWhenKinematic = InSkeletalMeshComponent->bUpdateMeshWhenKinematic;
 	Data.OriginalKinematicBonesUpdateType = InSkeletalMeshComponent->KinematicBonesUpdateType;
 	InSkeletalMeshComponent->bUpdateMeshWhenKinematic = true;
-	InSkeletalMeshComponent->KinematicBonesUpdateType = EKinematicBonesUpdateToPhysics::SkipAllBones;
 	// By default, kinematic bodies will have their blend weight set to zero. This is a problem for us since:
 	// 1. We expect there will be lots of cases where only part of the character is dynamic, and other 
 	//    parts are kinematic
