Support kinematics with async physics by adding a UpdateKinematicFromSimulation flag to BodyInstance. Defaults to false, so no change in default behaviour, unless using PhysicsControlComponent

diff --git a/Engine/Plugins/Experimental/PhysicsControl/Source/PhysicsControl/Private/PhysicsControlComponent.cpp b/Engine/Plugins/Experimental/PhysicsControl/Source/PhysicsControl/Private/PhysicsControlComponent.cpp
--- a/Engine/Plugins/Experimental/PhysicsControl/Source/PhysicsControl/Private/PhysicsControlComponent.cpp
+++ b/Engine/Plugins/Experimental/PhysicsControl/Source/PhysicsControl/Private/PhysicsControlComponent.cpp
@@ -166,6 +166,8 @@ void UPhysicsControlComponent::UpdateControls(float DeltaTime)
 		// expectation that zero blend weight means to disable physics.
 		BodyInstance->PhysicsBlendWeight = BodyModifier.PhysicsBlendWeight;
 
+		BodyInstance->SetUpdateKinematicFromSimulation(BodyModifier.bUpdateKinematicFromSimulation);
+
 		if (BodyModifier.bResetToCachedTarget)
 		{
 			BodyModifier.bResetToCachedTarget = false;
@@ -1470,12 +1472,13 @@ FName UPhysicsControlComponent::CreateBodyModifier(
 	const ECollisionEnabled::Type CollisionType,
 	const float                   GravityMultiplier,
 	const float                   PhysicsBlendWeight,
-	const bool                    bUseSkeletalAnimation)
+	const bool                    bUseSkeletalAnimation,
+	const bool                    bUpdateKinematicFromSimulation)
 {
 	const FName Name = Implementation->GetUniqueBodyModifierName(BoneName);
 	if (CreateNamedBodyModifier(
 		Name, MeshComponent, BoneName, Set, MovementType, CollisionType, 
-		GravityMultiplier, PhysicsBlendWeight, bUseSkeletalAnimation))
+		GravityMultiplier, PhysicsBlendWeight, bUseSkeletalAnimation, bUpdateKinematicFromSimulation))
 	{
 		return Name;
 	}
@@ -1492,7 +1495,8 @@ bool UPhysicsControlComponent::CreateNamedBodyModifier(
 	const ECollisionEnabled::Type CollisionType,
 	const float                   GravityMultiplier,
 	const float                   PhysicsBlendWeight,
-	const bool                    bUseSkeletalAnimation)
+	const bool                    bUseSkeletalAnimation,
+	const bool                    bUpdateKinematicFromSimulation)
 {
 	if (Implementation->FindBodyModifier(Name))
 	{
@@ -1509,7 +1513,7 @@ bool UPhysicsControlComponent::CreateNamedBodyModifier(
 	FPhysicsBodyModifier& Modifier = Implementation->PhysicsBodyModifiers.Add(
 		Name, FPhysicsBodyModifier(
 			MeshComponent, BoneName, MovementType, CollisionType, 
-			GravityMultiplier, PhysicsBlendWeight, bUseSkeletalAnimation));
+			GravityMultiplier, PhysicsBlendWeight, bUseSkeletalAnimation, bUpdateKinematicFromSimulation));
 
 	USkeletalMeshComponent* SkeletalMeshComponent = Cast<USkeletalMeshComponent>(MeshComponent);
 	if (SkeletalMeshComponent)
@@ -1533,7 +1537,8 @@ TArray<FName> UPhysicsControlComponent::CreateBodyModifiersFromSkeletalMeshBelow
 	const ECollisionEnabled::Type CollisionType,
 	const float                   GravityMultiplier,
 	const float                   PhysicsBlendWeight,
-	const bool                    bUseSkeletalAnimation)
+	const bool                    bUseSkeletalAnimation,
+	const bool                    bUpdateKinematicFromSimulation)
 {
 	TArray<FName> Result;
 	UPhysicsAsset* PhysicsAsset = SkeletalMeshComponent ? SkeletalMeshComponent->GetPhysicsAsset() : nullptr;
@@ -1546,15 +1551,15 @@ TArray<FName> UPhysicsControlComponent::CreateBodyModifiersFromSkeletalMeshBelow
 		BoneName, bIncludeSelf, /*bSkipCustomType=*/false,
 		[
 			this, PhysicsAsset, SkeletalMeshComponent, Set, MovementType, CollisionType,
-			GravityMultiplier, PhysicsBlendWeight, bUseSkeletalAnimation, &Result
+			GravityMultiplier, PhysicsBlendWeight, bUseSkeletalAnimation, bUpdateKinematicFromSimulation, &Result
 		](const FBodyInstance* BI)
 		{
 			if (USkeletalBodySetup* BodySetup = Cast<USkeletalBodySetup>(BI->BodySetup.Get()))
 			{
 				const FName BoneName = PhysicsAsset->SkeletalBodySetups[BI->InstanceBodyIndex]->BoneName;
 				const FName BodyModifierName = CreateBodyModifier(
-					SkeletalMeshComponent, BoneName, Set, MovementType, 
-					CollisionType, GravityMultiplier, PhysicsBlendWeight, bUseSkeletalAnimation);
+					SkeletalMeshComponent, BoneName, Set, MovementType, CollisionType, 
+					GravityMultiplier, PhysicsBlendWeight, bUseSkeletalAnimation, bUpdateKinematicFromSimulation);
 				Result.Add(BodyModifierName);
 			}
 		});
@@ -1570,7 +1575,8 @@ TMap<FName, FPhysicsControlNames> UPhysicsControlComponent::CreateBodyModifiersF
 	const ECollisionEnabled::Type                CollisionType,
 	const float                                  GravityMultiplier,
 	const float                                  PhysicsBlendWeight,
-	const bool                                   bUseSkeletalAnimation)
+	const bool                                   bUseSkeletalAnimation,
+	const bool                                   bUpdateKinematicFromSimulation)
 {
 	TMap<FName, FPhysicsControlNames> Result;
 	Result.Reserve(LimbBones.Num());
@@ -1595,8 +1601,8 @@ TMap<FName, FPhysicsControlNames> UPhysicsControlComponent::CreateBodyModifiersF
 		for (const FName BoneName : BonesInLimb.BoneNames)
 		{
 			const FName BodyModifierName = CreateBodyModifier(
-				BonesInLimb.SkeletalMeshComponent, BoneName, LimbName, MovementType, 
-				CollisionType, GravityMultiplier, PhysicsBlendWeight, bUseSkeletalAnimation);
+				BonesInLimb.SkeletalMeshComponent, BoneName, LimbName, MovementType, CollisionType, 
+				GravityMultiplier, PhysicsBlendWeight, bUseSkeletalAnimation, bUpdateKinematicFromSimulation);
 			if (!BodyModifierName.IsNone())
 			{
 				LimbResult.Names.Add(BodyModifierName);
@@ -1820,6 +1826,39 @@ void UPhysicsControlComponent::SetBodyModifiersInSetUseSkeletalAnimation(
 	SetBodyModifiersUseSkeletalAnimation(GetBodyModifierNamesInSet(SetName), bUseSkeletalAnimation);
 }
 
+//======================================================================================================================
+bool UPhysicsControlComponent::SetBodyModifierUpdateKinematicFromSimulation(
+	const FName Name,
+	const bool  bUpdateKinematicFromSimulation)
+{
+	FPhysicsBodyModifier* PhysicsBodyModifier = Implementation->FindBodyModifier(Name);
+	if (PhysicsBodyModifier)
+	{
+		PhysicsBodyModifier->bUpdateKinematicFromSimulation = bUpdateKinematicFromSimulation;
+		return true;
+	}
+	return false;
+}
+
+//======================================================================================================================
+void UPhysicsControlComponent::SetBodyModifiersUpdateKinematicFromSimulation(
+	const TArray<FName>& Names,
+	const bool           bUpdateKinematicFromSimulation)
+{
+	for (FName Name : Names)
+	{
+		SetBodyModifierUpdateKinematicFromSimulation(Name, bUpdateKinematicFromSimulation);
+	}
+}
+
+//======================================================================================================================
+void UPhysicsControlComponent::SetBodyModifiersInSetUpdateKinematicFromSimulation(
+	const FName SetName,
+	const bool  bUpdateKinematicFromSimulation)
+{
+	SetBodyModifiersUpdateKinematicFromSimulation(GetBodyModifierNamesInSet(SetName), bUpdateKinematicFromSimulation);
+}
+
 //======================================================================================================================
 const TArray<FName>& UPhysicsControlComponent::GetAllControlNames() const
 {
diff --git a/Engine/Plugins/Experimental/PhysicsControl/Source/PhysicsControl/Private/PhysicsControlRecord.h b/Engine/Plugins/Experimental/PhysicsControl/Source/PhysicsControl/Private/PhysicsControlRecord.h
--- a/Engine/Plugins/Experimental/PhysicsControl/Source/PhysicsControl/Private/PhysicsControlRecord.h
+++ b/Engine/Plugins/Experimental/PhysicsControl/Source/PhysicsControl/Private/PhysicsControlRecord.h
@@ -63,7 +63,8 @@ struct FPhysicsBodyModifier
 		ECollisionEnabled::Type    InCollisionType,
 		float                      InGravityMultiplier,
 		float                      InPhysicsBlendWeight,
-		bool                       InUseSkeletalAnimation)
+		bool                       InUseSkeletalAnimation,
+		bool                       InUpdateKinematicFromSimulation)
 		: MeshComponent(InMeshComponent)
 		, BoneName(InBoneName)
 		, MovementType(InMovementType)
@@ -73,6 +74,7 @@ struct FPhysicsBodyModifier
 		, KinematicTargetPosition(FVector::ZeroVector)
 		, KinematicTargetOrientation(FQuat::Identity)
 		, bUseSkeletalAnimation(InUseSkeletalAnimation)
+		, bUpdateKinematicFromSimulation(InUpdateKinematicFromSimulation)
 		, bResetToCachedTarget(false)
 	{}
 
@@ -114,6 +116,13 @@ struct FPhysicsBodyModifier
 	/** If true then the target will be applied on top of the skeletal animation (if there is any) */
 	uint8 bUseSkeletalAnimation:1;
 
+	/** 
+	 * If true then the associated actor's transform will be updated from the simulation when it is 
+	 * kinematic. This is most likely useful when using async physics in order to prevent different 
+	 * parts of the skeleton from being torn apart. 
+	 */
+	uint8 bUpdateKinematicFromSimulation:1;
+
 	/** 
 	 * If true then the body will be set to the transform/velocity stored in any cached target (if that
 	 * exists), and then this flag will be cleared.
diff --git a/Engine/Plugins/Experimental/PhysicsControl/Source/PhysicsControl/Public/PhysicsControlComponent.h b/Engine/Plugins/Experimental/PhysicsControl/Source/PhysicsControl/Public/PhysicsControlComponent.h
--- a/Engine/Plugins/Experimental/PhysicsControl/Source/PhysicsControl/Public/PhysicsControlComponent.h
+++ b/Engine/Plugins/Experimental/PhysicsControl/Source/PhysicsControl/Public/PhysicsControlComponent.h
@@ -906,8 +906,12 @@ public:
 	 * @param PhysicsBlendWeight The blend weight between the body transform coming from
 	 *        animation and that coming from simulation.
 	 * @param bUseSkeletalAnimation Whether the kinematic target is specified in the frame of the
-	 *                              skeletal animation, rather than world space. Only relevant if the
-	 *                              body is part of a skeletal mesh.
+	 *        skeletal animation, rather than world space. Only relevant if the body is part of 
+	 *        a skeletal mesh.
+	 * @param bUpdateKinematicFromSimulation Whether the body should be updated from the simulation when
+	 *        it is kinematic, or whether it should track the kinematic target directly. This will be most
+	 *        likely useful when using async physics, in order to make kinematic parts behave the same as 
+	 *        dynamic ones.
 	 */
 	UFUNCTION(BlueprintCallable, Category = PhysicsControl)
 	FName CreateBodyModifier(
@@ -918,7 +922,8 @@ public:
 		const ECollisionEnabled::Type CollisionType = ECollisionEnabled::QueryAndPhysics,
 		const float                   GravityMultiplier = 1.0f,
 		const float                   PhysicsBlendWeight = 1.0f,
-		const bool                    bUseSkeletalAnimation = true);
+		const bool                    bUseSkeletalAnimation = true,
+		const bool                    bUpdateKinematicFromSimulation = true);
 
 	/**
 	 * Creates a new body modifier for mesh components
@@ -934,8 +939,12 @@ public:
 	 * @param PhysicsBlendWeight The blend weight between the body transform coming from
 	 *        animation and that coming from simulation.
 	 * @param bUseSkeletalAnimation Whether the kinematic target is specified in the frame of the
-	 *                              skeletal animation, rather than world space. Only relevant if the
-	 *                              body is part of a skeletal mesh.
+	 *        skeletal animation, rather than world space. Only relevant if the body is part of 
+	 *        a skeletal mesh.
+	 * @param bUpdateKinematicFromSimulation Whether the body should be updated from the simulation when
+	 *        it is kinematic, or whether it should track the kinematic target directly. This will be most
+	 *        likely useful when using async physics, in order to make kinematic parts behave the same as
+	 *        dynamic ones.
 	 */
 	UFUNCTION(BlueprintCallable, Category = PhysicsControl)
 	bool CreateNamedBodyModifier(
@@ -947,7 +956,8 @@ public:
 		const ECollisionEnabled::Type CollisionType = ECollisionEnabled::QueryAndPhysics,
 		const float                   GravityMultiplier = 1.0f,
 		const float                   PhysicsBlendWeight = 1.0f,
-		const bool                    bUseSkeletalAnimation = true);
+		const bool                    bUseSkeletalAnimation = true,
+		const bool                    bUpdateKinematicFromSimulation = true);
 
 	/**
 	 * Creates new body modifiers for skeletal mesh components
@@ -963,8 +973,12 @@ public:
 	 * @param PhysicsBlendWeight The blend weight between the body transform coming from
 	 *        animation and that coming from simulation.
 	 * @param bUseSkeletalAnimation Whether the kinematic target is specified in the frame of the
-	 *                              skeletal animation, rather than world space. Only relevant if the
-	 *                              body is part of a skeletal mesh.
+	 *        skeletal animation, rather than world space. Only relevant if the body is part of
+	 *        a skeletal mesh.
+	 * @param bUpdateKinematicFromSimulation Whether the body should be updated from the simulation when
+	 *        it is kinematic, or whether it should track the kinematic target directly. This will be most
+	 *        likely useful when using async physics, in order to make kinematic parts behave the same as
+	 *        dynamic ones.
 	 */
 	UFUNCTION(BlueprintCallable, Category = PhysicsControl)
 	TArray<FName> CreateBodyModifiersFromSkeletalMeshBelow(
@@ -976,7 +990,8 @@ public:
 		const ECollisionEnabled::Type CollisionType = ECollisionEnabled::QueryAndPhysics,
 		const float                   GravityMultiplier = 1.0f,
 		const float                   PhysicsBlendWeight = 1.0f,
-		const bool                    bUseSkeletalAnimation = true);
+		const bool                    bUseSkeletalAnimation = true,
+		const bool                    bUpdateKinematicFromSimulation = true);
 
 
 	/**
@@ -991,9 +1006,13 @@ public:
 	 * @param PhysicsBlendWeight The blend weight between the body transform coming from
 	 *        animation and that coming from simulation.
 	 * @param bUseSkeletalAnimation Whether the kinematic target is specified in the frame of the
-	 *                              skeletal animation, rather than world space. Only relevant if the
-	 *                              body is part of a skeletal mesh.
-	 * 
+	 *        skeletal animation, rather than world space. Only relevant if the body is part of
+	 *        a skeletal mesh.
+	 * @param bUpdateKinematicFromSimulation Whether the body should be updated from the simulation when
+	 *        it is kinematic, or whether it should track the kinematic target directly. This will be most
+	 *        likely useful when using async physics, in order to make kinematic parts behave the same as
+	 *        dynamic ones.
+	 *
 	 * @return A map containing the modifiers for each limb
 	 */
 	UFUNCTION(BlueprintCallable, Category = PhysicsControl)
@@ -1004,7 +1023,8 @@ public:
 		const ECollisionEnabled::Type                CollisionType = ECollisionEnabled::QueryAndPhysics,
 		const float                                  GravityMultiplier = 1.0f,
 		const float                                  PhysicsBlendWeight = 1.0f,
-		const bool                                   bUseSkeletalAnimation = true);
+		const bool                                   bUseSkeletalAnimation = true,
+		const bool                                   bUpdateKinematicFromSimulation = true);
 
 	/**
 	 * Destroys a BodyModifier
@@ -1239,6 +1259,51 @@ public:
 		const FName Set,
 		const bool  bUseSkeletalAnimation);
 
+	/**
+	 * Sets whether a body modifier should update kinematics from the simulation results
+	 *
+	 * @param Name The name of the body modifier to access. 
+	 * @param bUpdateKinematicFromSimulation Whether the body should be updated from the simulation when
+	 *        it is kinematic, or whether it should track the kinematic target directly. This will be most
+	 *        likely useful when using async physics, in order to make kinematic parts behave the same as
+	 *        dynamic ones.
+	 * @return true if the body modifier was found, false if not
+	 */
+	UFUNCTION(BlueprintCallable, Category = PhysicsControl)
+	bool SetBodyModifierUpdateKinematicFromSimulation(
+		const FName Name,
+		const bool  bUpdateKinematicFromSimulation);
+
+	/**
+	 * Sets whether body modifiers should update kinematics from the simulation results
+	 *
+	 * @param Names The names of the body modifiers to access. Note that if you have these in a FPhysicsControlNameArray
+	 *        then it can be split.
+	 * @param bUpdateKinematicFromSimulation Whether the body should be updated from the simulation when
+	 *        it is kinematic, or whether it should track the kinematic target directly. This will be most
+	 *        likely useful when using async physics, in order to make kinematic parts behave the same as
+	 *        dynamic ones.
+	 */
+	UFUNCTION(BlueprintCallable, Category = PhysicsControl)
+	void SetBodyModifiersUpdateKinematicFromSimulation(
+		const TArray<FName>& Names,
+		const bool           bUpdateKinematicFromSimulation);
+
+	/**
+	 * Sets whether body modifiers should update kinematics from the simulation results
+	 *
+	 * @param Set The set of body modifiers to modify. Standard sets will include "All" and things like
+	 *        "ArmLeft", depending on how body modifiers have been created.
+	 * @param bUpdateKinematicFromSimulation Whether the body should be updated from the simulation when
+	 *        it is kinematic, or whether it should track the kinematic target directly. This will be most
+	 *        likely useful when using async physics, in order to make kinematic parts behave the same as
+	 *        dynamic ones.
+	 */
+	UFUNCTION(BlueprintCallable, Category = PhysicsControl)
+	void SetBodyModifiersInSetUpdateKinematicFromSimulation(
+		const FName Set,
+		const bool  bUpdateKinematicFromSimulation);
+
 	/**
 	 * Creates a collections of controls and body modifiers for a character, based on the description passed in. 
 	 * This makes:
diff --git a/Engine/Source/Runtime/Engine/Classes/PhysicsEngine/BodyInstance.h b/Engine/Source/Runtime/Engine/Classes/PhysicsEngine/BodyInstance.h
--- a/Engine/Source/Runtime/Engine/Classes/PhysicsEngine/BodyInstance.h
+++ b/Engine/Source/Runtime/Engine/Classes/PhysicsEngine/BodyInstance.h
@@ -846,6 +846,8 @@ public:
 	void SetInstanceNotifyRBCollision(bool bNewNotifyCollision);
 	/** Enables/disables whether this body is affected by gravity. */
 	void SetEnableGravity(bool bGravityEnabled);
+	/** Enables/disables whether this body, when kinematic, is updated from the simulation rather than when setting the kinematic target. */
+	void SetUpdateKinematicFromSimulation(bool bUpdateKinematicFromSimulation);
 	/** Enables/disables contact modification */
 	void SetContactModification(bool bNewContactModification);
 	/** Enables/disabled smoothed edge collisions */
diff --git a/Engine/Source/Runtime/Engine/Private/PhysicsEngine/BodyInstance.cpp b/Engine/Source/Runtime/Engine/Private/PhysicsEngine/BodyInstance.cpp
--- a/Engine/Source/Runtime/Engine/Private/PhysicsEngine/BodyInstance.cpp
+++ b/Engine/Source/Runtime/Engine/Private/PhysicsEngine/BodyInstance.cpp
@@ -1184,6 +1184,7 @@ void FInitBodiesHelperBase::CreateActor_AssumesLocked(FBodyInstance* Instance, c
 	ActorParams.DebugName = Instance->CharDebugName.IsValid() ? Instance->CharDebugName->GetData() : nullptr;
 #endif
 	ActorParams.bEnableGravity = Instance->bEnableGravity;
+	ActorParams.bUpdateKinematicFromSimulation = Instance->bUpdateKinematicFromSimulation;
 	ActorParams.bQueryOnly = CollisionType == ECollisionEnabled::QueryOnly;
 	ActorParams.Scene = PhysScene;
 
@@ -3443,6 +3444,25 @@ void FBodyInstance::SetEnableGravity(bool bInGravityEnabled)
 	}
 }
 
+void FBodyInstance::SetUpdateKinematicFromSimulation(bool bInUpdateKinematicFromSimulation)
+{
+	if (bUpdateKinematicFromSimulation != bInUpdateKinematicFromSimulation)
+	{
+		bUpdateKinematicFromSimulation = bInUpdateKinematicFromSimulation;
+
+		{
+			FPhysicsCommand::ExecuteWrite(ActorHandle, [&](const FPhysicsActorHandle& Actor)
+				{
+					if (FPhysicsInterface::IsRigidBody(Actor))
+					{
+						FPhysicsInterface::SetUpdateKinematicFromSimulation_AssumesLocked(Actor, bUpdateKinematicFromSimulation);
+					}
+				});
+		}
+	}
+}
+
+
 void FBodyInstance::SetContactModification(bool bNewContactModification)
 {
 	if (bNewContactModification != bContactModification)
diff --git a/Engine/Source/Runtime/Engine/Private/PhysicsEngine/Experimental/PhysScene_Chaos.cpp b/Engine/Source/Runtime/Engine/Private/PhysicsEngine/Experimental/PhysScene_Chaos.cpp
--- a/Engine/Source/Runtime/Engine/Private/PhysicsEngine/Experimental/PhysScene_Chaos.cpp
+++ b/Engine/Source/Runtime/Engine/Private/PhysicsEngine/Experimental/PhysScene_Chaos.cpp
@@ -1806,7 +1806,8 @@ void FPhysScene_Chaos::OnSyncBodies(Chaos::FPhysicsSolverBase* Solver)
 			{
 				if (BodyInstance->OwnerComponent.IsValid())
 				{
-					if (Chaos::SyncKinematicOnGameThread == 0 && BodyInstance->IsInstanceSimulatingPhysics() == false)
+					// Skip kinematics, unless they've been flagged to update following simulation
+					if (!BodyInstance->bUpdateKinematicFromSimulation && BodyInstance->IsInstanceSimulatingPhysics() == false)
 					{
 						return;
 					}
diff --git a/Engine/Source/Runtime/Engine/Private/PhysicsEngine/ImmediatePhysics/ImmediatePhysicsChaos/ImmediatePhysicsActorHandle_Chaos.cpp b/Engine/Source/Runtime/Engine/Private/PhysicsEngine/ImmediatePhysics/ImmediatePhysicsChaos/ImmediatePhysicsActorHandle_Chaos.cpp
--- a/Engine/Source/Runtime/Engine/Private/PhysicsEngine/ImmediatePhysics/ImmediatePhysicsChaos/ImmediatePhysicsActorHandle_Chaos.cpp
+++ b/Engine/Source/Runtime/Engine/Private/PhysicsEngine/ImmediatePhysics/ImmediatePhysicsChaos/ImmediatePhysicsActorHandle_Chaos.cpp
@@ -277,6 +277,7 @@ namespace ImmediatePhysics_Chaos
 						Dynamic->SetLinearEtherDrag(BodyInstance->LinearDamping);
 						Dynamic->SetAngularEtherDrag(BodyInstance->AngularDamping);
 						Dynamic->SetGravityEnabled(BodyInstance->bEnableGravity);
+						Dynamic->SetUpdateKinematicFromSimulation(BodyInstance->bUpdateKinematicFromSimulation);
 					}
 					Dynamic->SetDisabled(true);
 				}
diff --git a/Engine/Source/Runtime/Engine/Private/PhysicsEngine/PhysicalAnimationComponent.cpp b/Engine/Source/Runtime/Engine/Private/PhysicsEngine/PhysicalAnimationComponent.cpp
--- a/Engine/Source/Runtime/Engine/Private/PhysicsEngine/PhysicalAnimationComponent.cpp
+++ b/Engine/Source/Runtime/Engine/Private/PhysicsEngine/PhysicalAnimationComponent.cpp
@@ -411,6 +411,7 @@ void UPhysicalAnimationComponent::UpdatePhysicsEngineImp()
 							FActorCreationParams Params;
 							Params.bSimulatePhysics = false;
 							Params.bQueryOnly = false;
+							Params.bUpdateKinematicFromSimulation = false;
 							Params.Scene = Scene;
 							Params.bStatic = false;
 							Params.InitialTM = TargetTM;
diff --git a/Engine/Source/Runtime/Experimental/Chaos/Private/PhysicsProxy/SingleParticlePhysicsProxy.cpp b/Engine/Source/Runtime/Experimental/Chaos/Private/PhysicsProxy/SingleParticlePhysicsProxy.cpp
--- a/Engine/Source/Runtime/Experimental/Chaos/Private/PhysicsProxy/SingleParticlePhysicsProxy.cpp
+++ b/Engine/Source/Runtime/Experimental/Chaos/Private/PhysicsProxy/SingleParticlePhysicsProxy.cpp
@@ -20,10 +20,17 @@
 namespace Chaos
 {
 
-// This is a temporary workaround to avoid GT copying position from physics results for kinematics, as they are already at target.
-// Velocity and such is still copied. This will be handled better in the future.
-int32 SyncKinematicOnGameThread = 0;
-FAutoConsoleVariableRef CVar_SyncKinematicOnGameThread(TEXT("P.Chaos.SyncKinematicOnGameThread"), SyncKinematicOnGameThread, TEXT("If set to 1, if a kinematic is flagged to send position back to game thread, move component, if 0, do not."));
+// This allows forcing the game thread/actor to get or not get transform updates from the simulation result for
+// kinematics - noting that they may already have been set in the SetKinematicTransform function.
+// Velocity and such will still be copied in any case. The default value of -1 uses the UpdateKinematicFromSimulation 
+// flag on the BodyInstance.
+CHAOS_API int32 SyncKinematicOnGameThread = -1;
+FAutoConsoleVariableRef CVar_SyncKinematicOnGameThread(TEXT("P.Chaos.SyncKinematicOnGameThread"), 
+	SyncKinematicOnGameThread, TEXT(
+		"If set to 1, kinematic bodies will always send their transforms back to the game thread, following the "
+		"simulation step/results. If 0, then they will never do so, and kinematics will be updated immediately "
+		"their kinematic target is set. Any other value (e.g. the default -1) means that the behavior is "
+		"determined on a per-object basis with the UpdateKinematicFromSimulation flag in BodyInstance."));
 
 FSingleParticlePhysicsProxy::FSingleParticlePhysicsProxy(TUniquePtr<PARTICLE_TYPE>&& InParticle, FParticleHandle* InHandle, UObject* InOwner)
 	: IPhysicsProxyBase(EPhysicsProxyType::SingleParticleProxy, InOwner, MakeShared<FSingleParticleProxyTimestamp>())
@@ -263,7 +270,22 @@ bool FSingleParticlePhysicsProxy::PullFromPhysicsState(const Chaos::FDirtyRigidP
 	auto Rigid = Particle ? Particle->CastToRigidParticle() : nullptr;
 	if(Rigid)
 	{
-		const bool bSyncXR = SyncKinematicOnGameThread || (Rigid->ObjectState() != EObjectStateType::Kinematic);
+		// Note that kinematics should either be updated here (following simulation), or when the
+		// kinematic target is set in FChaosEngineInterface::SetKinematicTarget_AssumesLocked If the
+		// logic in one place is changed, it should be checked in the other place too.
+		bool bUpdatePositionFromSimulation = true;
+		if (Rigid->ObjectState() == EObjectStateType::Kinematic)
+		{
+			switch (SyncKinematicOnGameThread)
+			{
+			case 0:
+				bUpdatePositionFromSimulation = false ; break;
+			case 1: 
+				bUpdatePositionFromSimulation = true; break;
+			default:
+				bUpdatePositionFromSimulation = Rigid->UpdateKinematicFromSimulation();
+			}
+		}
 
 		const FSingleParticleProxyTimestamp* ProxyTimestamp = PullData.GetTimestamp();
 		
@@ -279,7 +301,7 @@ bool FSingleParticlePhysicsProxy::PullFromPhysicsState(const Chaos::FDirtyRigidP
 				return OverwriteProperty.Timestamp <= SolverSyncTimestamp ? (OverwriteProperty.Timestamp < SolverSyncTimestamp ? &Prev : &OverwriteProperty.Value) : nullptr;
 			};
 
-			if (bSyncXR)
+			if (bUpdatePositionFromSimulation)
 			{
 				const float UseResimInterpStrength = InterpolationData.GetInterpChannel_External() == 0 ? ResimInterpStrength : ResimInterpStrength2;
 
@@ -345,7 +367,7 @@ bool FSingleParticlePhysicsProxy::PullFromPhysicsState(const Chaos::FDirtyRigidP
 		}
 		else
 		{
-			if (bSyncXR)
+			if (bUpdatePositionFromSimulation)
 			{
 				//no interpolation, just ignore if overwrite comes after
 				if (SolverSyncTimestamp >= ProxyTimestamp->OverWriteX.Timestamp)
diff --git a/Engine/Source/Runtime/Experimental/Chaos/Public/Chaos/ParticleDirtyFlags.h b/Engine/Source/Runtime/Experimental/Chaos/Public/Chaos/ParticleDirtyFlags.h
--- a/Engine/Source/Runtime/Experimental/Chaos/Public/Chaos/ParticleDirtyFlags.h
+++ b/Engine/Source/Runtime/Experimental/Chaos/Public/Chaos/ParticleDirtyFlags.h
@@ -429,6 +429,9 @@ public:
 	bool GravityEnabled() const { return MControlFlags.GetGravityEnabled(); }
 	void SetGravityEnabled(bool bInGravity){ MControlFlags.SetGravityEnabled(bInGravity); }
 
+	bool UpdateKinematicFromSimulation() const { return MControlFlags.GetUpdateKinematicFromSimulation(); }
+	void SetUpdateKinematicFromSimulation(bool bUpdateKinematicFromSimulation) { MControlFlags.SetUpdateKinematicFromSimulation(bUpdateKinematicFromSimulation); }
+
 	bool CCDEnabled() const { return MControlFlags.GetCCDEnabled(); }
 	void SetCCDEnabled(bool bInCCDEnabled) { MControlFlags.SetCCDEnabled(bInCCDEnabled); }
 
diff --git a/Engine/Source/Runtime/Experimental/Chaos/Public/Chaos/ParticleHandle.h b/Engine/Source/Runtime/Experimental/Chaos/Public/Chaos/ParticleHandle.h
--- a/Engine/Source/Runtime/Experimental/Chaos/Public/Chaos/ParticleHandle.h
+++ b/Engine/Source/Runtime/Experimental/Chaos/Public/Chaos/ParticleHandle.h
@@ -1136,6 +1136,16 @@ public:
 		PBDRigidParticles->ControlFlags(ParticleIdx).SetGravityGroupIndex(GravityGroupIndex);
 	}
 
+	inline bool UpdateKinematicFromSimulation() const
+	{
+		return ControlFlags().GetUpdateKinematicFromSimulation();
+	}
+
+	inline void SetUpdateKinematicFromSimulation(bool bUpdateKinematicFromSimulation)
+	{
+		PBDRigidParticles->ControlFlags(ParticleIdx).SetUpdateKinematicFromSimulation(bUpdateKinematicFromSimulation);
+	}
+
 	inline bool CCDEnabled() const
 	{
 		return ControlFlags().GetCCDEnabled();
@@ -2902,6 +2912,12 @@ public:
 		MMiscData.Modify(true, MDirtyFlags, Proxy, [bInEnabled](auto& Data) { Data.SetGravityEnabled(bInEnabled); });
 	}
 	
+	bool UpdateKinematicFromSimulation() const { return MMiscData.Read().UpdateKinematicFromSimulation(); }
+	void SetUpdateKinematicFromSimulation(const bool bUpdateKinematicFromSimulation)
+	{
+		MMiscData.Modify(true, MDirtyFlags, Proxy, [bUpdateKinematicFromSimulation](auto& Data) { Data.SetUpdateKinematicFromSimulation(bUpdateKinematicFromSimulation); });
+	}
+
 	bool OneWayInteraction() const { return MMiscData.Read().OneWayInteraction(); }
 	void SetOneWayInteraction(const bool bInEnabled)
 	{
diff --git a/Engine/Source/Runtime/Experimental/Chaos/Public/Chaos/RigidParticleControlFlags.h b/Engine/Source/Runtime/Experimental/Chaos/Public/Chaos/RigidParticleControlFlags.h
--- a/Engine/Source/Runtime/Experimental/Chaos/Public/Chaos/RigidParticleControlFlags.h
+++ b/Engine/Source/Runtime/Experimental/Chaos/Public/Chaos/RigidParticleControlFlags.h
@@ -17,7 +17,7 @@ namespace Chaos
 	class FRigidParticleControlFlags
 	{
 	public:
-		using FStorage = uint8;
+		using FStorage = uint16;
 
 		FRigidParticleControlFlags()
 			: Bits(0)
@@ -30,6 +30,9 @@ namespace Chaos
 		bool GetGravityEnabled() const { return Flags.bGravityEnabled; }
 		FRigidParticleControlFlags& SetGravityEnabled(const bool bEnabled) { Flags.bGravityEnabled = bEnabled; return *this; }
 
+		bool GetUpdateKinematicFromSimulation() const { return Flags.bUpdateKinematicFromSimulation; }
+		FRigidParticleControlFlags& SetUpdateKinematicFromSimulation(const bool bUpdateKinematicFromSimulation) { Flags.bUpdateKinematicFromSimulation = bUpdateKinematicFromSimulation; return *this; }
+
 		int32 GetGravityGroupIndex() const { return Flags.GravityGroupIndex; }
 		FRigidParticleControlFlags& SetGravityGroupIndex(const int32 GravityGroupIndex) 
 		{ 
@@ -69,6 +72,7 @@ namespace Chaos
 			FStorage bMaxDepenetrationVelocityOverrideEnabled : 1;
 			FStorage bInertiaConditioningEnabled : 1;
 			FStorage GravityGroupIndex : 3;
+			FStorage bUpdateKinematicFromSimulation : 1;
 			// Add new properties above this line
 			// Change FStorage typedef if we exceed the max bits
 		};
diff --git a/Engine/Source/Runtime/Experimental/Chaos/Public/PhysicsProxy/SingleParticlePhysicsProxy.h b/Engine/Source/Runtime/Experimental/Chaos/Public/PhysicsProxy/SingleParticlePhysicsProxy.h
--- a/Engine/Source/Runtime/Experimental/Chaos/Public/PhysicsProxy/SingleParticlePhysicsProxy.h
+++ b/Engine/Source/Runtime/Experimental/Chaos/Public/PhysicsProxy/SingleParticlePhysicsProxy.h
@@ -22,7 +22,6 @@
 
 namespace Chaos
 {
-extern CHAOS_API int32 SyncKinematicOnGameThread;
 
 class FPBDRigidsEvolutionGBF;
 
@@ -376,6 +375,30 @@ public:
 		});
 	}
 
+	bool UpdateKinematicFromSimulation() const
+	{
+		return Read([](auto* Particle)
+		{
+			if (auto Rigid = Particle->CastToRigidParticle())
+			{
+				return Rigid->UpdateKinematicFromSimulation();
+			}
+
+			return false;
+		});
+	}
+
+	void SetUpdateKinematicFromSimulation(const bool InUpdateKinematicFromSimulation)
+	{
+		Write([InUpdateKinematicFromSimulation](auto* Particle)
+		{
+			if (auto Rigid = Particle->CastToRigidParticle())
+			{
+				return Rigid->SetUpdateKinematicFromSimulation(InUpdateKinematicFromSimulation);
+			}
+		});
+	}
+
 	bool CCDEnabled() const
 	{
 		return Read([](auto* Particle)
diff --git a/Engine/Source/Runtime/PhysicsCore/Private/BodyInstanceCore.cpp b/Engine/Source/Runtime/PhysicsCore/Private/BodyInstanceCore.cpp
--- a/Engine/Source/Runtime/PhysicsCore/Private/BodyInstanceCore.cpp
+++ b/Engine/Source/Runtime/PhysicsCore/Private/BodyInstanceCore.cpp
@@ -14,6 +14,7 @@ FBodyInstanceCore::FBodyInstanceCore()
 : bSimulatePhysics(false)
 , bOverrideMass(false)
 , bEnableGravity(true)
+, bUpdateKinematicFromSimulation(false)
 , bAutoWeld(false)
 , bStartAwake(true)
 , bGenerateWakeEvents(false)
diff --git a/Engine/Source/Runtime/PhysicsCore/Private/ChaosEngineInterface.cpp b/Engine/Source/Runtime/PhysicsCore/Private/ChaosEngineInterface.cpp
--- a/Engine/Source/Runtime/PhysicsCore/Private/ChaosEngineInterface.cpp
+++ b/Engine/Source/Runtime/PhysicsCore/Private/ChaosEngineInterface.cpp
@@ -33,6 +33,7 @@ namespace Chaos
 {
 	extern CHAOS_API int32 AccelerationStructureSplitStaticAndDynamic;
 	extern CHAOS_API int32 AccelerationStructureIsolateQueryOnlyObjects;
+	extern CHAOS_API int32 SyncKinematicOnGameThread;
 }
 
 bool bEnableChaosJointConstraints = true;
@@ -1120,6 +1121,15 @@ void FChaosEngineInterface::SetGravityEnabled_AssumesLocked(const FPhysicsActorH
 	InActorReference->GetGameThreadAPI().SetGravityEnabled(bEnabled);
 }
 
+bool FChaosEngineInterface::GetUpdateKinematicFromSimulation_AssumesLocked(const FPhysicsActorHandle& InActorReference)
+{
+	return InActorReference->GetGameThreadAPI().UpdateKinematicFromSimulation();
+}
+void FChaosEngineInterface::SetUpdateKinematicFromSimulation_AssumesLocked(const FPhysicsActorHandle& InActorReference, bool bUpdateKinematicFromSimulation)
+{
+	InActorReference->GetGameThreadAPI().SetUpdateKinematicFromSimulation(bUpdateKinematicFromSimulation);
+}
+
 void FChaosEngineInterface::SetOneWayInteraction_AssumesLocked(const FPhysicsActorHandle& InHandle, bool InOneWayInteraction)
 {
 	InHandle->GetGameThreadAPI().SetOneWayInteraction(InOneWayInteraction);
@@ -2002,6 +2012,7 @@ void FChaosEngineInterface::CreateActor(const FActorCreationParams& InParams,FPh
 		// Create an underlying dynamic particle
 		TUniquePtr<FPBDRigidParticle> Rigid = FPBDRigidParticle::CreateParticle();
 		Rigid->SetGravityEnabled(InParams.bEnableGravity);
+		Rigid->SetUpdateKinematicFromSimulation(InParams.bUpdateKinematicFromSimulation);
 		if(InParams.bSimulatePhysics)
 		{
 			if(InParams.bStartAwake)
@@ -2148,11 +2159,31 @@ void FChaosEngineInterface::SetKinematicTarget_AssumesLocked(const FPhysicsActor
 	const Chaos::FKinematicTarget NewKinematicTarget = Chaos::FKinematicTarget::MakePositionTarget(InNewTarget);
 	InActorReference->GetGameThreadAPI().SetKinematicTarget(NewKinematicTarget);
 
-	// IMPORTANT : we do not invalidate X and R as they will be properly computed using the kinematic target information 
-	InActorReference->GetGameThreadAPI().SetX(InNewTarget.GetLocation(), false); 
-	InActorReference->GetGameThreadAPI().SetR(InNewTarget.GetRotation(), false); 
-	InActorReference->GetGameThreadAPI().UpdateShapeBounds();
+	// Match the logic in FSingleParticlePhysicsProxy::PullFromPhysicsState - to see if that will be
+	// updating the position. If not, then we need to do it here. Note t
+	bool bUpdatePositionFromSimulation = false;
+	Chaos::FPBDRigidParticle* Rigid = InActorReference->GetRigidParticleUnsafe();
+	if (Rigid && Rigid->ObjectState() == Chaos::EObjectStateType::Kinematic)
+	{
+		switch (Chaos::SyncKinematicOnGameThread)
+		{
+		case 0:
+			bUpdatePositionFromSimulation = false; break;
+		case 1:
+			bUpdatePositionFromSimulation = true; break;
+		default:
+			bUpdatePositionFromSimulation = Rigid->UpdateKinematicFromSimulation();
+		}
+	}
+
+	if (!bUpdatePositionFromSimulation)
+	{
+		// IMPORTANT : we do not invalidate X and R as they will be properly computed using the kinematic target information 
+		InActorReference->GetGameThreadAPI().SetX(InNewTarget.GetLocation(), false); 
+		InActorReference->GetGameThreadAPI().SetR(InNewTarget.GetRotation(), false); 
+		InActorReference->GetGameThreadAPI().UpdateShapeBounds();
 
-	FChaosScene* Scene = GetCurrentScene(InActorReference);
-	Scene->UpdateActorInAccelerationStructure(InActorReference);
+		FChaosScene* Scene = GetCurrentScene(InActorReference);
+		Scene->UpdateActorInAccelerationStructure(InActorReference);
+	}
 }
diff --git a/Engine/Source/Runtime/PhysicsCore/Public/BodyInstanceCore.h b/Engine/Source/Runtime/PhysicsCore/Public/BodyInstanceCore.h
--- a/Engine/Source/Runtime/PhysicsCore/Public/BodyInstanceCore.h
+++ b/Engine/Source/Runtime/PhysicsCore/Public/BodyInstanceCore.h
@@ -39,6 +39,10 @@ struct PHYSICSCORE_API FBodyInstanceCore
 	UPROPERTY(EditAnywhere,BlueprintReadOnly,Category = Physics)
 	uint8 bEnableGravity : 1;
 
+	/** When kinematic, whether the actor transform should be updated as a result of movement in the simulation, rather than immediately whenever a target transform is set. */
+	UPROPERTY(EditAnywhere,BlueprintReadOnly,Category = Physics)
+	uint8 bUpdateKinematicFromSimulation : 1;
+
 	/** If true and is attached to a parent, the two bodies will be joined into a single rigid body. Physical settings like collision profile and body settings are determined by the root */
 	UPROPERTY(EditAnywhere,AdvancedDisplay,BlueprintReadWrite,Category = Physics,meta = (editcondition = "!bSimulatePhysics"))
 	uint8 bAutoWeld : 1;
diff --git a/Engine/Source/Runtime/PhysicsCore/Public/Chaos/ChaosEngineInterface.h b/Engine/Source/Runtime/PhysicsCore/Public/Chaos/ChaosEngineInterface.h
--- a/Engine/Source/Runtime/PhysicsCore/Public/Chaos/ChaosEngineInterface.h
+++ b/Engine/Source/Runtime/PhysicsCore/Public/Chaos/ChaosEngineInterface.h
@@ -456,6 +456,9 @@ public:
 	static bool IsGravityEnabled_AssumesLocked(const FPhysicsActorHandle& InActorReference);
 	static void SetGravityEnabled_AssumesLocked(const FPhysicsActorHandle& InActorReference,bool bEnabled);
 
+	static bool GetUpdateKinematicFromSimulation_AssumesLocked(const FPhysicsActorHandle& InActorReference);
+	static void SetUpdateKinematicFromSimulation_AssumesLocked(const FPhysicsActorHandle& InActorReference, bool bUpdateKinematicFromSimulation);
+
 	static void SetOneWayInteraction_AssumesLocked(const FPhysicsActorHandle& InHandle, bool InOneWayInteraction);
 
 	static float GetSleepEnergyThreshold_AssumesLocked(const FPhysicsActorHandle& InActorReference);
diff --git a/Engine/Source/Runtime/PhysicsCore/Public/PhysicsInterfaceTypesCore.h b/Engine/Source/Runtime/PhysicsCore/Public/PhysicsInterfaceTypesCore.h
--- a/Engine/Source/Runtime/PhysicsCore/Public/PhysicsInterfaceTypesCore.h
+++ b/Engine/Source/Runtime/PhysicsCore/Public/PhysicsInterfaceTypesCore.h
@@ -18,6 +18,7 @@ struct FActorCreationParams
 		, bStatic(false)
 		, bQueryOnly(false)
 		, bEnableGravity(false)
+		, bUpdateKinematicFromSimulation(false)
 		, bSimulatePhysics(false)
 		, bStartAwake(true)
 		, DebugName(nullptr)
@@ -27,6 +28,7 @@ struct FActorCreationParams
 	bool bStatic;
 	bool bQueryOnly;
 	bool bEnableGravity;
+	bool bUpdateKinematicFromSimulation;
 	bool bSimulatePhysics;
 	bool bStartAwake;
 	char* DebugName;
