Update/add files associated with RigidBodyWithControl

diff --git a/Engine/Plugins/Experimental/PhysicsControl/PhysicsControl.uplugin b/Engine/Plugins/Experimental/PhysicsControl/PhysicsControl.uplugin
--- a/Engine/Plugins/Experimental/PhysicsControl/PhysicsControl.uplugin
+++ b/Engine/Plugins/Experimental/PhysicsControl/PhysicsControl.uplugin
@@ -21,6 +21,11 @@
 			"Type": "Runtime",
 			"LoadingPhase": "Default"
 		},
+		{
+			"Name": "PhysicsControlUncookedOnly",
+			"Type": "UncookedOnly",
+			"LoadingPhase": "Default"
+		},
 		{
 			"Name": "PhysicsControlEditor",
 			"Type": "Editor",
diff --git a/Engine/Plugins/Experimental/PhysicsControl/Source/PhysicsControl/PhysicsControl.Build.cs b/Engine/Plugins/Experimental/PhysicsControl/Source/PhysicsControl/PhysicsControl.Build.cs
--- a/Engine/Plugins/Experimental/PhysicsControl/Source/PhysicsControl/PhysicsControl.Build.cs
+++ b/Engine/Plugins/Experimental/PhysicsControl/Source/PhysicsControl/PhysicsControl.Build.cs
@@ -19,11 +19,18 @@ public class PhysicsControl : ModuleRules
 		PrivateDependencyModuleNames.AddRange(
 			new string[]
 			{
+				"AnimationCore",
+				"AnimGraphRuntime",
+				"Chaos",
+				"ChaosCore",
+				"Core",
 				"CoreUObject",
-				"PhysicsCore",
-				"Engine"
+				"Engine",
+				"PhysicsCore"
 			}
 		);
-		
+
+		SetupModulePhysicsSupport(Target);
+		PrivateDefinitions.Add("CHAOS_INCLUDE_LEVEL_1=1");
 	}
 }
diff --git a/Engine/Plugins/Experimental/PhysicsControl/Source/PhysicsControl/Private/AnimNode_RigidBodyWithControl.cpp b/Engine/Plugins/Experimental/PhysicsControl/Source/PhysicsControl/Private/AnimNode_RigidBodyWithControl.cpp
new file mode 100644
--- /dev/null
+++ b/Engine/Plugins/Experimental/PhysicsControl/Source/PhysicsControl/Private/AnimNode_RigidBodyWithControl.cpp
@@ -0,0 +1,1924 @@
+// Copyright Epic Games, Inc. All Rights Reserved.
+
+/* 
+ * This is essentially a copy + paste of the RigidBody file
+ */
+
+#include "AnimNode_RigidBodyWithControl.h"
+#include "PhysicsControlLog.h"
+#include "AnimationRuntime.h"
+#include "Animation/AnimInstanceProxy.h"
+#include "ClothCollisionSource.h"
+#include "GameFramework/Pawn.h"
+#include "HAL/Event.h"
+#include "HAL/LowLevelMemTracker.h"
+#include "PhysicsEngine/BodySetup.h"
+#include "PhysicsEngine/PhysicsAsset.h"
+#include "PhysicsEngine/PhysicsConstraintTemplate.h"
+#include "GameFramework/PawnMovementComponent.h"
+#include "Physics/PhysicsInterfaceCore.h"
+#include "Physics/ImmediatePhysics/ImmediatePhysicsActorHandle.h"
+#include "Physics/ImmediatePhysics/ImmediatePhysicsJointHandle.h"
+#include "Physics/ImmediatePhysics/ImmediatePhysicsSimulation.h"
+#include "Physics/ImmediatePhysics/ImmediatePhysicsStats.h"
+#include "Chaos/PBDJointConstraints.h"
+#include "PhysicsField/PhysicsFieldComponent.h"
+#include "PhysicsEngine/PhysicsSettings.h"
+#include "Logging/MessageLog.h"
+
+#include "Chaos/PBDJointConstraintTypes.h"
+
+#include UE_INLINE_GENERATED_CPP_BY_NAME(AnimNode_RigidBodyWithControl)
+
+LLM_DEFINE_TAG(Animation_RigidBodyWithControl);
+
+//UE_DISABLE_OPTIMIZATION;
+
+/////////////////////////////////////////////////////
+// FAnimNode_RigidBodyWithControl
+
+#define LOCTEXT_NAMESPACE "ImmediatePhysicsWithControl"
+
+//DEFINE_STAT(STAT_RigidBodyNodeWithControlInitTime);
+
+CSV_DECLARE_CATEGORY_MODULE_EXTERN(ENGINE_API, Animation);
+
+TAutoConsoleVariable<int32> CVarEnableRigidBodyNodeWithControl(TEXT("p.RigidBodyNodeWithControl"), 1, TEXT("Enables/disables the whole rigid body node system. When disabled, avoids all allocations and runtime costs. Can be used to disable RB Nodes on low-end platforms."), ECVF_Scalability);
+TAutoConsoleVariable<int32> CVarEnableRigidBodyNodeWithControlSimulation(TEXT("p.RigidBodyNodeWithControl.EnableSimulation"), 1, TEXT("Runtime Enable/Disable RB Node Simulation for debugging and testing (node is initialized and bodies and constraints are created, even when disabled.)"), ECVF_Default);
+TAutoConsoleVariable<int32> CVarRigidBodyNodeWithControlLODThreshold(TEXT("p.RigidBodyWithControlLODThreshold"), -1, TEXT("Max LOD that rigid body node is allowed to run on. Provides a global threshold that overrides per-node the LODThreshold property. -1 means no override."), ECVF_Scalability);
+
+int32 RBANWithControl_MaxSubSteps = 4;
+bool bRBANWithControl_EnableTimeBasedReset = true;
+bool bRBANWithControl_EnableComponentAcceleration = true;
+int32 RBANWithControl_WorldObjectExpiry = 4;
+FAutoConsoleVariableRef CVarRigidBodyNodeWithControlMaxSteps(TEXT("p.RigidBodyNodeWithControl.MaxSubSteps"), RBANWithControl_MaxSubSteps, TEXT("Set the maximum number of simulation steps in the update loop"), ECVF_Default);
+FAutoConsoleVariableRef CVarRigidBodyNodeWithControlEnableTimeBasedReset(TEXT("p.RigidBodyNodeWithControl.EnableTimeBasedReset"), bRBANWithControl_EnableTimeBasedReset, TEXT("If true, Rigid Body nodes are reset when they have not been updated for a while (default true)"), ECVF_Default);
+FAutoConsoleVariableRef CVarRigidBodyNodeWithControlEnableComponentAcceleration(TEXT("p.RigidBodyNodeWithControl.EnableComponentAcceleration"), bRBANWithControl_EnableComponentAcceleration, TEXT("Enable/Disable the simple acceleration transfer system for component- or bone-space simulation"), ECVF_Default);
+FAutoConsoleVariableRef CVarRigidBodyNodeWithControlWorldObjectExpiry(TEXT("p.RigidBodyNodeWithControl.WorldObjectExpiry"), RBANWithControl_WorldObjectExpiry, TEXT("World objects are removed from the simulation if not detected after this many tests"), ECVF_Default);
+
+bool bRBANWithControl_IncludeClothColliders = true;
+FAutoConsoleVariableRef CVarRigidBodyNodeWithControlIncludeClothColliders(TEXT("p.RigidBodyNodeWithControl.IncludeClothColliders"), bRBANWithControl_IncludeClothColliders, TEXT("Include cloth colliders as kinematic bodies in the immediate physics simulation."), ECVF_Default);
+
+// FSimSpaceSettingsWithControl forced overrides for testing
+bool bRBANWithControl_SimSpace_EnableOverride = false;
+FSimSpaceSettings RBANWithControl_SimSpaceOverride;
+FAutoConsoleVariableRef CVarRigidBodyNodeWithControlSpaceOverride(TEXT("p.RigidBodyNodeWithControl.Space.Override"), bRBANWithControl_SimSpace_EnableOverride, TEXT("Force-enable the advanced simulation space movement forces"), ECVF_Default);
+FAutoConsoleVariableRef CVarRigidBodyNodeWithControlSpaceWorldAlpha(TEXT("p.RigidBodyNodeWithControl.Space.WorldAlpha"), RBANWithControl_SimSpaceOverride.WorldAlpha, TEXT("RBANWithControl SimSpaceSettings overrides"), ECVF_Default);
+FAutoConsoleVariableRef CVarRigidBodyNodeWithControlSpaceVelScaleZ(TEXT("p.RigidBodyNodeWithControl.Space.VelocityScaleZ"), RBANWithControl_SimSpaceOverride.VelocityScaleZ, TEXT("RBANWithControl SimSpaceSettings overrides"), ECVF_Default);
+FAutoConsoleVariableRef CVarRigidBodyNodeWithControlSpaceMaxCompLinVel(TEXT("p.RigidBodyNodeWithControl.Space.MaxLinearVelocity"), RBANWithControl_SimSpaceOverride.MaxLinearVelocity, TEXT("RBANWithControl SimSpaceSettings overrides"), ECVF_Default);
+FAutoConsoleVariableRef CVarRigidBodyNodeWithControlSpaceMaxCompAngVel(TEXT("p.RigidBodyNodeWithControl.Space.MaxAngularVelocity"), RBANWithControl_SimSpaceOverride.MaxAngularVelocity, TEXT("RBANWithControl SimSpaceSettings overrides"), ECVF_Default);
+FAutoConsoleVariableRef CVarRigidBodyNodeWithControlSpaceMaxCompLinAcc(TEXT("p.RigidBodyNodeWithControl.Space.MaxLinearAcceleration"), RBANWithControl_SimSpaceOverride.MaxLinearAcceleration, TEXT("RBANWithControl SimSpaceSettings overrides"), ECVF_Default);
+FAutoConsoleVariableRef CVarRigidBodyNodeWithControlSpaceMaxCompAngAcc(TEXT("p.RigidBodyNodeWithControl.Space.MaxAngularAcceleration"), RBANWithControl_SimSpaceOverride.MaxAngularAcceleration, TEXT("RBANWithControl SimSpaceSettings overrides"), ECVF_Default);
+// LWC_TODO: Double support for console variables
+#if 0
+FAutoConsoleVariableRef CVarRigidBodyNodeWithControlSpaceExternalLinearDragX(TEXT("p.RigidBodyNodeWithControl.Space.ExternalLinearDrag.X"), RBANWithControl_SimSpaceOverride.ExternalLinearDragV.X, TEXT("RBANWithControl SimSpaceSettings overrides"), ECVF_Default);
+FAutoConsoleVariableRef CVarRigidBodyNodeWithControlSpaceExternalLinearDragY(TEXT("p.RigidBodyNodeWithControl.Space.ExternalLinearDrag.Y"), RBANWithControl_SimSpaceOverride.ExternalLinearDragV.Y, TEXT("RBANWithControl SimSpaceSettings overrides"), ECVF_Default);
+FAutoConsoleVariableRef CVarRigidBodyNodeWithControlSpaceExternalLinearDragZ(TEXT("p.RigidBodyNodeWithControl.Space.ExternalLinearDrag.Z"), RBANWithControl_SimSpaceOverride.ExternalLinearDragV.Z, TEXT("RBANWithControl SimSpaceSettings overrides"), ECVF_Default);
+FAutoConsoleVariableRef CVarRigidBodyNodeWithControlSpaceExternalLinearVelocityX(TEXT("p.RigidBodyNodeWithControl.Space.ExternalLinearVelocity.X"), RBANWithControl_SimSpaceOverride.ExternalLinearVelocity.X, TEXT("RBANWithControl SimSpaceSettings overrides"), ECVF_Default);
+FAutoConsoleVariableRef CVarRigidBodyNodeWithControlSpaceExternalLinearVelocityY(TEXT("p.RigidBodyNodeWithControl.Space.ExternalLinearVelocity.Y"), RBANWithControl_SimSpaceOverride.ExternalLinearVelocity.Y, TEXT("RBANWithControl SimSpaceSettings overrides"), ECVF_Default);
+FAutoConsoleVariableRef CVarRigidBodyNodeWithControlSpaceExternalLinearVelocityZ(TEXT("p.RigidBodyNodeWithControl.Space.ExternalLinearVelocity.Z"), RBANWithControl_SimSpaceOverride.ExternalLinearVelocity.Z, TEXT("RBANWithControl SimSpaceSettings overrides"), ECVF_Default);
+#endif
+
+bool bRBANWithControl_DeferredSimulationDefault = false;
+FAutoConsoleVariableRef CVarRigidBodyNodeWithControlDeferredSimulationDefault(
+	TEXT("p.RigidBodyNodeWithControl.DeferredSimulationDefault"),
+	bRBANWithControl_DeferredSimulationDefault,
+	TEXT("Whether rigid body simulations are deferred one frame for assets that don't opt into a specific simulation timing"),
+	ECVF_Default);
+
+bool bRBANWithControl_DebugDraw = false;
+FAutoConsoleVariableRef CVarRigidBodyNodeWithControlDebugDraw(TEXT("p.RigidBodyNodeWithControl.DebugDraw"), bRBANWithControl_DebugDraw, TEXT("Whether to debug draw the rigid body simulation state. Requires p.Chaos.DebugDraw.Enabled 1 to function as well."), ECVF_Default);
+
+// Array of priorities that can be indexed into with CVars, since task priorities cannot be set from scalability .ini
+static UE::Tasks::ETaskPriority GRigidBodyNodeWithControlTaskPriorities[] =
+{
+	UE::Tasks::ETaskPriority::High,
+	UE::Tasks::ETaskPriority::Normal,
+	UE::Tasks::ETaskPriority::BackgroundHigh,
+	UE::Tasks::ETaskPriority::BackgroundNormal,
+	UE::Tasks::ETaskPriority::BackgroundLow
+};
+
+static int32 GRigidBodyNodeWithControlSimulationTaskPriority = 0;
+FAutoConsoleVariableRef CVarRigidBodyNodeWithControlSimulationTaskPriority(
+	TEXT("p.RigidBodyNodeWithControl.TaskPriority.Simulation"),
+	GRigidBodyNodeWithControlSimulationTaskPriority,
+	TEXT("Task priority for running the rigid body node simulation task (0 = foreground/high, 1 = foreground/normal, 2 = background/high, 3 = background/normal, 4 = background/low)."),
+	ECVF_Default
+);
+
+FAnimNode_RigidBodyWithControl::FAnimNode_RigidBodyWithControl()
+	: OverridePhysicsAsset(nullptr)
+	, PreviousCompWorldSpaceTM()
+	, CurrentTransform()
+	, PreviousTransform()
+	, PhysicsAssetToUse(nullptr)
+	, OverrideWorldGravity(0.0f)
+	, ExternalForce(0.0f)
+	, ComponentLinearAccScale(0.0f)
+	, ComponentLinearVelScale(0.0f)
+	, ComponentAppliedLinearAccClamp(10000.0f)
+	, SimSpaceSettings()
+	, CachedBoundsScale(1.2f)
+	, BaseBoneRef()
+	, OverlapChannel(ECC_WorldStatic)
+	, SimulationSpace(ESimulationSpace::ComponentSpace)
+	, bForceDisableCollisionBetweenConstraintBodies(false)
+	, bUseExternalClothCollision(false)
+	, ResetSimulatedTeleportType(ETeleportType::None)
+	, bEnableWorldGeometry(false)
+	, bOverrideWorldGravity(false)
+	, bTransferBoneVelocities(false)
+	, bFreezeIncomingPoseOnStart(false)
+	, bClampLinearTranslationLimitToRefPose(false)
+	, WorldSpaceMinimumScale(0.01f)
+	, EvaluationResetTime(0.01f)
+	, bEnableControls(false)
+	, bEnabled(false)
+	, bSimulationStarted(false)
+	, bCheckForBodyTransformInit(false)
+	, bHaveRunControlSeup(false)
+	, SimulationTiming(ESimulationTiming::Default)
+	, WorldTimeSeconds(0.0f)
+	, LastEvalTimeSeconds(0.0f)
+	, AccumulatedDeltaTime(0.0f)
+	, AnimPhysicsMinDeltaTime(0.0f)
+	, bSimulateAnimPhysicsAfterReset(false)
+	, SkelMeshCompWeakPtr()
+	, PhysicsSimulation(nullptr)
+	, SolverSettings()
+	, SolverIterations()
+	, SimulationTask()
+	, OutputBoneData()
+	, Bodies()
+	, SkeletonBoneIndexToBodyIndex()
+	, BodyAnimData()
+	, PendingRadialForces()
+	, PerSolverField()
+	, ComponentsInSim()
+	, ComponentsInSimTick(0)
+	, WorldSpaceGravity(0.0f)
+	, TotalMass(0.0f)
+	, CachedBounds(FVector::ZeroVector, 0.0f)
+	, QueryParams(NAME_None, FCollisionQueryParams::GetUnknownStatId())
+	, PhysScene(nullptr)
+	, UnsafeWorld(nullptr)
+	, UnsafeOwner(nullptr)
+	, CapturedBoneVelocityBoneContainer()
+	, CapturedBoneVelocityPose()
+	, CapturedFrozenPose()
+	, CapturedFrozenCurves()
+	, PreviousComponentLinearVelocity(0.0f)
+	, SimSpacePreviousComponentToWorld()
+	, SimSpacePreviousBoneToComponent()
+	, SimSpacePreviousComponentLinearVelocity(0.0f)
+	, SimSpacePreviousComponentAngularVelocity(0.0f)
+	, SimSpacePreviousBoneLinearVelocity(0.0f)
+	, SimSpacePreviousBoneAngularVelocity(0.0f)
+{
+}
+
+FAnimNode_RigidBodyWithControl::~FAnimNode_RigidBodyWithControl()
+{
+	DestroyPhysicsSimulation();
+}
+
+void FAnimNode_RigidBodyWithControl::GatherDebugData(FNodeDebugData& DebugData)
+{
+	DECLARE_SCOPE_HIERARCHICAL_COUNTER_ANIMNODE(GatherDebugData)
+	FString DebugLine = DebugData.GetNodeName(this);
+
+	DebugLine += "(";
+	AddDebugNodeData(DebugLine);
+	DebugLine += ")";
+
+	DebugData.AddDebugItem(DebugLine);
+
+	const bool bUsingFrozenPose = bFreezeIncomingPoseOnStart && bSimulationStarted && (CapturedFrozenPose.GetPose().GetNumBones() > 0);
+	if (!bUsingFrozenPose)
+	{
+		ComponentPose.GatherDebugData(DebugData);
+	}
+}
+
+void FAnimNode_RigidBodyWithControl::Initialize_AnyThread(const FAnimationInitializeContext& Context)
+{
+	Super::Initialize_AnyThread(Context);
+
+#if WITH_EDITOR
+	if(GIsReinstancing)
+	{
+		InitPhysics(Cast<UAnimInstance>(Context.GetAnimInstanceObject()));
+	}
+#endif
+}
+
+
+void FAnimNode_RigidBodyWithControl::UpdateComponentPose_AnyThread(const FAnimationUpdateContext& Context)
+{
+	DECLARE_SCOPE_HIERARCHICAL_COUNTER_ANIMNODE(UpdateComponentPose_AnyThread)
+	// Only freeze update graph after initial update, as we want to get that pose through.
+	if (bFreezeIncomingPoseOnStart && bSimulationStarted && ResetSimulatedTeleportType == ETeleportType::None)
+	{
+		// If we have a Frozen Pose captured, 
+		// then we don't need to update the rest of the graph.
+		if (CapturedFrozenPose.GetPose().GetNumBones() > 0)
+		{
+		}
+		else
+		{
+			// Create a new context with zero deltatime to freeze time in rest of the graph.
+			// This will be used to capture a frozen pose.
+			FAnimationUpdateContext FrozenContext = Context.FractionalWeightAndTime(1.f, 0.f);
+
+			Super::UpdateComponentPose_AnyThread(FrozenContext);
+		}
+	}
+	else
+	{
+		Super::UpdateComponentPose_AnyThread(Context);
+	}
+}
+
+void FAnimNode_RigidBodyWithControl::EvaluateComponentPose_AnyThread(FComponentSpacePoseContext& Output)
+{
+	DECLARE_SCOPE_HIERARCHICAL_COUNTER_ANIMNODE(EvaluateComponentPose_AnyThread)
+	if (bFreezeIncomingPoseOnStart && bSimulationStarted)
+	{
+		// If we have a Frozen Pose captured, use it.
+		// Only after our intialize setup. As we need new pose for that.
+		if (ResetSimulatedTeleportType == ETeleportType::None && (CapturedFrozenPose.GetPose().GetNumBones() > 0))
+		{
+			Output.Pose.CopyPose(CapturedFrozenPose);
+			Output.Curve.CopyFrom(CapturedFrozenCurves);
+		}
+		// Otherwise eval graph to capture it.
+		else
+		{
+			Super::EvaluateComponentPose_AnyThread(Output);
+			CapturedFrozenPose.CopyPose(Output.Pose);
+			CapturedFrozenCurves.CopyFrom(Output.Curve);
+		}
+	}
+	else
+	{
+		Super::EvaluateComponentPose_AnyThread(Output);
+	}
+
+	// Capture incoming pose if 'bTransferBoneVelocities' is set.
+	// That is, until simulation starts.
+	if (bTransferBoneVelocities && !bSimulationStarted)
+	{
+		CapturedBoneVelocityPose.CopyPose(Output.Pose);
+		CapturedBoneVelocityPose.CopyAndAssignBoneContainer(CapturedBoneVelocityBoneContainer);
+	}
+}
+
+void FAnimNode_RigidBodyWithControl::InitializeNewBodyTransformsDuringSimulation(FComponentSpacePoseContext& Output, const FTransform& ComponentTransform, const FTransform& BaseBoneTM)
+{
+	for (const RigidBodyWithControl::FOutputBoneData& OutputData : OutputBoneData)
+	{
+		const int32 BodyIndex = OutputData.BodyIndex;
+		FBodyAnimData& BodyData = BodyAnimData[BodyIndex];
+		if (!BodyData.bBodyTransformInitialized)
+		{
+			BodyData.bBodyTransformInitialized = true;
+
+			// If we have a parent body, we need to grab relative transforms to it.
+			if (OutputData.ParentBodyIndex != INDEX_NONE)
+			{
+				ensure(BodyAnimData[OutputData.ParentBodyIndex].bBodyTransformInitialized);
+
+				FTransform BodyRelativeTransform = FTransform::Identity;
+				for (const FCompactPoseBoneIndex& CompactBoneIndex : OutputData.BoneIndicesToParentBody)
+				{
+					const FTransform& LocalSpaceTM = Output.Pose.GetLocalSpaceTransform(CompactBoneIndex);
+					BodyRelativeTransform = BodyRelativeTransform * LocalSpaceTM;
+				}
+
+				const FTransform WSBodyTM = BodyRelativeTransform * Bodies[OutputData.ParentBodyIndex]->GetWorldTransform();
+				Bodies[BodyIndex]->InitWorldTransform(WSBodyTM);
+				BodyAnimData[BodyIndex].RefPoseLength = (float) BodyRelativeTransform.GetLocation().Size();
+			}
+			// If we don't have a parent body, then we can just grab the incoming pose in component space.
+			else
+			{
+				const FTransform& ComponentSpaceTM = Output.Pose.GetComponentSpaceTransform(OutputData.CompactPoseBoneIndex);
+				const FTransform BodyTM = ConvertCSTransformToSimSpace(SimulationSpace, ComponentSpaceTM, ComponentTransform, BaseBoneTM);
+
+				Bodies[BodyIndex]->InitWorldTransform(BodyTM);
+			}
+		}
+	}
+}
+
+void FAnimNode_RigidBodyWithControl::InitSimulationSpace(
+	const FTransform& ComponentToWorld,
+	const FTransform& BoneToComponent)
+{
+	SimSpacePreviousComponentToWorld = ComponentToWorld;
+	SimSpacePreviousBoneToComponent = BoneToComponent;
+	SimSpacePreviousComponentLinearVelocity = FVector::ZeroVector;
+	SimSpacePreviousComponentAngularVelocity = FVector::ZeroVector;
+	SimSpacePreviousBoneLinearVelocity = FVector::ZeroVector;
+	SimSpacePreviousBoneAngularVelocity = FVector::ZeroVector;
+}
+
+void FAnimNode_RigidBodyWithControl::CalculateSimulationSpace(
+	ESimulationSpace Space, 
+	const FTransform& ComponentToWorld, 
+	const FTransform& BoneToComponent,
+	const float Dt,
+	const FSimSpaceSettings& Settings,
+	FTransform& SpaceTransform, 
+	FVector& SpaceLinearVel, 
+	FVector& SpaceAngularVel, 
+	FVector& SpaceLinearAcc, 
+	FVector& SpaceAngularAcc)
+{
+	// World-space transform of the simulation space
+	SpaceTransform = SpaceToWorldTransform(Space, ComponentToWorld, BoneToComponent);
+	SpaceLinearVel = FVector::ZeroVector;
+	SpaceAngularVel = FVector::ZeroVector;
+	SpaceLinearAcc = FVector::ZeroVector;
+	SpaceAngularAcc = FVector::ZeroVector;
+
+	// The simulation scale does not change - we scale the inputs and outputs instead.
+	// This means we do not support phantom forces resulting from scale changes, but that's ok.
+	// NOTE: If we don't clear the scale, rapid scaling to zero can introduce large phantom forces
+	// leading to major instability in the simulation
+	SpaceTransform.SetScale3D(FVector::One());
+
+	// If the system is disabled, nothing else to do
+	if ((Settings.WorldAlpha == 0.0f) || (Dt < SMALL_NUMBER))
+	{
+		return;
+	}
+
+	if (Space == ESimulationSpace::WorldSpace)
+	{
+		SpaceLinearVel = Settings.ExternalLinearVelocity;
+		SpaceAngularVel = Settings.ExternalAngularVelocity;
+		return;
+	}
+
+	// World-space component velocity and acceleration
+	FVector CompLinVel = Chaos::FVec3::CalculateVelocity(SimSpacePreviousComponentToWorld.GetTranslation(), ComponentToWorld.GetTranslation(), Dt);
+	FVector CompAngVel = Chaos::FRotation3::CalculateAngularVelocity(SimSpacePreviousComponentToWorld.GetRotation(), ComponentToWorld.GetRotation(), Dt);
+	FVector CompLinAcc = (CompLinVel - SimSpacePreviousComponentLinearVelocity) / Dt;
+	FVector CompAngAcc = (CompAngVel - SimSpacePreviousComponentAngularVelocity) / Dt;
+	SimSpacePreviousComponentToWorld = ComponentToWorld;
+	SimSpacePreviousComponentLinearVelocity = CompLinVel;
+	SimSpacePreviousComponentAngularVelocity = CompAngVel;
+
+	if (Space == ESimulationSpace::ComponentSpace)
+	{
+		CompLinVel.Z *= Settings.VelocityScaleZ;
+		CompLinAcc.Z *= Settings.VelocityScaleZ;
+
+		SpaceLinearVel = CompLinVel.GetClampedToMaxSize(Settings.MaxLinearVelocity) + Settings.ExternalLinearVelocity;
+		SpaceAngularVel = CompAngVel.GetClampedToMaxSize(Settings.MaxAngularVelocity) + Settings.ExternalAngularVelocity;
+		SpaceLinearAcc = CompLinAcc.GetClampedToMaxSize(Settings.MaxLinearAcceleration);
+		SpaceAngularAcc = CompAngAcc.GetClampedToMaxSize(Settings.MaxAngularAcceleration);
+		return;
+	}
+	
+	if (Space == ESimulationSpace::BaseBoneSpace)
+	{
+		// World-space component-relative bone velocity and acceleration
+		FVector BoneLinVel = Chaos::FVec3::CalculateVelocity(SimSpacePreviousBoneToComponent.GetTranslation(), BoneToComponent.GetTranslation(), Dt);
+		FVector BoneAngVel = Chaos::FRotation3::CalculateAngularVelocity(SimSpacePreviousBoneToComponent.GetRotation(), BoneToComponent.GetRotation(), Dt);
+		BoneLinVel = ComponentToWorld.TransformVector(BoneLinVel);
+		BoneAngVel = ComponentToWorld.TransformVector(BoneAngVel);
+		FVector BoneLinAcc = (BoneLinVel - SimSpacePreviousBoneLinearVelocity) / Dt;
+		FVector BoneAngAcc = (BoneAngVel - SimSpacePreviousBoneAngularVelocity) / Dt;
+		SimSpacePreviousBoneToComponent = BoneToComponent;
+		SimSpacePreviousBoneLinearVelocity = BoneLinVel;
+		SimSpacePreviousBoneAngularVelocity = BoneAngVel;
+
+		// World-space bone velocity and acceleration
+		FVector NetAngVel = CompAngVel + BoneAngVel;
+		FVector NetAngAcc = CompAngAcc + BoneAngAcc;
+
+		// If we limit the angular velocity, we also need to limit the component of linear velocity that comes from (angvel x offset)
+		float AngVelScale = 1.0f;
+		float NetAngVelLenSq = (float) NetAngVel.SizeSquared();
+		if (NetAngVelLenSq > FMath::Square(Settings.MaxAngularVelocity))
+		{
+			AngVelScale = Settings.MaxAngularVelocity * FMath::InvSqrt(NetAngVelLenSq);
+		}
+
+		// Add the linear velocity and acceleration that comes from rotation of the space about the component
+		// NOTE: Component angular velocity constribution is scaled
+		FVector SpaceCompOffset = ComponentToWorld.TransformVector(BoneToComponent.GetTranslation());
+		FVector NetLinVel = CompLinVel + BoneLinVel + FVector::CrossProduct(AngVelScale * CompAngVel, SpaceCompOffset);
+		FVector NetLinAcc = CompLinAcc + BoneLinAcc + FVector::CrossProduct(AngVelScale * CompAngAcc, SpaceCompOffset);
+
+		NetLinVel.Z *= Settings.VelocityScaleZ;
+		NetLinAcc.Z *= Settings.VelocityScaleZ;
+
+		SpaceLinearVel = NetLinVel.GetClampedToMaxSize(Settings.MaxLinearVelocity) + Settings.ExternalLinearVelocity;
+		SpaceAngularVel = NetAngVel.GetClampedToMaxSize(Settings.MaxAngularVelocity) + Settings.ExternalAngularVelocity;
+		SpaceLinearAcc = NetLinAcc.GetClampedToMaxSize(Settings.MaxLinearAcceleration);
+		SpaceAngularAcc = NetAngAcc.GetClampedToMaxSize(Settings.MaxAngularAcceleration);
+		return;
+	}
+}
+
+
+DECLARE_CYCLE_STAT(TEXT("RigidBody_Eval"), STAT_RigidBody_Eval, STATGROUP_Anim);
+DECLARE_CYCLE_STAT(TEXT("RigidBodyNodeWithControl_Simulation"), STAT_RigidBodyNodeWithControl_Simulation, STATGROUP_Anim);
+DECLARE_CYCLE_STAT(TEXT("RigidBodyNodeWithControl_SimulationWait"), STAT_RigidBodyNodeWithControl_SimulationWait, STATGROUP_Anim);
+DECLARE_CYCLE_STAT(TEXT("FAnimNode_RigidBodyWithControl::EvaluateSkeletalControl_AnyThread"), STAT_ImmediateEvaluateSkeletalControl, STATGROUP_ImmediatePhysics);
+
+void FAnimNode_RigidBodyWithControl::RunPhysicsSimulation(float DeltaSeconds, const FVector& SimSpaceGravity)
+{
+	SCOPE_CYCLE_COUNTER(STAT_RigidBodyNodeWithControl_Simulation);
+	CSV_SCOPED_TIMING_STAT(Animation, RigidBodyNodeWithControlSimulation);
+	FScopeCycleCounterUObject AdditionalScope(PhysicsAssetToUse, GET_STATID(STAT_RigidBodyNodeWithControl_Simulation));
+
+	const int32 MaxSteps = RBANWithControl_MaxSubSteps;
+	const float MaxDeltaSeconds = 1.f / 30.f;
+
+	PhysicsSimulation->Simulate_AssumesLocked(DeltaSeconds, MaxDeltaSeconds, MaxSteps, SimSpaceGravity);
+	bSimulationStarted = true;
+}
+
+void FAnimNode_RigidBodyWithControl::FlushDeferredSimulationTask()
+{
+	if (SimulationTask.IsValid() && !SimulationTask.IsCompleted())
+	{
+		SCOPE_CYCLE_COUNTER(STAT_RigidBodyNodeWithControl_SimulationWait);
+		CSV_SCOPED_TIMING_STAT(Animation, RigidBodyNodeWithControlSimulationWait);
+		SimulationTask.Wait();
+	}
+}
+
+void FAnimNode_RigidBodyWithControl::DestroyPhysicsSimulation()
+{
+	ClothColliders.Reset();
+	FlushDeferredSimulationTask();
+	delete PhysicsSimulation;
+	PhysicsSimulation = nullptr;
+}
+
+static FQuat CalculateTargetQuat(
+	const Chaos::FPBDJointSettings& Settings, 
+	const FTransform&               ChildBoneTM,
+	const FTransform&               ParentBoneTM)
+{
+	FQuat ChildBoneQuat = ChildBoneTM.GetRotation();
+	FQuat ParentBoneQuat = ParentBoneTM.GetRotation();
+	FQuat ChildConnectorQuat = Settings.ConnectorTransforms[0].GetRotation();
+	FQuat ParentConnectorQuat = Settings.ConnectorTransforms[1].GetRotation();
+	FQuat LimitQuat = ParentBoneQuat * ParentConnectorQuat;
+	FQuat DialQuat = ChildBoneQuat * ChildConnectorQuat;
+	return LimitQuat.Inverse() * DialQuat;
+}
+
+static FTransform CalculateJointTargetTransform(
+	const Chaos::FPBDJointSettings& Settings,
+	const FTransform& ChildBoneTM,
+	const FTransform& ParentBoneTM)
+{
+	const FTransform& ChildLocalConnectorTM = Settings.ConnectorTransforms[0];
+	const FTransform& ParentLocalConnectorTM = Settings.ConnectorTransforms[1];
+	const FTransform ChildConnectorTM = ChildBoneTM * ChildLocalConnectorTM;
+	const FTransform ParentConnectorTM = ParentBoneTM * ParentLocalConnectorTM;
+
+	return ParentConnectorTM.Inverse() * ChildConnectorTM;
+}
+
+
+void FAnimNode_RigidBodyWithControl::EvaluateSkeletalControl_AnyThread(FComponentSpacePoseContext& Output, TArray<FBoneTransform>& OutBoneTransforms)
+{
+	DECLARE_SCOPE_HIERARCHICAL_COUNTER_ANIMNODE(EvaluateSkeletalControl_AnyThread)
+	SCOPE_CYCLE_COUNTER(STAT_RigidBody_Eval);
+	CSV_SCOPED_TIMING_STAT(Animation, RigidBodyEval);
+	SCOPE_CYCLE_COUNTER(STAT_ImmediateEvaluateSkeletalControl);
+	//SCOPED_NAMED_EVENT_TEXT("FAnimNode_RigidBodyWithControl::EvaluateSkeletalControl_AnyThread", FColor::Magenta);
+
+	if (CVarEnableRigidBodyNodeWithControlSimulation.GetValueOnAnyThread() == 0)
+	{
+		return;
+	}
+
+	const float DeltaSeconds = AccumulatedDeltaTime;
+	AccumulatedDeltaTime = 0.f;
+
+	if (bEnabled && PhysicsSimulation)	
+	{
+		FlushDeferredSimulationTask();
+
+		// Handle deferred control creation
+		if (!bHaveRunControlSeup && bEnableControls)
+		{
+			CreateWorldSpaceControlRootBody(PhysicsAssetToUse);
+			bHaveRunControlSeup = true;
+			InitControlsAndBodyModifiers(Output.AnimInstanceProxy->GetSkelMeshComponent());
+		}
+
+		const FBoneContainer& BoneContainer = Output.Pose.GetPose().GetBoneContainer();
+		const FTransform CompWorldSpaceTM = Output.AnimInstanceProxy->GetComponentTransform();
+
+		bool bFirstEvalSinceReset = !Output.AnimInstanceProxy->GetEvaluationCounter().HasEverBeenUpdated();
+
+		// First-frame initialization
+		if (bFirstEvalSinceReset)
+		{
+			PreviousCompWorldSpaceTM = CompWorldSpaceTM;
+			ResetSimulatedTeleportType = ETeleportType::ResetPhysics;
+		}
+
+		// See if we need to reset physics because too much time passed since our last update (e.g., because we we off-screen for a while), 
+		// in which case the current sim state may be too far from the current anim pose. This is mostly a problem with world-space 
+		// simulation, whereas bone- and component-space sims can be fairly robust against missing updates.
+		// Don't do this on first frame or if time-based reset is disabled. 
+		if ((EvaluationResetTime > 0.0f) && !bFirstEvalSinceReset)
+		{
+			// NOTE: under normal conditions, when this anim node is being serviced at the usual rate (which may not be every frame
+			// if URO is enabled), we expect that WorldTimeSeconds == (LastEvalTimeSeconds + DeltaSeconds). DeltaSeconds is the 
+			// accumulated time since the last update, including frames dropped by URO, but not frames dropped because of
+			// being off-screen or LOD changes.
+			if (WorldTimeSeconds - (LastEvalTimeSeconds + DeltaSeconds) > EvaluationResetTime)
+			{
+				UE_LOG(LogRigidBodyWithControl, Verbose, TEXT("%s Time-Based Reset"), *Output.AnimInstanceProxy->GetAnimInstanceName());
+				ResetSimulatedTeleportType = ETeleportType::ResetPhysics;
+			}
+		}
+
+		// Update the evaluation time to the current time
+		LastEvalTimeSeconds = WorldTimeSeconds;
+
+		// Disable simulation below minimum scale in world space mode. World space sim doesn't play nice with scale anyway - we do not scale joint offets or collision shapes.
+		if ((SimulationSpace == ESimulationSpace::WorldSpace) && (CompWorldSpaceTM.GetScale3D().SizeSquared() < WorldSpaceMinimumScale * WorldSpaceMinimumScale))
+		{
+			return;
+		}
+
+		const FTransform BaseBoneTM = Output.Pose.GetComponentSpaceTransform(BaseBoneRef.GetCompactPoseIndex(BoneContainer));
+
+		// Initialize potential new bodies because of LOD change.
+		if (ResetSimulatedTeleportType == ETeleportType::None && bCheckForBodyTransformInit)
+		{
+			bCheckForBodyTransformInit = false;
+			InitializeNewBodyTransformsDuringSimulation(Output, CompWorldSpaceTM, BaseBoneTM);
+		}
+
+		// If time advances, update simulation
+		// Reset if necessary
+		const bool bResetOrTeleportBodies = (ResetSimulatedTeleportType != ETeleportType::None);
+		if (bResetOrTeleportBodies)
+		{
+			// Capture bone velocities if we have captured a bone velocity pose.
+			if (bTransferBoneVelocities && (CapturedBoneVelocityPose.GetPose().GetNumBones() > 0))
+			{
+				for (const RigidBodyWithControl::FOutputBoneData& OutputData : OutputBoneData)
+				{
+					const int32 BodyIndex = OutputData.BodyIndex;
+					FBodyAnimData& BodyData = BodyAnimData[BodyIndex];
+
+					if (BodyData.bIsSimulated)
+					{
+						const FCompactPoseBoneIndex NextCompactPoseBoneIndex = OutputData.CompactPoseBoneIndex;
+						// Convert CompactPoseBoneIndex to SkeletonBoneIndex...
+						const FSkeletonPoseBoneIndex PoseSkeletonBoneIndex = BoneContainer.GetSkeletonPoseIndexFromCompactPoseIndex(NextCompactPoseBoneIndex);
+						// ... So we can convert to the captured pose CompactPoseBoneIndex. 
+						// In case there was a LOD change, and poses are not compatible anymore.
+						const FCompactPoseBoneIndex PrevCompactPoseBoneIndex = CapturedBoneVelocityBoneContainer.GetCompactPoseIndexFromSkeletonPoseIndex(PoseSkeletonBoneIndex);
+
+						if (PrevCompactPoseBoneIndex.IsValid())
+						{
+							const FTransform PrevCSTM = CapturedBoneVelocityPose.GetComponentSpaceTransform(PrevCompactPoseBoneIndex);
+							const FTransform NextCSTM = Output.Pose.GetComponentSpaceTransform(NextCompactPoseBoneIndex);
+
+							const FTransform PrevSSTM = ConvertCSTransformToSimSpace(SimulationSpace, PrevCSTM, CompWorldSpaceTM, BaseBoneTM);
+							const FTransform NextSSTM = ConvertCSTransformToSimSpace(SimulationSpace, NextCSTM, CompWorldSpaceTM, BaseBoneTM);
+
+							if(DeltaSeconds > 0.0f)
+							{
+								// Linear Velocity
+								BodyData.TransferedBoneLinearVelocity = ((NextSSTM.GetLocation() - PrevSSTM.GetLocation()) / DeltaSeconds);
+								
+								// Angular Velocity
+								const FQuat DeltaRotation = (NextSSTM.GetRotation().Inverse() * PrevSSTM.GetRotation());
+								const float RotationAngle = (float) DeltaRotation.GetAngle() / DeltaSeconds;
+								BodyData.TransferedBoneAngularVelocity = (FQuat(DeltaRotation.GetRotationAxis(), RotationAngle)); 
+							}
+							else
+							{
+								BodyData.TransferedBoneLinearVelocity = (FVector::ZeroVector);
+								BodyData.TransferedBoneAngularVelocity = (FQuat::Identity); 
+							}
+
+						}
+					}
+				}
+			}
+
+			
+			switch(ResetSimulatedTeleportType)
+			{
+				case ETeleportType::TeleportPhysics:
+				{
+					UE_LOG(LogRigidBodyWithControl, Verbose, TEXT("%s TeleportPhysics (Scale: %f %f %f)"), *Output.AnimInstanceProxy->GetAnimInstanceName(), CompWorldSpaceTM.GetScale3D().X, CompWorldSpaceTM.GetScale3D().Y, CompWorldSpaceTM.GetScale3D().Z);
+
+					// Teleport bodies.
+					for (const RigidBodyWithControl::FOutputBoneData& OutputData : OutputBoneData)
+					{
+						const int32 BodyIndex = OutputData.BodyIndex;
+						BodyAnimData[BodyIndex].bBodyTransformInitialized = true;
+
+						FTransform BodyTM = Bodies[BodyIndex]->GetWorldTransform();
+						FTransform ComponentSpaceTM;
+
+						switch(SimulationSpace)
+						{
+							case ESimulationSpace::ComponentSpace: ComponentSpaceTM = BodyTM; break;
+							case ESimulationSpace::WorldSpace: ComponentSpaceTM = BodyTM.GetRelativeTransform(PreviousCompWorldSpaceTM); break;
+							case ESimulationSpace::BaseBoneSpace: ComponentSpaceTM = BodyTM * BaseBoneTM; break;
+							default: ensureMsgf(false, TEXT("Unsupported Simulation Space")); ComponentSpaceTM = BodyTM;
+						}
+
+						BodyTM = ConvertCSTransformToSimSpace(SimulationSpace, ComponentSpaceTM, CompWorldSpaceTM, BaseBoneTM);
+						Bodies[BodyIndex]->SetWorldTransform(BodyTM);
+						if (OutputData.ParentBodyIndex != INDEX_NONE)
+						{
+							BodyAnimData[BodyIndex].RefPoseLength = (float) 
+								BodyTM.GetRelativeTransform(Bodies[OutputData.ParentBodyIndex]->GetWorldTransform()).GetLocation().Size();
+						}
+					}
+				}
+				break;
+
+				case ETeleportType::ResetPhysics:
+				{
+					UE_LOG(LogRigidBodyWithControl, Verbose, TEXT("%s ResetPhysics (Scale: %f %f %f)"), *Output.AnimInstanceProxy->GetAnimInstanceName(), CompWorldSpaceTM.GetScale3D().X, CompWorldSpaceTM.GetScale3D().Y, CompWorldSpaceTM.GetScale3D().Z);
+
+					InitSimulationSpace(CompWorldSpaceTM, BaseBoneTM);
+
+					// Completely reset bodies.
+					for (const RigidBodyWithControl::FOutputBoneData& OutputData : OutputBoneData)
+					{
+						const int32 BodyIndex = OutputData.BodyIndex;
+						BodyAnimData[BodyIndex].bBodyTransformInitialized = true;
+
+						const FTransform& ComponentSpaceTM = Output.Pose.GetComponentSpaceTransform(OutputData.CompactPoseBoneIndex);
+						const FTransform BodyTM = ConvertCSTransformToSimSpace(SimulationSpace, ComponentSpaceTM, CompWorldSpaceTM, BaseBoneTM);
+						Bodies[BodyIndex]->InitWorldTransform(BodyTM);
+						if (OutputData.ParentBodyIndex != INDEX_NONE)
+						{
+							BodyAnimData[BodyIndex].RefPoseLength = (float)
+								BodyTM.GetRelativeTransform(Bodies[OutputData.ParentBodyIndex]->GetWorldTransform()).GetLocation().Size();
+						}
+					}
+				}
+				break;
+			}
+
+			// Always reset after a teleport
+			PreviousCompWorldSpaceTM = CompWorldSpaceTM;
+			PreviousComponentLinearVelocity = FVector::ZeroVector;
+		}
+
+		// Assets can override config for deferred simulation
+		const bool bUseDeferredSimulationTask =
+			(SimulationTiming == ESimulationTiming::Deferred) ||
+			((SimulationTiming == ESimulationTiming::Default) && bRBANWithControl_DeferredSimulationDefault);
+		FVector SimSpaceGravity(0.f);
+
+		// Only need to tick physics if we didn't reset and we have some time to simulate
+		const bool bNeedsSimulationTick = 
+			((bSimulateAnimPhysicsAfterReset || (ResetSimulatedTeleportType != ETeleportType::ResetPhysics)) && 
+				DeltaSeconds > AnimPhysicsMinDeltaTime);
+		if (bNeedsSimulationTick)
+		{
+			// Update the pose data
+			PoseData.Update(
+				Output, OutputBoneData, SimulationSpace, BaseBoneRef, Output.AnimInstanceProxy->GetUpdateCounter());
+
+			// Transfer bone velocities previously captured.
+			if (bTransferBoneVelocities && (CapturedBoneVelocityPose.GetPose().GetNumBones() > 0))
+			{
+				for (const RigidBodyWithControl::FOutputBoneData& OutputData : OutputBoneData)
+				{
+					const int32 BodyIndex = OutputData.BodyIndex;
+					const FBodyAnimData& BodyData = BodyAnimData[BodyIndex];
+
+					if (BodyData.bIsSimulated)
+					{
+						ImmediatePhysics::FActorHandle* Body = Bodies[BodyIndex];
+						Body->SetLinearVelocity(BodyData.TransferedBoneLinearVelocity);
+
+						const FQuat AngularVelocity = BodyData.TransferedBoneAngularVelocity;
+						Body->SetAngularVelocity(AngularVelocity.GetRotationAxis() * AngularVelocity.GetAngle());
+					}
+				}
+
+				// Free up our captured pose after it's been used.
+				CapturedBoneVelocityPose.Empty();
+			}
+			else if ((SimulationSpace != ESimulationSpace::WorldSpace) && bRBANWithControl_EnableComponentAcceleration)
+			{
+				if (!ComponentLinearVelScale.IsNearlyZero() || !ComponentLinearAccScale.IsNearlyZero())
+				{
+					// Calc linear velocity
+					const FVector ComponentDeltaLocation = CurrentTransform.GetTranslation() - PreviousTransform.GetTranslation();
+					const FVector ComponentLinearVelocity = ComponentDeltaLocation / DeltaSeconds;
+					// Apply acceleration that opposed velocity (basically 'drag')
+					FVector ApplyLinearAcc = WorldVectorToSpaceNoScale(SimulationSpace, -ComponentLinearVelocity, CompWorldSpaceTM, BaseBoneTM) * ComponentLinearVelScale;
+
+					// Calc linear acceleration
+					const FVector ComponentLinearAcceleration = (ComponentLinearVelocity - PreviousComponentLinearVelocity) / DeltaSeconds;
+					PreviousComponentLinearVelocity = ComponentLinearVelocity;
+					// Apply opposite acceleration to bodies
+					ApplyLinearAcc += WorldVectorToSpaceNoScale(SimulationSpace, -ComponentLinearAcceleration, CompWorldSpaceTM, BaseBoneTM) * ComponentLinearAccScale;
+
+					// Iterate over bodies
+					for (const RigidBodyWithControl::FOutputBoneData& OutputData : OutputBoneData)
+					{
+						const int32 BodyIndex = OutputData.BodyIndex;
+						const FBodyAnimData& BodyData = BodyAnimData[BodyIndex];
+
+						if (BodyData.bIsSimulated)
+						{
+							ImmediatePhysics::FActorHandle* Body = Bodies[BodyIndex];
+
+							// Apply 
+							const float BodyInvMass = (float) Body->GetInverseMass();
+							if (BodyInvMass > 0.f)
+							{
+								// Final desired acceleration to apply to body
+								FVector FinalBodyLinearAcc = ApplyLinearAcc;
+
+								// Clamp if desired
+								if (!ComponentAppliedLinearAccClamp.IsNearlyZero())
+								{
+									FinalBodyLinearAcc = FinalBodyLinearAcc.BoundToBox(-ComponentAppliedLinearAccClamp, ComponentAppliedLinearAccClamp);
+								}
+
+								// Apply to body
+								Body->AddForce(FinalBodyLinearAcc / BodyInvMass);
+							}
+						}
+					}
+				}
+			}
+
+			UpdateWorldForces(CompWorldSpaceTM, BaseBoneTM, DeltaSeconds);
+			SimSpaceGravity = WorldVectorToSpaceNoScale(SimulationSpace, WorldSpaceGravity, CompWorldSpaceTM, BaseBoneTM);
+
+			// Adjust the constraint profiles when it is changed
+			if (ConstraintProfile != CurrentConstraintProfile)
+			{
+				CurrentConstraintProfile = ConstraintProfile;
+				ApplyCurrentConstraintProfile();
+			}
+
+			// Apply the controls. Note that these won't set kinematic targets - we'll do that afterwards as they
+			// can apply to bodies that aren't under the influence of a modifier
+			ApplyControlAndModifierUpdatesAndParametersToRecords(ControlAndModifierUpdates, ControlAndModifierParameters);
+			ApplyControlsAndModifiers(SimSpaceGravity, DeltaSeconds);
+
+			// @todo(ccaulfield): We should be interpolating kinematic targets for each sub-step below
+			for (const RigidBodyWithControl::FOutputBoneData& OutputData : OutputBoneData)
+			{
+				const int32 BodyIndex = OutputData.BodyIndex;
+				BodyAnimData[BodyIndex].bIsSimulated = !Bodies[BodyIndex]->GetIsKinematic();
+				if (!BodyAnimData[BodyIndex].bIsSimulated)
+				{
+					// TODO support explicit targets for kinematics, like we do for controls
+					FTransform TM = PoseData.GetTM(BodyIndex);
+					Bodies[BodyIndex]->SetKinematicTarget(TM);
+				}
+			}
+
+			// Note that this must be called after the "regular" kinematic targets above, as it will
+			// overwrite any that have an explicit target set (which is probably quite unusual, so
+			// the duplication is not likely to be expensive)
+			ApplyKinematicTargets();
+
+			FSimSpaceSettings* UseSimSpaceSettings = &SimSpaceSettings;
+			if (bRBANWithControl_SimSpace_EnableOverride)
+			{
+				UseSimSpaceSettings = &RBANWithControl_SimSpaceOverride;
+			}
+
+			FTransform SimulationTransform;
+			FVector SimulationLinearVelocity;
+			FVector SimulationAngularVelocity;
+			FVector SimulationLinearAcceleration;
+			FVector SimulationAngularAcceleration;
+			CalculateSimulationSpace(
+				SimulationSpace, 
+				CompWorldSpaceTM, 
+				BaseBoneTM,
+				DeltaSeconds,
+				*UseSimSpaceSettings,
+				SimulationTransform,
+				SimulationLinearVelocity,
+				SimulationAngularVelocity,
+				SimulationLinearAcceleration,
+				SimulationAngularAcceleration);
+
+			UpdateWorldObjects(SimulationTransform);
+			UpdateClothColliderObjects(SimulationTransform);
+
+			PhysicsSimulation->UpdateSimulationSpace(
+				SimulationTransform, 
+				SimulationLinearVelocity,
+				SimulationAngularVelocity,
+				SimulationLinearAcceleration,
+				SimulationAngularAcceleration);
+
+			PhysicsSimulation->SetSimulationSpaceSettings(
+				UseSimSpaceSettings->WorldAlpha,
+				UseSimSpaceSettings->ExternalLinearDragV);
+
+			PhysicsSimulation->SetSolverSettings(
+				SolverSettings.FixedTimeStep,
+				SolverSettings.CullDistance,
+				SolverSettings.MaxDepenetrationVelocity,
+				SolverSettings.bUseLinearJointSolver,
+				SolverSettings.PositionIterations,
+				SolverSettings.VelocityIterations,
+				SolverSettings.ProjectionIterations);
+
+			if (!bUseDeferredSimulationTask)
+			{
+				RunPhysicsSimulation(DeltaSeconds, SimSpaceGravity);
+			}
+
+			// Draw here even if the simulation is deferred since we want the shapes drawn relative to the current transform
+			if (bRBANWithControl_DebugDraw)
+			{
+				PhysicsSimulation->DebugDraw();
+			}
+		}
+		
+		//write back to animation system
+		const FTransform& SimulationWorldSpaceTM = bUseDeferredSimulationTask ? PreviousCompWorldSpaceTM : CompWorldSpaceTM;
+		for (const RigidBodyWithControl::FOutputBoneData& OutputData : OutputBoneData)
+		{
+			const int32 BodyIndex = OutputData.BodyIndex;
+			// Note that we always read back, whether kinematic or simulated
+			FTransform BodyTM = Bodies[BodyIndex]->GetWorldTransform();
+
+			// if we clamp translation, we only do this when all linear translation are locked
+			// 
+			// @todo(ccaulfield): this shouldn't be required with Chaos - projection should be handling it...
+			if (bClampLinearTranslationLimitToRefPose
+				&&BodyAnimData[BodyIndex].LinearXMotion == ELinearConstraintMotion::LCM_Locked
+				&& BodyAnimData[BodyIndex].LinearYMotion == ELinearConstraintMotion::LCM_Locked
+				&& BodyAnimData[BodyIndex].LinearZMotion == ELinearConstraintMotion::LCM_Locked)
+			{
+				// grab local space of length from ref pose 
+				// we have linear limit value - see if that works
+				// calculate current local space from parent
+				// find parent transform
+				const int32 ParentBodyIndex = OutputData.ParentBodyIndex;
+				FTransform ParentTransform = FTransform::Identity;
+				if (ParentBodyIndex != INDEX_NONE)
+				{
+					ParentTransform = Bodies[ParentBodyIndex]->GetWorldTransform();
+				}
+
+				// get local transform
+				FTransform LocalTransform = BodyTM.GetRelativeTransform(ParentTransform);
+				const float CurrentLength = (float) LocalTransform.GetTranslation().Size();
+
+				// this is inconsistent with constraint. The actual linear limit is set by constraint
+				if (!FMath::IsNearlyEqual(CurrentLength, BodyAnimData[BodyIndex].RefPoseLength, KINDA_SMALL_NUMBER))
+				{
+					float RefPoseLength = BodyAnimData[BodyIndex].RefPoseLength;
+					if (CurrentLength > RefPoseLength)
+					{
+						float Scale = (CurrentLength > KINDA_SMALL_NUMBER) ? RefPoseLength / CurrentLength : 0.f;
+						// we don't use 1.f here because 1.f can create pops based on float issue. 
+						// so we only activate clamping when less than 90%
+						if (Scale < 0.9f)
+						{
+							LocalTransform.ScaleTranslation(Scale);
+							BodyTM = LocalTransform * ParentTransform;
+							Bodies[BodyIndex]->SetWorldTransform(BodyTM);
+						}
+					}
+				}
+			}
+
+			FTransform ComponentSpaceTM;
+
+			switch(SimulationSpace)
+			{
+				case ESimulationSpace::ComponentSpace: ComponentSpaceTM = BodyTM; break;
+				case ESimulationSpace::WorldSpace: ComponentSpaceTM = BodyTM.GetRelativeTransform(SimulationWorldSpaceTM); break;
+				case ESimulationSpace::BaseBoneSpace: ComponentSpaceTM = BodyTM * BaseBoneTM; break;
+				default: ensureMsgf(false, TEXT("Unsupported Simulation Space")); ComponentSpaceTM = BodyTM;
+			}
+					
+			OutBoneTransforms.Add(FBoneTransform(OutputData.CompactPoseBoneIndex, ComponentSpaceTM));
+		}
+
+		// Deferred task must be started after we read actor poses to avoid a race
+		if (bNeedsSimulationTick && bUseDeferredSimulationTask)
+		{
+			// FlushDeferredSimulationTask() should have already ensured task is done.
+			ensure(SimulationTask.IsCompleted());
+			const int32 PriorityIndex = FMath::Clamp<int32>(GRigidBodyNodeWithControlSimulationTaskPriority, 0, UE_ARRAY_COUNT(GRigidBodyNodeWithControlTaskPriorities) - 1);
+			const UE::Tasks::ETaskPriority TaskPriority = GRigidBodyNodeWithControlTaskPriorities[PriorityIndex];
+			SimulationTask = UE::Tasks::Launch(
+				TEXT("RigidBodyNodeWithControlSimulationTask"),
+				[this, DeltaSeconds, SimSpaceGravity] { RunPhysicsSimulation(DeltaSeconds, SimSpaceGravity); },
+				TaskPriority);
+		}
+
+		PreviousCompWorldSpaceTM = CompWorldSpaceTM;
+		ResetSimulatedTeleportType = ETeleportType::None;
+	}
+}
+
+void ComputeBodyInsertionOrderWithControl(TArray<FBoneIndexType>& InsertionOrder, const USkeletalMeshComponent& SKC)
+{
+	//We want to ensure simulated bodies are sorted by LOD so that the first simulated bodies are at the highest LOD.
+	//Since LOD2 is a subset of LOD1 which is a subset of LOD0 we can change the number of simulated bodies without any reordering
+	//For this to work we must first insert all simulated bodies in the right order. We then insert all the kinematic bodies in the right order
+
+	InsertionOrder.Reset();
+
+	if (SKC.GetSkeletalMeshAsset() == nullptr)
+	{
+		return;
+	}
+
+	const int32 NumLODs = SKC.GetNumLODs();
+	if(NumLODs > 0)
+	{
+		TArray<FBoneIndexType> RequiredBones0;
+		TArray<FBoneIndexType> ComponentSpaceTMs0;
+		SKC.ComputeRequiredBones(RequiredBones0, ComponentSpaceTMs0, 0, /*bIgnorePhysicsAsset=*/ true);
+
+		TArray<bool> InSortedOrder;
+		InSortedOrder.AddZeroed(SKC.GetSkeletalMeshAsset()->GetRefSkeleton().GetNum());
+
+		auto MergeIndices = [&InsertionOrder, &InSortedOrder](const TArray<FBoneIndexType>& RequiredBones) -> void
+		{
+			for (FBoneIndexType BoneIdx : RequiredBones)
+			{
+				if (!InSortedOrder[BoneIdx])
+				{
+					InsertionOrder.Add(BoneIdx);
+				}
+
+				InSortedOrder[BoneIdx] = true;
+			}
+		};
+
+
+		for(int32 LodIdx = NumLODs - 1; LodIdx > 0; --LodIdx)
+		{
+			TArray<FBoneIndexType> RequiredBones;
+			TArray<FBoneIndexType> ComponentSpaceTMs;
+			SKC.ComputeRequiredBones(RequiredBones, ComponentSpaceTMs, LodIdx, /*bIgnorePhysicsAsset=*/ true);
+			MergeIndices(RequiredBones);
+		}
+
+		MergeIndices(RequiredBones0);
+	}
+}
+
+void FAnimNode_RigidBodyWithControl::InitPhysics(const UAnimInstance* InAnimInstance)
+{
+	LLM_SCOPE_BYNAME(TEXT("Animation/RigidBodyWithControl")); 
+	
+//	SCOPE_CYCLE_COUNTER(STAT_RigidBodyNodeInitTime);
+
+	DestroyPhysicsSimulation();
+
+	const USkeletalMeshComponent* SkeletalMeshComp = InAnimInstance->GetSkelMeshComponent();
+	const USkeletalMesh* SkeletalMeshAsset = SkeletalMeshComp->GetSkeletalMeshAsset();
+	USkeleton* SkeletonAsset = InAnimInstance->CurrentSkeleton;
+
+	if(!SkeletalMeshAsset || !SkeletonAsset)
+	{
+		// Without both the skeleton and the mesh we can't create a new simulation.
+		// The previous simulation has just been cleaned up above so we can return early here and not instantiate a new one
+		return;
+	}
+
+	const FReferenceSkeleton& SkelMeshRefSkel = SkeletalMeshAsset->GetRefSkeleton();
+	PhysicsAssetToUse = OverridePhysicsAsset ? ToRawPtr(OverridePhysicsAsset) : InAnimInstance->GetSkelMeshComponent()->GetPhysicsAsset();
+
+	ensure(SkeletonAsset == SkeletalMeshAsset->GetSkeleton());
+
+	const int32 SkelMeshLinkupIndex = SkeletonAsset->GetMeshLinkupIndex(SkeletalMeshAsset);
+	ensure(SkelMeshLinkupIndex != INDEX_NONE);
+	const FSkeletonToMeshLinkup& SkeletonToMeshLinkupTable = SkeletonAsset->LinkupCache[SkelMeshLinkupIndex];
+	const TArray<int32>& MeshToSkeletonBoneIndex = SkeletonToMeshLinkupTable.MeshToSkeletonTable;
+	
+	const int32 NumSkeletonBones = SkeletonAsset->GetReferenceSkeleton().GetNum();
+	SkeletonBoneIndexToBodyIndex.Reset(NumSkeletonBones);
+	SkeletonBoneIndexToBodyIndex.Init(INDEX_NONE, NumSkeletonBones);
+
+	PreviousTransform = SkeletalMeshComp->GetComponentToWorld();
+
+	ComponentsInSim.Reset();
+	ComponentsInSimTick = 0;
+
+	if (UPhysicsSettings* Settings = UPhysicsSettings::Get())
+	{
+		AnimPhysicsMinDeltaTime = Settings->AnimPhysicsMinDeltaTime;
+		bSimulateAnimPhysicsAfterReset = Settings->bSimulateAnimPhysicsAfterReset;
+	}
+	else
+	{
+		AnimPhysicsMinDeltaTime = 0.f;
+		bSimulateAnimPhysicsAfterReset = false;
+	}
+	
+	bEnabled = PhysicsAssetToUse && SkeletalMeshComp->GetAllowRigidBodyAnimNode() && CVarEnableRigidBodyNodeWithControl.GetValueOnAnyThread() != 0;
+	if(bEnabled)
+	{
+		PhysicsSimulation = new ImmediatePhysics::FSimulation();
+		const int32 NumBodies = PhysicsAssetToUse->SkeletalBodySetups.Num();
+		BodyNameToIndexMap.Reset();
+		Bodies.Empty(NumBodies);
+		Joints.Empty(NumBodies);
+		WorldSpaceControlActorHandle = nullptr;
+		BodyAnimData.Reset(NumBodies);
+		BodyAnimData.AddDefaulted(NumBodies);
+		TotalMass = 0.f;
+
+		DestroyControlsAndBodyModifiers();
+
+		// Instantiate a FBodyInstance/FConstraintInstance set that will be cloned into the Immediate Physics sim.
+		// NOTE: We do not have a skeleton at the moment, so we have to use the ref pose
+		TArray<FBodyInstance*> HighLevelBodyInstances;
+		TArray<FConstraintInstance*> HighLevelConstraintInstances;
+
+		// Chaos relies on the initial pose to set up constraint positions
+		constexpr bool bCreateBodiesInRefPose = true;
+		SkeletalMeshComp->InstantiatePhysicsAssetRefPose(
+			*PhysicsAssetToUse, 
+			SimulationSpace == ESimulationSpace::WorldSpace ? SkeletalMeshComp->GetComponentToWorld().GetScale3D() : FVector(1.f), 
+			HighLevelBodyInstances, 
+			HighLevelConstraintInstances, 
+			nullptr, 
+			nullptr, 
+			INDEX_NONE, 
+			FPhysicsAggregateHandle(),
+			bCreateBodiesInRefPose);
+
+		TMap<FName, ImmediatePhysics::FActorHandle*> NamesToHandles;
+		TArray<ImmediatePhysics::FActorHandle*> IgnoreCollisionActors;
+
+		// Map bone names to the joint connecting them to the bone's parent
+		TMap<FName, ImmediatePhysics::FJointHandle*> NamesToJointHandles;
+
+		TArray<FBoneIndexType> InsertionOrder;
+		ComputeBodyInsertionOrderWithControl(InsertionOrder, *SkeletalMeshComp);
+
+		// NOTE: NumBonesLOD0 may be less than NumBonesTotal, and it may be middle bones that are missing from LOD0.
+		// In this case, LOD0 bone indices may be >= NumBonesLOD0, but always < NumBonesTotal. Arrays indexed by
+		// bone index must be size NumBonesTotal.
+		const int32 NumBonesLOD0 = InsertionOrder.Num();
+		const int32 NumBonesTotal = SkelMeshRefSkel.GetNum();
+
+		// If our skeleton is not the one that was used to build the PhysicsAsset, some bodies may be missing, or rearranged.
+		// We need to map the original indices to the new bodies for use by the CollisionDisableTable.
+		// NOTE: This array is indexed by the original BodyInstance body index (BodyInstance->InstanceBodyIndex)
+		TArray<ImmediatePhysics::FActorHandle*> BodyIndexToActorHandle;
+		BodyIndexToActorHandle.AddZeroed(HighLevelBodyInstances.Num());
+
+		TArray<FBodyInstance*> BodiesSorted;
+		BodiesSorted.AddZeroed(NumBonesTotal);
+
+		for (FBodyInstance* BI : HighLevelBodyInstances)
+		{
+			if(BI->IsValidBodyInstance())
+			{
+				BodiesSorted[BI->InstanceBoneIndex] = BI;
+			}
+		}
+
+		// Create the immediate physics bodies
+		for (FBoneIndexType InsertBone : InsertionOrder)
+		{
+			if (FBodyInstance* BodyInstance = BodiesSorted[InsertBone])
+			{
+				UBodySetup* BodySetup = PhysicsAssetToUse->SkeletalBodySetups[BodyInstance->InstanceBodyIndex];
+
+				// Note that we create actors as dynamics, and then set them to kinematic if desired. Creating them
+				// as kinematic prevents them being subsequently made dynamic.
+				bool bSimulated = (BodySetup->PhysicsType == EPhysicsType::PhysType_Simulated);
+				ImmediatePhysics::EActorType ActorType = ImmediatePhysics::EActorType::DynamicActor;
+				ImmediatePhysics::FActorHandle* NewBodyHandle = PhysicsSimulation->CreateActor(ActorType, BodyInstance, BodyInstance->GetUnrealWorldTransform());
+				if (NewBodyHandle)
+				{
+					const float InvMass = (float) NewBodyHandle->GetInverseMass();
+					TotalMass += InvMass > 0.f ? 1.f / InvMass : 0.f;
+					if (!bSimulated)
+					{
+						NewBodyHandle->SetIsKinematic(true);
+					}
+
+					NewBodyHandle->SetName(BodySetup->BoneName);
+					const int32 BodyIndex = AddBody(NewBodyHandle);
+					const int32 SkeletonBoneIndex = MeshToSkeletonBoneIndex[InsertBone];
+					if (ensure(SkeletonBoneIndex >= 0))
+					{
+						SkeletonBoneIndexToBodyIndex[SkeletonBoneIndex] = BodyIndex;
+					}
+					BodyAnimData[BodyIndex].bIsSimulated = bSimulated;
+					NamesToHandles.Add(BodySetup->BoneName, NewBodyHandle);
+					BodyIndexToActorHandle[BodyInstance->InstanceBodyIndex] = NewBodyHandle;
+
+					if (BodySetup->CollisionReponse == EBodyCollisionResponse::BodyCollision_Disabled)
+					{
+						IgnoreCollisionActors.Add(NewBodyHandle);
+					}
+				}
+			}
+		}
+
+		//Insert joints so that they coincide body order. That is, if we stop simulating all bodies past some index, we can simply ignore joints past a corresponding index without any re-order
+		//For this to work we consider the most last inserted bone in each joint
+		TArray<int32> InsertionOrderPerBone;
+		InsertionOrderPerBone.AddUninitialized(NumBonesTotal);
+
+		for(int32 Position = 0; Position < NumBonesLOD0; ++Position)
+		{
+			InsertionOrderPerBone[InsertionOrder[Position]] = Position;
+		}
+
+		HighLevelConstraintInstances.Sort([&InsertionOrderPerBone, &SkelMeshRefSkel](const FConstraintInstance& LHS, const FConstraintInstance& RHS)
+		{
+			if(LHS.IsValidConstraintInstance() && RHS.IsValidConstraintInstance())
+			{
+				const int32 BoneIdxLHS1 = SkelMeshRefSkel.FindBoneIndex(LHS.ConstraintBone1);
+				const int32 BoneIdxLHS2 = SkelMeshRefSkel.FindBoneIndex(LHS.ConstraintBone2);
+
+				const int32 BoneIdxRHS1 = SkelMeshRefSkel.FindBoneIndex(RHS.ConstraintBone1);
+				const int32 BoneIdxRHS2 = SkelMeshRefSkel.FindBoneIndex(RHS.ConstraintBone2);
+
+				const int32 MaxPositionLHS = FMath::Max(InsertionOrderPerBone[BoneIdxLHS1], InsertionOrderPerBone[BoneIdxLHS2]);
+				const int32 MaxPositionRHS = FMath::Max(InsertionOrderPerBone[BoneIdxRHS1], InsertionOrderPerBone[BoneIdxRHS2]);
+
+				return MaxPositionLHS < MaxPositionRHS;
+			}
+			
+			return false;
+		});
+
+		Joints.SetNumZeroed(Bodies.Num());
+
+		TArray<ImmediatePhysics::FSimulation::FIgnorePair> IgnorePairs;
+		if(NamesToHandles.Num() > 0)
+		{
+			//constraints
+			for(int32 ConstraintIdx = 0; ConstraintIdx < HighLevelConstraintInstances.Num(); ++ConstraintIdx)
+			{
+				FConstraintInstance* CI = HighLevelConstraintInstances[ConstraintIdx];
+				ImmediatePhysics::FActorHandle* Body1Handle = NamesToHandles.FindRef(CI->ConstraintBone1);
+				ImmediatePhysics::FActorHandle* Body2Handle = NamesToHandles.FindRef(CI->ConstraintBone2);
+
+				if(Body1Handle && Body2Handle)
+				{
+					if (Body1Handle->IsSimulated() || Body2Handle->IsSimulated())
+					{
+						ImmediatePhysics::FJointHandle* NewJointHandle =
+							PhysicsSimulation->CreateJoint(CI, Body1Handle, Body2Handle);
+						// Record the joint handle under the child bone (i.e. more leaf-ward bone),
+						// since each bone may have multiple children, but only one parent.
+						NamesToJointHandles.Add(CI->ConstraintBone1, NewJointHandle);
+
+						if (bForceDisableCollisionBetweenConstraintBodies)
+						{
+							int32 BodyIndex1 = PhysicsAssetToUse->FindBodyIndex(CI->ConstraintBone1);
+							int32 BodyIndex2 = PhysicsAssetToUse->FindBodyIndex(CI->ConstraintBone2);
+							if (BodyIndex1 != INDEX_NONE && BodyIndex2 != INDEX_NONE)
+							{
+								PhysicsAssetToUse->DisableCollision(BodyIndex1, BodyIndex2);
+							}
+						}
+
+						int32 BodyIndex;
+						if (Bodies.Find(Body1Handle, BodyIndex))
+						{
+							Joints[BodyIndex] = NewJointHandle;
+
+							BodyAnimData[BodyIndex].LinearXMotion = CI->GetLinearXMotion();
+							BodyAnimData[BodyIndex].LinearYMotion = CI->GetLinearYMotion();
+							BodyAnimData[BodyIndex].LinearZMotion = CI->GetLinearZMotion();
+							BodyAnimData[BodyIndex].LinearLimit = CI->GetLinearLimit();
+
+							//set limit to ref pose 
+							FTransform Body1Transform = Body1Handle->GetWorldTransform();
+							FTransform Body2Transform = Body2Handle->GetWorldTransform();
+							BodyAnimData[BodyIndex].RefPoseLength = (float) 
+								Body1Transform.GetRelativeTransform(Body2Transform).GetLocation().Size();
+						}
+
+						if (CI->IsCollisionDisabled())
+						{
+							ImmediatePhysics::FSimulation::FIgnorePair Pair;
+							Pair.A = Body1Handle;
+							Pair.B = Body2Handle;
+							IgnorePairs.Add(Pair);
+						}
+					}
+				}
+			}
+
+			ResetSimulatedTeleportType = ETeleportType::ResetPhysics;
+		}
+
+		//== WithControl ==================================
+
+		if (!bHaveRunControlSeup && bEnableControls)
+		{
+			CreateWorldSpaceControlRootBody(PhysicsAssetToUse);
+			bHaveRunControlSeup = true;
+			InitControlsAndBodyModifiers(InAnimInstance->GetSkelMeshComponent());
+		}
+
+		//== WithControl ==================================
+
+		// Terminate all the constraint instances
+		for (FConstraintInstance* CI : HighLevelConstraintInstances)
+		{
+			CI->TermConstraint();
+			delete CI;
+		}
+
+		// Terminate all of the instances, cannot be done during insert or we may break constraint chains
+		for(FBodyInstance* Instance : HighLevelBodyInstances)
+		{
+			if(Instance->IsValidBodyInstance())
+			{
+				Instance->TermBody(true);
+			}
+
+			delete Instance;
+		}
+
+		HighLevelConstraintInstances.Empty();
+		HighLevelBodyInstances.Empty();
+		BodiesSorted.Empty();
+
+		const TMap<FRigidBodyIndexPair, bool>& DisableTable = PhysicsAssetToUse->CollisionDisableTable;
+		for(auto ConstItr = DisableTable.CreateConstIterator(); ConstItr; ++ConstItr)
+		{
+			int32 IndexA = ConstItr.Key().Indices[0];
+			int32 IndexB = ConstItr.Key().Indices[1];
+			if ((IndexA < BodyIndexToActorHandle.Num()) && (IndexB < BodyIndexToActorHandle.Num()))
+			{
+				if ((BodyIndexToActorHandle[IndexA] != nullptr) && (BodyIndexToActorHandle[IndexB] != nullptr))
+				{
+					ImmediatePhysics::FSimulation::FIgnorePair Pair;
+					Pair.A = BodyIndexToActorHandle[IndexA];
+					Pair.B = BodyIndexToActorHandle[IndexB];
+					IgnorePairs.Add(Pair);
+				}
+			}
+		}
+
+		PhysicsSimulation->SetIgnoreCollisionPairTable(IgnorePairs);
+		PhysicsSimulation->SetIgnoreCollisionActors(IgnoreCollisionActors);
+
+		CollectClothColliderObjects(SkeletalMeshComp);
+
+		SolverSettings = PhysicsAssetToUse->SolverSettings;
+		PhysicsSimulation->SetSolverSettings(
+			SolverSettings.FixedTimeStep,
+			SolverSettings.CullDistance,
+			SolverSettings.MaxDepenetrationVelocity,
+			SolverSettings.bUseLinearJointSolver,
+			SolverSettings.PositionIterations,
+			SolverSettings.VelocityIterations,
+			SolverSettings.ProjectionIterations);
+
+		SolverIterations = PhysicsAssetToUse->SolverIterations;
+
+		ApplyCurrentConstraintProfile();
+	}
+}
+
+DECLARE_CYCLE_STAT(TEXT("FAnimNode_RigidBodyWithControl::UpdateWorldGeometry"), STAT_ImmediateUpdateWorldGeometry, STATGROUP_ImmediatePhysics);
+
+void FAnimNode_RigidBodyWithControl::UpdateWorldGeometry(const UWorld& World, const USkeletalMeshComponent& SKC)
+{
+	LLM_SCOPE_BYNAME(TEXT("Animation/RigidBodyWithControl")); 
+	
+	SCOPE_CYCLE_COUNTER(STAT_ImmediateUpdateWorldGeometry);
+	QueryParams = FCollisionQueryParams(SCENE_QUERY_STAT(RagdollNodeFindGeometry), /*bTraceComplex=*/false);
+#if WITH_EDITOR
+	if(!World.IsGameWorld())
+	{
+		QueryParams.MobilityType = EQueryMobilityType::Any;	//If we're in some preview world trace against everything because things like the preview floor are not static
+		QueryParams.AddIgnoredComponent(&SKC);
+	}
+	else
+#endif
+	{
+		QueryParams.MobilityType = EQueryMobilityType::Static;	//We only want static actors
+	}
+
+	// Check for deleted world objects and flag for removal (later in anim task)
+	ExpireWorldObjects();
+
+	// If we have moved outside of the bounds we checked for world objects we need to gather new world objects
+	FSphere Bounds = SKC.CalcBounds(SKC.GetComponentToWorld()).GetSphere();
+	if (!Bounds.IsInside(CachedBounds))
+	{
+		// Since the cached bounds are no longer valid, update them.
+		CachedBounds = Bounds;
+		CachedBounds.W *= CachedBoundsScale;
+
+		// Cache the PhysScene and World for use in UpdateWorldForces and CollectWorldObjects
+		// When these are non-null it is an indicator that we need to update the collected world objects list
+		PhysScene = World.GetPhysicsScene();
+		UnsafeWorld = &World;
+		UnsafeOwner = SKC.GetOwner();
+
+		// A timer to track objects we haven't detected in a while
+		++ComponentsInSimTick;
+	}
+}
+
+DECLARE_CYCLE_STAT(TEXT("FAnimNode_RigidBodyWithControl::UpdateWorldForces"), STAT_ImmediateUpdateWorldForces, STATGROUP_ImmediatePhysics);
+
+void FAnimNode_RigidBodyWithControl::UpdateWorldForces(const FTransform& ComponentToWorld, const FTransform& BaseBoneTM, const float DeltaSeconds)
+{
+	LLM_SCOPE_BYNAME(TEXT("Animation/RigidBodyWithControl")); 
+	SCOPE_CYCLE_COUNTER(STAT_ImmediateUpdateWorldForces);
+
+	if(TotalMass > 0.f)
+	{
+		for (const USkeletalMeshComponent::FPendingRadialForces& PendingRadialForce : PendingRadialForces)
+		{
+			const FVector RadialForceOrigin = WorldPositionToSpace(SimulationSpace, PendingRadialForce.Origin, ComponentToWorld, BaseBoneTM);
+			for(ImmediatePhysics::FActorHandle* Body : Bodies)
+			{
+				const float InvMass = (float) Body->GetInverseMass();
+				if(InvMass > 0.f)
+				{
+					const float StrengthPerBody = PendingRadialForce.bIgnoreMass ? PendingRadialForce.Strength : PendingRadialForce.Strength / (TotalMass * InvMass);
+					ImmediatePhysics::EForceType ForceType;
+					if (PendingRadialForce.Type == USkeletalMeshComponent::FPendingRadialForces::AddImpulse)
+					{
+						ForceType = PendingRadialForce.bIgnoreMass ? ImmediatePhysics::EForceType::AddVelocity : ImmediatePhysics::EForceType::AddImpulse;
+					}
+					else
+					{
+						ForceType = PendingRadialForce.bIgnoreMass ? ImmediatePhysics::EForceType::AddAcceleration : ImmediatePhysics::EForceType::AddForce;
+					}
+					
+					Body->AddRadialForce(RadialForceOrigin, StrengthPerBody, PendingRadialForce.Radius, PendingRadialForce.Falloff, ForceType);
+				}
+			}
+		}
+
+		if(!ExternalForce.IsNearlyZero())
+		{
+			const FVector ExternalForceInSimSpace = WorldVectorToSpaceNoScale(SimulationSpace, ExternalForce, ComponentToWorld, BaseBoneTM);
+			for (ImmediatePhysics::FActorHandle* Body : Bodies)
+			{
+				const float InvMass = (float) Body->GetInverseMass();
+				if (InvMass > 0.f)
+				{
+					Body->AddForce(ExternalForceInSimSpace);
+				}
+			}
+		}
+		if(DeltaSeconds != 0.0)
+		{
+			if(!PerSolverField.IsEmpty())
+			{
+				TArray<FVector>& SamplePositions = PerSolverField.GetSamplePositions();
+				TArray<FFieldContextIndex>& SampleIndices = PerSolverField.GetSampleIndices();
+
+				SamplePositions.SetNum(Bodies.Num(),false);
+				SampleIndices.SetNum(Bodies.Num(), false);
+
+				int32 Index = 0;
+				for (ImmediatePhysics::FActorHandle* Body : Bodies)
+				{
+					SamplePositions[Index] = (Body->GetWorldTransform() * SpaceToWorldTransform(SimulationSpace, ComponentToWorld, BaseBoneTM)).GetLocation();
+					SampleIndices[Index] = FFieldContextIndex(Index, Index);
+					++Index;
+				}
+				PerSolverField.ComputeFieldRigidImpulse(WorldTimeSeconds);
+
+				const TArray<FVector>& LinearVelocities = PerSolverField.GetOutputResults(EFieldCommandOutputType::LinearVelocity);
+				const TArray<FVector>& LinearForces = PerSolverField.GetOutputResults(EFieldCommandOutputType::LinearForce);
+				const TArray<FVector>& AngularVelocities = PerSolverField.GetOutputResults(EFieldCommandOutputType::AngularVelocity);
+				const TArray<FVector>& AngularTorques = PerSolverField.GetOutputResults(EFieldCommandOutputType::AngularTorque);
+
+				if (LinearVelocities.Num() == Bodies.Num())
+				{
+					Index = 0;
+					for (ImmediatePhysics::FActorHandle* Body : Bodies)
+					{
+						const FVector ExternalForceInSimSpace = WorldVectorToSpaceNoScale(SimulationSpace, LinearVelocities[Index++], ComponentToWorld, BaseBoneTM) * Body->GetMass() / DeltaSeconds;
+						Body->AddForce(ExternalForceInSimSpace);
+					}
+				}
+				if (LinearForces.Num() == Bodies.Num())
+				{
+					Index = 0;
+					for (ImmediatePhysics::FActorHandle* Body : Bodies)
+					{
+						const FVector ExternalForceInSimSpace = WorldVectorToSpaceNoScale(SimulationSpace, LinearForces[Index++], ComponentToWorld, BaseBoneTM);
+						Body->AddForce(ExternalForceInSimSpace);
+					}
+				}
+				if (AngularVelocities.Num() == Bodies.Num())
+				{
+					Index = 0;
+					for (ImmediatePhysics::FActorHandle* Body : Bodies)
+					{
+						const FVector ExternalTorqueInSimSpace = WorldVectorToSpaceNoScale(SimulationSpace, AngularVelocities[Index++], ComponentToWorld, BaseBoneTM) * Body->GetInertia() / DeltaSeconds;
+						Body->AddTorque(ExternalTorqueInSimSpace);
+					}
+				}
+				if (AngularTorques.Num() == Bodies.Num())
+				{
+					Index = 0;
+					for (ImmediatePhysics::FActorHandle* Body : Bodies)
+					{
+						const FVector ExternalTorqueInSimSpace = WorldVectorToSpaceNoScale(SimulationSpace, AngularTorques[Index++], ComponentToWorld, BaseBoneTM);
+						Body->AddTorque(ExternalTorqueInSimSpace);
+					}
+				}
+			}
+		}
+	}
+}
+
+bool FAnimNode_RigidBodyWithControl::NeedsDynamicReset() const
+{
+	return true;
+}
+
+void FAnimNode_RigidBodyWithControl::ResetDynamics(ETeleportType InTeleportType)
+{
+	// This will be picked up next evaluate and reset our simulation.
+	// Teleport type can only go higher - i.e. if we have requested a reset, then a teleport will still reset fully
+	ResetSimulatedTeleportType = ((InTeleportType > ResetSimulatedTeleportType) ? InTeleportType : ResetSimulatedTeleportType);
+}
+
+DECLARE_CYCLE_STAT(TEXT("RigidBody_PreUpdate"), STAT_RigidBody_PreUpdate, STATGROUP_Anim);
+
+void FAnimNode_RigidBodyWithControl::PreUpdate(const UAnimInstance* InAnimInstance)
+{
+	// Don't update geometry if RBN is disabled
+	if(!bEnabled)
+	{
+		return;
+	}
+
+	SCOPE_CYCLE_COUNTER(STAT_RigidBody_PreUpdate);
+
+	USkeletalMeshComponent* SKC = InAnimInstance->GetSkelMeshComponent();
+	APawn* PawnOwner = InAnimInstance->TryGetPawnOwner();
+	UPawnMovementComponent* MovementComp = PawnOwner ? PawnOwner->GetMovementComponent() : nullptr;
+
+#if WITH_EDITOR
+	if (bEnableWorldGeometry && SimulationSpace != ESimulationSpace::WorldSpace && SKC && SKC->GetRelativeScale3D() != FVector(1.f, 1.f, 1.f))
+	{
+		FMessageLog("PIE").Warning(FText::Format(LOCTEXT("WorldCollisionComponentSpace", "Trying to use world collision without world space simulation for scaled ''{0}''. This is not supported, please change SimulationSpace to WorldSpace"),
+			FText::FromString(GetPathNameSafe(SKC))));
+	}
+#endif
+
+	UWorld* World = InAnimInstance->GetWorld();
+	if (World)
+	{
+		WorldSpaceGravity = bOverrideWorldGravity ? OverrideWorldGravity : (MovementComp ? FVector(0.f, 0.f, MovementComp->GetGravityZ()) : FVector(0.f, 0.f, World->GetGravityZ()));
+
+		if(SKC)
+		{
+			// Store game time for use in parallel evaluation. This may be the totol time (inc pauses) or the time the game has been unpaused.
+			WorldTimeSeconds = (float) (SKC->PrimaryComponentTick.bTickEvenWhenPaused ? World->UnpausedTimeSeconds : World->TimeSeconds);
+
+			if (PhysicsSimulation && bEnableWorldGeometry)
+			{ 
+				// @todo: this logic can be simplified now. We used to run PurgeExpiredWorldObjects and CollectWorldObjects
+				// in UpdateAnimation, but we can't access the world actor's geometry there
+				UpdateWorldGeometry(*World, *SKC);
+
+				// Remove expired objects from the sim
+				PurgeExpiredWorldObjects();
+
+				// Find nearby world objects to add to the sim (gated on UnsafeWorld - see UpdateWorldGeometry)
+				CollectWorldObjects();
+			}
+
+			PendingRadialForces = SKC->GetPendingRadialForces();
+
+			PreviousTransform = CurrentTransform;
+			CurrentTransform = SKC->GetComponentToWorld();
+
+			if (World->PhysicsField)
+			{
+				const FBox BoundingBox = SKC->CalcBounds(SKC->GetComponentTransform()).GetBox();
+
+				World->PhysicsField->FillTransientCommands(false, BoundingBox, WorldTimeSeconds, PerSolverField.GetTransientCommands());
+				World->PhysicsField->FillPersistentCommands(false, BoundingBox, WorldTimeSeconds, PerSolverField.GetPersistentCommands());
+			}
+		}
+	}
+}
+
+int32 FAnimNode_RigidBodyWithControl::GetLODThreshold() const
+{
+	if(CVarRigidBodyNodeWithControlLODThreshold.GetValueOnAnyThread() != -1)
+	{
+		if(LODThreshold != -1)
+		{
+			return FMath::Min(LODThreshold, CVarRigidBodyNodeWithControlLODThreshold.GetValueOnAnyThread());
+		}
+		else
+		{
+			return CVarRigidBodyNodeWithControlLODThreshold.GetValueOnAnyThread();
+		}
+	}
+	else
+	{
+		return LODThreshold;
+	}
+}
+
+DECLARE_CYCLE_STAT(TEXT("RigidBody_Update"), STAT_RigidBody_Update, STATGROUP_Anim);
+
+void FAnimNode_RigidBodyWithControl::UpdateInternal(const FAnimationUpdateContext& Context)
+{
+	LLM_SCOPE_BYNAME(TEXT("Animation/RigidBodyWithControl")); 
+	DECLARE_SCOPE_HIERARCHICAL_COUNTER_ANIMNODE(UpdateInternal)
+	// Avoid this work if RBN is disabled, as the results would be discarded
+	if(!bEnabled)
+	{
+		return;
+	}
+
+	SCOPE_CYCLE_COUNTER(STAT_RigidBody_Update);
+	
+	// Must flush the simulation since we may be making changes to the scene
+	FlushDeferredSimulationTask();
+
+	// Accumulate deltatime elapsed during update. To be used during evaluation.
+	AccumulatedDeltaTime += Context.AnimInstanceProxy->GetDeltaSeconds();
+
+	if (UnsafeWorld != nullptr)
+	{
+		// Node is valid to evaluate. Simulation is starting.
+		bSimulationStarted = true;
+	}
+
+	// These get set again if our bounds change. Subsequent calls to CollectWorldObjects will early-out until then
+	UnsafeWorld = nullptr;
+	UnsafeOwner = nullptr;
+	PhysScene = nullptr;
+}
+
+void FAnimNode_RigidBodyWithControl::CollectClothColliderObjects(const USkeletalMeshComponent* SkeletalMeshComp)
+{
+	if (bUseExternalClothCollision && bRBANWithControl_IncludeClothColliders && SkeletalMeshComp && PhysicsSimulation)
+	{
+		const TArray<FClothCollisionSource>& SkeletalMeshClothCollisionSources = SkeletalMeshComp->GetClothCollisionSources();
+		
+		for (const FClothCollisionSource& ClothCollisionSource : SkeletalMeshClothCollisionSources)
+		{
+			const USkeletalMeshComponent* const SourceComponent = ClothCollisionSource.SourceComponent.Get();
+			const UPhysicsAsset* const PhysicsAsset = ClothCollisionSource.SourcePhysicsAsset.Get();
+
+			if (SourceComponent && PhysicsAsset)
+			{
+				TArray<FBodyInstance*> BodyInstances;
+				SourceComponent->InstantiatePhysicsAssetBodies(*PhysicsAsset, BodyInstances);
+
+				for (uint32 BodyInstanceIndex = 0, BodyInstanceMax = BodyInstances.Num(); BodyInstanceIndex < BodyInstanceMax; ++BodyInstanceIndex)
+				{
+					FBodyInstance* const BodyInstance = BodyInstances[BodyInstanceIndex];
+
+					ImmediatePhysics::FActorHandle* const ActorHandle = PhysicsSimulation->CreateActor(ImmediatePhysics::EActorType::KinematicActor, BodyInstance, BodyInstance->GetUnrealWorldTransform());
+					PhysicsSimulation->AddToCollidingPairs(ActorHandle); // <-allow collision between this actor and all dynamic actors.
+					ClothColliders.Add(FClothCollider(ActorHandle, SourceComponent, BodyInstance->InstanceBoneIndex));
+
+					// Terminate the instance.
+					if (BodyInstance->IsValidBodyInstance())
+					{
+						BodyInstance->TermBody(true);
+					}
+
+					delete BodyInstance;
+					BodyInstances[BodyInstanceIndex] = nullptr;
+				}
+
+				BodyInstances.Reset();
+			}
+		}
+	}
+}
+
+void FAnimNode_RigidBodyWithControl::RemoveClothColliderObjects()
+{
+	for (const FClothCollider& ClothCollider : ClothColliders)
+	{
+		PhysicsSimulation->DestroyActor(ClothCollider.ActorHandle);
+	}
+	
+	ClothColliders.Reset();
+}
+
+void FAnimNode_RigidBodyWithControl::UpdateClothColliderObjects(const FTransform& SpaceTransform)
+{
+	for (FClothCollider& ClothCollider : ClothColliders)
+	{
+		if (ClothCollider.ActorHandle && ClothCollider.SkeletalMeshComponent)
+		{
+			// Calculate the sim-space transform of this object
+			const FTransform CompWorldTransform = ClothCollider.SkeletalMeshComponent->GetBoneTransform(ClothCollider.BoneIndex);
+			FTransform CompSpaceTransform;
+			CompSpaceTransform.SetTranslation(SpaceTransform.InverseTransformPosition(CompWorldTransform.GetLocation()));
+			CompSpaceTransform.SetRotation(SpaceTransform.InverseTransformRotation(CompWorldTransform.GetRotation()));
+			CompSpaceTransform.SetScale3D(FVector::OneVector);	// TODO - sort out scale for world objects in local sim
+
+			// Update the sim's copy of the world object
+			ClothCollider.ActorHandle->SetWorldTransform(CompSpaceTransform);
+		}
+	}
+}
+
+void FAnimNode_RigidBodyWithControl::CollectWorldObjects()
+{
+	if ((UnsafeWorld != nullptr) && (PhysScene != nullptr))
+	{
+		// @todo(ccaulfield): should this use CachedBounds?
+		TArray<FOverlapResult> Overlaps;
+		UnsafeWorld->OverlapMultiByChannel(Overlaps, CachedBounds.Center, FQuat::Identity, OverlapChannel, FCollisionShape::MakeSphere((float) CachedBounds.W), QueryParams, FCollisionResponseParams(ECR_Overlap));
+
+		for (const FOverlapResult& Overlap : Overlaps)
+		{
+			if (UPrimitiveComponent* OverlapComp = Overlap.GetComponent())
+			{
+				FWorldObject* WorldObject = ComponentsInSim.Find(OverlapComp);
+				if (WorldObject != nullptr)
+				{
+					// Existing object - reset its age
+					WorldObject->LastSeenTick = ComponentsInSimTick;
+				}
+				else
+				{
+					// New object - add it to the sim
+					const bool bIsSelf = (UnsafeOwner == OverlapComp->GetOwner());
+					if (!bIsSelf)
+					{
+						// Create a kinematic actor. Not using Static as world-static objects may move in the simulation's frame of reference
+						ImmediatePhysics::FActorHandle* ActorHandle = PhysicsSimulation->CreateActor(ImmediatePhysics::EActorType::KinematicActor, &OverlapComp->BodyInstance, OverlapComp->GetComponentTransform());
+						PhysicsSimulation->AddToCollidingPairs(ActorHandle);
+						ComponentsInSim.Add(OverlapComp, FWorldObject(ActorHandle, ComponentsInSimTick));
+					}
+				}
+			}
+		}
+	}
+}
+
+// Flag invalid objects for purging
+void FAnimNode_RigidBodyWithControl::ExpireWorldObjects()
+{
+	// Invalidate deleted and expired world objects
+	TArray<const UPrimitiveComponent*> PrunedEntries;
+	for (auto& WorldEntry : ComponentsInSim)
+	{
+		const UPrimitiveComponent* WorldComp = WorldEntry.Key;
+		FWorldObject& WorldObject = WorldEntry.Value;
+
+		// Do we need to expire this object?
+		const int32 ExpireTickCount = RBANWithControl_WorldObjectExpiry;
+		bool bIsInvalid =
+			((ComponentsInSimTick - WorldObject.LastSeenTick) > ExpireTickCount)	// Haven't seen this object for a while
+			|| !IsValid(WorldComp)
+			|| (WorldComp->GetBodyInstance() == nullptr)
+			|| (!WorldComp->GetBodyInstance()->IsValidBodyInstance());
+
+		// Remove from sim if necessary
+		if (bIsInvalid)
+		{
+			WorldObject.bExpired = true;
+		}
+	}
+}
+
+void FAnimNode_RigidBodyWithControl::PurgeExpiredWorldObjects()
+{
+	// Destroy expired simulated objects
+	TArray<const UPrimitiveComponent*> PurgedEntries;
+	for (auto& WorldEntry : ComponentsInSim)
+	{
+		FWorldObject& WorldObject = WorldEntry.Value;
+
+		if (WorldObject.bExpired)
+		{
+			PhysicsSimulation->DestroyActor(WorldObject.ActorHandle);
+			WorldObject.ActorHandle = nullptr;
+
+			PurgedEntries.Add(WorldEntry.Key);
+		}
+	}
+
+	// Remove purged map entries
+	for (const UPrimitiveComponent* PurgedEntry : PurgedEntries)
+	{
+		ComponentsInSim.Remove(PurgedEntry);
+	}
+}
+
+// Update the transforms of the world objects we added to the sim. This is required
+// if we have a component- or bone-space simulation as even world-static objects
+// will be moving in the simulation's frame of reference.
+void FAnimNode_RigidBodyWithControl::UpdateWorldObjects(const FTransform& SpaceTransform)
+{
+	LLM_SCOPE_BYNAME(TEXT("Animation/RigidBodyWithControl")); 
+
+	if (SimulationSpace != ESimulationSpace::WorldSpace)
+	{
+		for (const auto& WorldEntry : ComponentsInSim)
+		{ 
+			const UPrimitiveComponent* OverlapComp = WorldEntry.Key;
+			if (OverlapComp != nullptr)
+			{
+				ImmediatePhysics::FActorHandle* ActorHandle = WorldEntry.Value.ActorHandle;
+
+				// Calculate the sim-space transform of this object
+				const FTransform CompWorldTransform = OverlapComp->BodyInstance.GetUnrealWorldTransform();
+				FTransform CompSpaceTransform;
+				CompSpaceTransform.SetTranslation(SpaceTransform.InverseTransformPosition(CompWorldTransform.GetLocation()));
+				CompSpaceTransform.SetRotation(SpaceTransform.InverseTransformRotation(CompWorldTransform.GetRotation()));
+				CompSpaceTransform.SetScale3D(FVector::OneVector);	// TODO - sort out scale for world objects in local sim
+
+				// Update the sim's copy of the world object
+				ActorHandle->SetWorldTransform(CompSpaceTransform);
+			}
+		}
+	}
+}
+
+void FAnimNode_RigidBodyWithControl::InitializeBoneReferences(const FBoneContainer& RequiredBones) 
+{
+	DECLARE_SCOPE_HIERARCHICAL_COUNTER_ANIMNODE(InitializeBoneReferences)
+	/** We only need to update simulated bones and children of simulated bones*/
+	const int32 NumBodies = Bodies.Num();
+	const TArray<FBoneIndexType>& RequiredBoneIndices = RequiredBones.GetBoneIndicesArray();
+	const int32 NumRequiredBoneIndices = RequiredBoneIndices.Num();
+	const FReferenceSkeleton& RefSkeleton = RequiredBones.GetReferenceSkeleton();
+
+	OutputBoneData.Empty(NumBodies);
+
+	int32 NumSimulatedBodies = 0;
+	TArray<int32> SimulatedBodyIndices;
+	// if no name is entered, use root
+	if (BaseBoneRef.BoneName == NAME_None)
+	{
+		BaseBoneRef.BoneName = RefSkeleton.GetBoneName(0);
+	}
+
+	if (BaseBoneRef.BoneName != NAME_None)
+	{
+		BaseBoneRef.Initialize(RequiredBones);
+	}
+
+	if (!BaseBoneRef.HasValidSetup())
+	{
+		// If the user specified a simulation root that is not used by the skelmesh, issue a warning 
+		// (FAnimNode_RigidBodyWithControl::IsValidToEvaluate will return false and the simulation will not run)
+		UE_LOG(LogRigidBodyWithControl, Log, TEXT("FAnimNode_RigidBodyWithControl: RigidBodyWithControl Simulation Base Bone \'%s\' does not exist on SkeletalMesh %s."), *BaseBoneRef.BoneName.ToString(), *GetNameSafe(RequiredBones.GetSkeletalMeshAsset()));
+	}
+
+	bool bHasInvalidBoneReference = false;
+	for (int32 Index = 0; Index < NumRequiredBoneIndices; ++Index)
+	{
+		const FCompactPoseBoneIndex CompactPoseBoneIndex(Index);
+		const FBoneIndexType SkeletonBoneIndex = (FBoneIndexType) RequiredBones.GetSkeletonIndex(CompactPoseBoneIndex);
+		const FBoneIndexType IndexToBodyNum = (FBoneIndexType) SkeletonBoneIndexToBodyIndex.Num();
+
+		// If we have a missing bone in our skeleton, we don't want to have an out of bounds access.
+		if (SkeletonBoneIndex >= IndexToBodyNum)
+		{
+			bHasInvalidBoneReference = true;
+			break;
+		}
+
+		const int32 BodyIndex = SkeletonBoneIndexToBodyIndex[SkeletonBoneIndex];
+
+		if (BodyIndex != INDEX_NONE)
+		{
+			//If we have a body we need to save it for later
+			RigidBodyWithControl::FOutputBoneData* OutputData = new (OutputBoneData) RigidBodyWithControl::FOutputBoneData();
+			OutputData->BodyIndex = BodyIndex;
+			OutputData->CompactPoseBoneIndex = CompactPoseBoneIndex;
+
+			if (BodyAnimData[BodyIndex].bIsSimulated)
+			{
+				++NumSimulatedBodies;
+				SimulatedBodyIndices.AddUnique(BodyIndex);
+			}
+
+			OutputData->BoneIndicesToParentBody.Add(CompactPoseBoneIndex);
+
+			// Walk up parent chain until we find parent body.
+			OutputData->ParentBodyIndex = INDEX_NONE;
+			FCompactPoseBoneIndex CompactParentIndex = RequiredBones.GetParentBoneIndex(CompactPoseBoneIndex);
+			OutputData->CompactPoseParentBoneIndex = CompactParentIndex;
+			while (CompactParentIndex != INDEX_NONE)
+			{
+				const FBoneIndexType SkeletonParentBoneIndex = (FBoneIndexType ) RequiredBones.GetSkeletonIndex(CompactParentIndex);
+
+				// Must check our parent as well for a missing bone.
+				if (SkeletonParentBoneIndex >= IndexToBodyNum)
+				{
+					bHasInvalidBoneReference = true;
+					break;
+				}
+
+				OutputData->ParentBodyIndex = SkeletonBoneIndexToBodyIndex[SkeletonParentBoneIndex];
+				OutputData->CompactPoseParentBoneIndex = CompactParentIndex;
+				if (OutputData->ParentBodyIndex != INDEX_NONE)
+				{
+					break;
+				}
+
+				OutputData->BoneIndicesToParentBody.Add(CompactParentIndex);
+				CompactParentIndex = RequiredBones.GetParentBoneIndex(CompactParentIndex);
+			}
+
+			if (bHasInvalidBoneReference)
+			{
+				break;
+			}
+		}
+	}
+
+	if (bHasInvalidBoneReference)
+	{
+		// If a bone was missing, let us know which asset it happened on, and clear our bone container to make the bad asset visible.
+		ensureMsgf(false, TEXT("FAnimNode_RigidBodyWithControl::InitializeBoneReferences: The Skeleton %s, is missing bones that SkeletalMesh %s needs. Skeleton might need to be resaved."),
+			*GetNameSafe(RequiredBones.GetSkeletonAsset()), *GetNameSafe(RequiredBones.GetSkeletalMeshAsset()));
+		OutputBoneData.Empty();
+	}
+	else
+	{
+		// New bodies potentially introduced with new LOD
+		// We'll have to initialize their transform.
+		bCheckForBodyTransformInit = true;
+
+		if (PhysicsSimulation)
+		{
+			PhysicsSimulation->SetNumActiveBodies(NumSimulatedBodies, SimulatedBodyIndices);
+		}
+
+		// We're switching to a new LOD, this invalidates our captured poses.
+		CapturedFrozenPose.Empty();
+		CapturedFrozenCurves.Empty();
+	}
+}
+
+void FAnimNode_RigidBodyWithControl::AddImpulseAtLocation(FVector Impulse, FVector Location, FName BoneName)
+{
+	// Find the body. This is currently only used in the editor and will need optimizing if used in game
+	for (int32 BodyIndex = 0; BodyIndex < Bodies.Num(); ++BodyIndex)
+	{
+		ImmediatePhysics::FActorHandle* Body = Bodies[BodyIndex];
+		if (Body->GetName() == BoneName)
+		{
+			Body->AddImpulseAtLocation(Impulse, Location);
+		}
+	}
+}
+
+void FAnimNode_RigidBodyWithControl::OnInitializeAnimInstance(const FAnimInstanceProxy* InProxy, const UAnimInstance* InAnimInstance)
+{
+	InitPhysics(InAnimInstance);
+}
+
+#if WITH_EDITORONLY_DATA
+void FAnimNode_RigidBodyWithControl::PostSerialize(const FArchive& Ar)
+{
+}
+#endif
+
+bool FAnimNode_RigidBodyWithControl::IsValidToEvaluate(const USkeleton* Skeleton, const FBoneContainer& RequiredBones)
+{
+	return BaseBoneRef.IsValidToEvaluate(RequiredBones);
+}
+
+const int32 FAnimNode_RigidBodyWithControl::GetNumBodies() const
+{
+	return Bodies.Num();
+}
+
+
+
+#undef LOCTEXT_NAMESPACE
+
diff --git a/Engine/Plugins/Experimental/PhysicsControl/Source/PhysicsControl/Private/PhysicsControlBPLibrary.cpp b/Engine/Plugins/Experimental/PhysicsControl/Source/PhysicsControl/Private/PhysicsControlBPLibrary.cpp
new file mode 100644
--- /dev/null
+++ b/Engine/Plugins/Experimental/PhysicsControl/Source/PhysicsControl/Private/PhysicsControlBPLibrary.cpp
@@ -0,0 +1,215 @@
+// Copyright Epic Games, Inc. All Rights Reserved.
+
+#include "PhysicsControlBPLibrary.h"
+
+#include "RigidBodyControlData.h"
+#include "AnimNode_RigidBodyWithControl.h"
+
+//UE_DISABLE_OPTIMIZATION
+
+#if 1
+
+//======================================================================================================================
+template<typename TParameterType> const TArray<FName>* FindNamesInSet(
+	const FRigidBodyNameRecords& NameRecords, const FName SetName)
+{
+	return nullptr;
+}
+
+//======================================================================================================================
+template<> const TArray<FName>* FindNamesInSet<FRigidBodyControl>(
+	const FRigidBodyNameRecords& NameRecords, const FName SetName)
+{
+	return NameRecords.ControlSets.Find(SetName);
+}
+
+//======================================================================================================================
+template<> const TArray<FName>* FindNamesInSet<FRigidBodyModifier>(
+	const FRigidBodyNameRecords& NameRecords, const FName SetName)
+{
+	return NameRecords.BodyModifierSets.Find(SetName);
+}
+
+//======================================================================================================================
+template<typename TParameterType> TArray<FName> GetNamesInSet(
+	const FRigidBodyWithControlReference& RigidBodyWithControl, const FName SetName)
+{
+	TArray<FName> OutputNames;
+
+	RigidBodyWithControl.CallAnimNodeFunction<FAnimNode_RigidBodyWithControl>(
+		TEXT("GetControlNamesInSet"),
+		[&OutputNames, SetName](FAnimNode_RigidBodyWithControl& InRigidBodyWithControl)
+		{
+			if (const TArray<FName>* FoundNames = FindNamesInSet<TParameterType>(
+				InRigidBodyWithControl.GetNameRecords(), SetName))
+			{
+				OutputNames = *FoundNames;
+			}
+		});
+
+	return OutputNames;
+}
+
+//======================================================================================================================
+template<typename TNamedParameters> void InterpolateParametersContainers(
+	TArray<TNamedParameters>& TargetContainer, const TArray<TNamedParameters>& SourceContainer, const float Weight)
+{
+	for (const TNamedParameters& Source : SourceContainer)
+	{
+		const FName SourceName = Source.Name;
+		if (TNamedParameters* const ExistingUpdate = TargetContainer.FindByPredicate(
+			[SourceName](const TNamedParameters& Element) { return Element.Name == SourceName; }))
+		{
+			// Interpolate values if the target container includes an element with the source values name.
+			ExistingUpdate->Data = Interpolate(ExistingUpdate->Data, Source.Data, Weight);
+		}
+		else
+		{
+			// Add this element to the target container.
+			TargetContainer.Add(Source);
+		}
+	}
+}
+
+//======================================================================================================================
+template<typename TNamedParameters> void BlendParametersThroughSet(
+	const FRigidBodyControlAndModifierParameters& InParametersContainer, 
+	const TNamedParameters&                       InStartParameters,
+	const TNamedParameters&                       InEndParameters,
+	const TArray<FName>&                          Names,
+	FRigidBodyControlAndModifierParameters&       OutParametersContainer)
+{
+	// TODO - Limbs can include branches. Would need a more sophisticated approach to deal with
+	// those properly, perhaps include a depth index in the control name ?
+
+	OutParametersContainer = InParametersContainer;
+
+	const float WeightDelta = 1.0f / static_cast<float>(Names.Num() - 1);
+	float Weight = 0.0f;
+
+	for (const FName& Name : Names)
+	{
+		OutParametersContainer.Add(
+			TNamedParameters(Name, Interpolate(InStartParameters.Data, InEndParameters.Data, Weight)));
+		Weight += WeightDelta;
+	}
+}
+
+//======================================================================================================================
+void UPhysicsControlBPLibrary::AddControlParameters(
+	UPARAM(ref) FRigidBodyControlAndModifierParameters& InParametersContainer,
+	FRigidBodyControlAndModifierParameters&             OutParametersContainer,
+	const FName                                         Name, 
+	const FRigidBodyControlSparseData&                  ControlData)
+{
+	OutParametersContainer = InParametersContainer;
+	OutParametersContainer.Add(FRigidBodyNamedControlParameters(Name, ControlData));
+}
+
+//======================================================================================================================
+void UPhysicsControlBPLibrary::AddMultipleControlParameters(
+	UPARAM(ref) FRigidBodyControlAndModifierParameters& InParametersContainer,
+	FRigidBodyControlAndModifierParameters&             OutParametersContainer,
+	const TArray<FName>&                                Names,
+	const FRigidBodyControlSparseData&                  ControlData)
+{
+	OutParametersContainer = InParametersContainer;
+	OutParametersContainer.ControlParameters.Reserve(OutParametersContainer.ControlParameters.Num() + Names.Num());
+
+	for (const FName Name : Names)
+	{
+		OutParametersContainer.Add(FRigidBodyNamedControlParameters(Name, ControlData));
+	}	
+}
+
+//======================================================================================================================
+void UPhysicsControlBPLibrary::AddModifierParameters(
+	UPARAM(ref) FRigidBodyControlAndModifierParameters& InParametersContainer,
+	FRigidBodyControlAndModifierParameters&             OutParametersContainer,
+	const FName                                         Name,
+	const FRigidBodyModifierSparseData&                 ModifierData)
+{
+	OutParametersContainer = InParametersContainer;
+	OutParametersContainer.Add(FRigidBodyNamedModifierParameters(Name, ModifierData));
+}
+
+//======================================================================================================================
+void UPhysicsControlBPLibrary::AddMultipleModifierParameters(
+	UPARAM(ref) FRigidBodyControlAndModifierParameters& InParametersContainer,
+	FRigidBodyControlAndModifierParameters&             OutParametersContainer,
+	const TArray<FName>&                                Names,
+	const FRigidBodyModifierSparseData&                 ModifierData)
+{
+	OutParametersContainer = InParametersContainer;
+	OutParametersContainer.ModifierParameters.Reserve(OutParametersContainer.ModifierParameters.Num() + Names.Num());
+
+	for (const FName Name : Names)
+	{
+		OutParametersContainer.Add(FRigidBodyNamedModifierParameters(Name, ModifierData));
+	}
+}
+
+//======================================================================================================================
+void UPhysicsControlBPLibrary::BlendParameters(
+	UPARAM(ref) FRigidBodyControlAndModifierParameters& InParametersContainerA, 
+	UPARAM(ref) FRigidBodyControlAndModifierParameters& InParametersContainerB, 
+	const float                                         InInterpolationWeight, 
+	FRigidBodyControlAndModifierParameters&             OutParametersContainers)
+{
+	OutParametersContainers = InParametersContainerA;
+
+	InterpolateParametersContainers(
+		OutParametersContainers.ControlParameters, InParametersContainerB.ControlParameters, InInterpolationWeight);
+	InterpolateParametersContainers(
+		OutParametersContainers.ModifierParameters, InParametersContainerB.ModifierParameters, InInterpolationWeight);
+}
+
+//======================================================================================================================
+void UPhysicsControlBPLibrary::BlendControlParametersThroughSet(
+	UPARAM(ref) FRigidBodyControlAndModifierParameters& InParametersContainer,
+	UPARAM(ref) const FRigidBodyNamedControlParameters& InStartParameters,
+	UPARAM(ref) const FRigidBodyNamedControlParameters& InEndParameters,
+	const TArray<FName>&                                InNames,
+	FRigidBodyControlAndModifierParameters&             OutParametersContainer)
+{
+	BlendParametersThroughSet(InParametersContainer, InStartParameters, InEndParameters, InNames, OutParametersContainer);
+}
+
+//======================================================================================================================
+void UPhysicsControlBPLibrary::BlendModifierParametersThroughSet(
+	UPARAM(ref) FRigidBodyControlAndModifierParameters&  InParametersContainer, 
+	UPARAM(ref) const FRigidBodyNamedModifierParameters& InStartParameters,
+	UPARAM(ref) const FRigidBodyNamedModifierParameters& InEndParameters,
+	const TArray<FName>&                                 InNames,
+	FRigidBodyControlAndModifierParameters&              OutParametersContainer)
+{
+	BlendParametersThroughSet(InParametersContainer, InStartParameters, InEndParameters, InNames, OutParametersContainer);
+}
+
+//======================================================================================================================
+FRigidBodyWithControlReference UPhysicsControlBPLibrary::ConvertToRigidBodyWithControl(
+	const FAnimNodeReference&           Node,
+	EAnimNodeReferenceConversionResult& Result)
+{
+	return FAnimNodeReference::ConvertToType<FRigidBodyWithControlReference>(Node, Result);
+}
+
+//======================================================================================================================
+UFUNCTION(BlueprintPure, Category = "Animation|PhysicsControl", meta = (BlueprintThreadSafe))
+TArray<FName> UPhysicsControlBPLibrary::GetControlNamesInSet(
+	const FRigidBodyWithControlReference& RigidBodyWithControl, 
+	const FName                           SetName)
+{
+	return GetNamesInSet<FRigidBodyControl>(RigidBodyWithControl, SetName);
+}
+
+//======================================================================================================================
+UFUNCTION(BlueprintPure, Category = "Animation|PhysicsControl", meta = (BlueprintThreadSafe))
+TArray<FName> UPhysicsControlBPLibrary::GetBodyModifierNamesInSet(
+	const FRigidBodyWithControlReference& RigidBodyWithControl,
+	const FName                           SetName)
+{
+	return GetNamesInSet<FRigidBodyModifier>(RigidBodyWithControl, SetName);
+}
+
+#endif
diff --git a/Engine/Plugins/Experimental/PhysicsControl/Source/PhysicsControl/Private/PhysicsControlComponent.cpp b/Engine/Plugins/Experimental/PhysicsControl/Source/PhysicsControl/Private/PhysicsControlComponent.cpp
--- a/Engine/Plugins/Experimental/PhysicsControl/Source/PhysicsControl/Private/PhysicsControlComponent.cpp
+++ b/Engine/Plugins/Experimental/PhysicsControl/Source/PhysicsControl/Private/PhysicsControlComponent.cpp
@@ -1,7 +1,7 @@
 // Copyright Epic Games, Inc. All Rights Reserved.
 
 #include "PhysicsControlComponent.h"
-#include "PhysicsControlComponentLog.h"
+#include "PhysicsControlLog.h"
 #include "PhysicsControlRecord.h"
 #include "PhysicsControlComponentHelpers.h"
 #include "PhysicsControlComponentImpl.h"
diff --git a/Engine/Plugins/Experimental/PhysicsControl/Source/PhysicsControl/Private/PhysicsControlComponentImpl.cpp b/Engine/Plugins/Experimental/PhysicsControl/Source/PhysicsControl/Private/PhysicsControlComponentImpl.cpp
--- a/Engine/Plugins/Experimental/PhysicsControl/Source/PhysicsControl/Private/PhysicsControlComponentImpl.cpp
+++ b/Engine/Plugins/Experimental/PhysicsControl/Source/PhysicsControl/Private/PhysicsControlComponentImpl.cpp
@@ -2,7 +2,7 @@
 
 #include "PhysicsControlComponentImpl.h"
 #include "Engine/SkeletalMesh.h"
-#include "PhysicsControlComponentLog.h"
+#include "PhysicsControlLog.h"
 #include "PhysicsControlComponent.h"
 #include "PhysicsControlComponentHelpers.h"
 
diff --git a/Engine/Plugins/Experimental/PhysicsControl/Source/PhysicsControl/Private/PhysicsControlLog.cpp b/Engine/Plugins/Experimental/PhysicsControl/Source/PhysicsControl/Private/PhysicsControlLog.cpp
new file mode 100644
--- /dev/null
+++ b/Engine/Plugins/Experimental/PhysicsControl/Source/PhysicsControl/Private/PhysicsControlLog.cpp
@@ -0,0 +1,6 @@
+// Copyright Epic Games, Inc. All Rights Reserved.
+
+#include "PhysicsControlLog.h"
+
+DEFINE_LOG_CATEGORY(LogPhysicsControlComponent);
+DEFINE_LOG_CATEGORY(LogRigidBodyWithControl);
diff --git a/Engine/Plugins/Experimental/PhysicsControl/Source/PhysicsControl/Private/PhysicsControlLog.h b/Engine/Plugins/Experimental/PhysicsControl/Source/PhysicsControl/Private/PhysicsControlLog.h
new file mode 100644
--- /dev/null
+++ b/Engine/Plugins/Experimental/PhysicsControl/Source/PhysicsControl/Private/PhysicsControlLog.h
@@ -0,0 +1,15 @@
+// Copyright Epic Games, Inc. All Rights Reserved.
+
+#pragma once
+
+#include "Logging/LogMacros.h"
+
+#if UE_BUILD_SHIPPING || UE_BUILD_TEST
+DECLARE_LOG_CATEGORY_EXTERN(LogRigidBodyWithControl, Log, Warning);
+DECLARE_LOG_CATEGORY_EXTERN(LogPhysicsControlComponent, Log, Warning);
+#else
+DECLARE_LOG_CATEGORY_EXTERN(LogRigidBodyWithControl, Log, All);
+DECLARE_LOG_CATEGORY_EXTERN(LogPhysicsControlComponent, Log, All);
+#endif
+
+
diff --git a/Engine/Plugins/Experimental/PhysicsControl/Source/PhysicsControl/Private/PhysicsControlRecord.cpp b/Engine/Plugins/Experimental/PhysicsControl/Source/PhysicsControl/Private/PhysicsControlRecord.cpp
--- a/Engine/Plugins/Experimental/PhysicsControl/Source/PhysicsControl/Private/PhysicsControlRecord.cpp
+++ b/Engine/Plugins/Experimental/PhysicsControl/Source/PhysicsControl/Private/PhysicsControlRecord.cpp
@@ -1,7 +1,7 @@
 // Copyright Epic Games, Inc. All Rights Reserved.
 
 #include "PhysicsControlRecord.h"
-#include "PhysicsControlComponentLog.h"
+#include "PhysicsControlLog.h"
 #include "PhysicsControlComponentHelpers.h"
 
 #include "PhysicsEngine/ConstraintInstance.h"
diff --git a/Engine/Plugins/Experimental/PhysicsControl/Source/PhysicsControl/Private/RigidBodyControlData.cpp b/Engine/Plugins/Experimental/PhysicsControl/Source/PhysicsControl/Private/RigidBodyControlData.cpp
new file mode 100644
--- /dev/null
+++ b/Engine/Plugins/Experimental/PhysicsControl/Source/PhysicsControl/Private/RigidBodyControlData.cpp
@@ -0,0 +1,161 @@
+// Copyright Epic Games, Inc. All Rights Reserved.
+
+#include "RigidBodyControlData.h"
+
+//======================================================================================================================
+FName GetControlTypeName(const ERigidBodyMovementType MovementType)
+{
+	switch (MovementType)
+	{
+	case ERigidBodyMovementType::Kinematic:
+		return "Kinematic";
+	case ERigidBodyMovementType::Simulated:
+		return "Simulated";
+	}
+	return "None";
+}
+
+//======================================================================================================================
+FName GetControlTypeName(const ERigidBodyControlType ControlType)
+{
+	switch (ControlType)
+	{
+	case ERigidBodyControlType::ParentSpace:
+		return "ParentSpace";
+	case ERigidBodyControlType::WorldSpace:
+		return "WorldSpace";
+	}
+	return "None";
+}
+
+#define INTERPOLATE_PARAM(NAME) Output.NAME = FMath::Lerp(A.NAME, B.NAME, Weight)
+#define SET_ENABLED_PARAM(NAME) Output.bEnable##NAME = A.bEnable##NAME && B.bEnable##NAME
+
+//======================================================================================================================
+FRigidBodyControlData Interpolate(
+	const FRigidBodyControlData& A, const FRigidBodyControlData& B, const float Weight)
+{
+	FRigidBodyControlData Output;
+	INTERPOLATE_PARAM(LinearStrength);
+	INTERPOLATE_PARAM(LinearDampingRatio);
+	INTERPOLATE_PARAM(LinearExtraDamping);
+	INTERPOLATE_PARAM(AngularStrength);
+	INTERPOLATE_PARAM(AngularDampingRatio);
+	INTERPOLATE_PARAM(AngularExtraDamping);
+	INTERPOLATE_PARAM(LinearTargetVelocityMultiplier);
+	INTERPOLATE_PARAM(AngularTargetVelocityMultiplier);
+	Output.bEnabled = (Weight < 0.5f) ? A.bEnabled : B.bEnabled;
+	return Output;
+}
+
+//======================================================================================================================
+FRigidBodyControlSparseData Interpolate(
+	const FRigidBodyControlSparseData& A, const FRigidBodyControlSparseData& B, const float Weight)
+{
+	FRigidBodyControlSparseData Output;
+	INTERPOLATE_PARAM(LinearStrength);
+	INTERPOLATE_PARAM(LinearDampingRatio);
+	INTERPOLATE_PARAM(LinearExtraDamping);
+	INTERPOLATE_PARAM(AngularStrength);
+	INTERPOLATE_PARAM(AngularDampingRatio);
+	INTERPOLATE_PARAM(AngularExtraDamping);
+	INTERPOLATE_PARAM(LinearTargetVelocityMultiplier);
+	INTERPOLATE_PARAM(AngularTargetVelocityMultiplier);
+	Output.bEnabled = (Weight < 0.5f) ? A.bEnabled : B.bEnabled;
+
+	SET_ENABLED_PARAM(LinearStrength);
+	SET_ENABLED_PARAM(LinearDampingRatio);
+	SET_ENABLED_PARAM(LinearExtraDamping);
+	SET_ENABLED_PARAM(AngularStrength);
+	SET_ENABLED_PARAM(AngularDampingRatio);
+	SET_ENABLED_PARAM(AngularExtraDamping);
+	SET_ENABLED_PARAM(LinearTargetVelocityMultiplier);
+	SET_ENABLED_PARAM(AngularTargetVelocityMultiplier);
+	SET_ENABLED_PARAM(bEnabled);
+
+	return Output;
+}
+
+//======================================================================================================================
+FRigidBodyModifierData Interpolate(
+	const FRigidBodyModifierData& A, const FRigidBodyModifierData& B, const float Weight)
+{
+	FRigidBodyModifierData Output;
+	Output.MovementType = (Weight < 0.5f) ? A.MovementType : B.MovementType;
+	INTERPOLATE_PARAM(GravityMultiplier);
+	return Output;
+}
+
+//======================================================================================================================
+FRigidBodyModifierSparseData Interpolate(
+	const FRigidBodyModifierSparseData& A, const FRigidBodyModifierSparseData& B, const float Weight)
+{
+	FRigidBodyModifierSparseData Output;
+	Output.MovementType = (Weight < 0.5f) ? A.MovementType : B.MovementType;
+	INTERPOLATE_PARAM(GravityMultiplier);
+
+	SET_ENABLED_PARAM(MovementType);
+	SET_ENABLED_PARAM(GravityMultiplier);
+	return Output;
+}
+
+#undef INTERPOLATE_PARAM
+#undef SET_ENABLED_PARAM
+
+//======================================================================================================================
+FControlRecord::FControlRecord(const FRigidBodyControl& InControl, ImmediatePhysics::FJointHandle* InJointHandle)
+	: Control(InControl)
+	, JointHandle(InJointHandle)
+	, CurrentData(Control.ControlData)
+{
+}
+
+//======================================================================================================================
+void FControlRecord::ResetCurrent(bool bResetTarget)
+{
+	CurrentData = Control.ControlData;
+
+	if (bResetTarget)
+	{
+		ControlTarget = FRigidBodyControlTarget();
+	}
+}
+
+//======================================================================================================================
+FBodyModifierRecord::FBodyModifierRecord(const FRigidBodyModifier& InModifier, ImmediatePhysics::FActorHandle* InActorHandle)
+	: Modifier(InModifier)
+	, ActorHandle(InActorHandle)
+	, CurrentData(Modifier.ModifierData)
+{
+}
+
+//======================================================================================================================
+void FBodyModifierRecord::ResetCurrent()
+{
+	CurrentData = Modifier.ModifierData;
+}
+
+#define SET_SPARSE_DATA(NAME) NAME = SparseData.bEnable##NAME ? SparseData.NAME : NAME
+
+//======================================================================================================================
+void FRigidBodyControlData::UpdateFromSparseData(const FRigidBodyControlSparseData& SparseData)
+{
+	SET_SPARSE_DATA(LinearStrength);
+	SET_SPARSE_DATA(LinearDampingRatio);
+	SET_SPARSE_DATA(LinearExtraDamping);
+	SET_SPARSE_DATA(AngularStrength);
+	SET_SPARSE_DATA(AngularDampingRatio);
+	SET_SPARSE_DATA(AngularExtraDamping);
+	SET_SPARSE_DATA(LinearTargetVelocityMultiplier);
+	SET_SPARSE_DATA(AngularTargetVelocityMultiplier);
+	SET_SPARSE_DATA(bEnabled);
+}
+
+//======================================================================================================================
+void FRigidBodyModifierData::UpdateFromSparseData(const FRigidBodyModifierSparseData& SparseData)
+{
+	SET_SPARSE_DATA(MovementType);
+	SET_SPARSE_DATA(GravityMultiplier);
+}
+
+#undef SET_SPARSE_DATA
diff --git a/Engine/Plugins/Experimental/PhysicsControl/Source/PhysicsControl/Private/RigidBodyNameRecords.cpp b/Engine/Plugins/Experimental/PhysicsControl/Source/PhysicsControl/Private/RigidBodyNameRecords.cpp
new file mode 100644
--- /dev/null
+++ b/Engine/Plugins/Experimental/PhysicsControl/Source/PhysicsControl/Private/RigidBodyNameRecords.cpp
@@ -0,0 +1,139 @@
+// Copyright Epic Games, Inc. All Rights Reserved.
+#include "RigidBodyNameRecords.h"
+
+//======================================================================================================================
+void AddUniqueNamesToSet(const TArray<FName>& Names, const FName SetName, TMap<FName, TArray<FName>>& SetCollection)
+{
+	TArray<FName>& SetElementNames = SetCollection.FindOrAdd(SetName);
+
+	for (const FName Name : Names)
+	{
+		SetElementNames.AddUnique(Name);
+	}
+}
+
+//======================================================================================================================
+void FRigidBodyNameRecords::AddControl(FName Name, FName SetName)
+{
+	ControlSets.FindOrAdd(SetName).AddUnique(Name);
+	ControlSets.FindOrAdd("All").AddUnique(Name);
+}
+
+//======================================================================================================================
+void FRigidBodyNameRecords::AddControl(FName Name, const TArray<FName>& SetNames)
+{
+	for (const FName SetName : SetNames)
+	{
+		ControlSets.FindOrAdd(SetName).AddUnique(Name);
+	}
+	ControlSets.FindOrAdd("All").AddUnique(Name);
+}
+
+
+//======================================================================================================================
+void FRigidBodyNameRecords::AddControls(const TArray<FName>& ControlNames, const FName SetName)
+{
+	AddUniqueNamesToSet(ControlNames, SetName, ControlSets);
+	AddUniqueNamesToSet(ControlNames, "All", ControlSets);
+}
+
+//======================================================================================================================
+void FRigidBodyNameRecords::RemoveControl(FName Name)
+{
+	for (TPair<FName, TArray<FName>>& Set : ControlSets)
+	{
+		Set.Value.Remove(Name);
+	}
+}
+
+//======================================================================================================================
+const TArray<FName>& FRigidBodyNameRecords::GetControlNamesInSet(FName SetName) const
+{
+	const TArray<FName>* Names = ControlSets.Find(SetName);
+	if (Names)
+	{
+		return *Names;
+	}
+	static TArray<FName> FailureResult;
+	return FailureResult;
+}
+
+//======================================================================================================================
+void FRigidBodyNameRecords::AddBodyModifier(FName Name, FName SetName)
+{
+	BodyModifierSets.FindOrAdd(SetName).AddUnique(Name);
+	BodyModifierSets.FindOrAdd("All").AddUnique(Name);
+}
+
+//======================================================================================================================
+void FRigidBodyNameRecords::AddBodyModifier(FName Name, const TArray<FName>& SetNames)
+{
+	for (FName SetName : SetNames)
+	{
+		BodyModifierSets.FindOrAdd(SetName).AddUnique(Name);
+	}
+	BodyModifierSets.FindOrAdd("All").AddUnique(Name);
+}
+
+//======================================================================================================================
+void FRigidBodyNameRecords::AddBodyModifiers(const TArray<FName>& BodyModifierNames, const FName SetName)
+{
+	AddUniqueNamesToSet(BodyModifierNames, SetName, BodyModifierSets);
+	AddUniqueNamesToSet(BodyModifierNames, "All", BodyModifierSets);
+}
+
+//======================================================================================================================
+void FRigidBodyNameRecords::RemoveBodyModifier(FName Name)
+{
+	for (TPair<FName, TArray<FName>>& Set : BodyModifierSets)
+	{
+		Set.Value.Remove(Name);
+	}
+}
+
+//======================================================================================================================
+const TArray<FName>& FRigidBodyNameRecords::GetBodyModifierNamesInSet(FName SetName) const
+{
+	const TArray<FName>* Names = BodyModifierSets.Find(SetName);
+	if (Names)
+	{
+		return *Names;
+	}
+	static TArray<FName> FailureResult;
+	return FailureResult;
+}
+
+//======================================================================================================================
+void FRigidBodyNameRecords::Reset()
+{
+	ControlSets.Reset();
+	BodyModifierSets.Reset();
+}
+
+//======================================================================================================================
+TArray<FName> ExpandSetName(const FName InName, const TMap<FName, TArray<FName>>& SetNames)
+{
+	TArray<FName> OutputNames;
+	if (const TArray<FName>* const FoundSet = SetNames.Find(InName))
+	{
+		OutputNames.Append(*FoundSet);
+	}
+	else
+	{
+		OutputNames.AddUnique(InName);
+	}
+	return OutputNames;
+}
+
+//======================================================================================================================
+TArray<FName> ExpandSetNames(const TArray<FName>& InNames, const TMap<FName, TArray<FName>>& SetNames)
+{
+	TArray<FName> OutputNames;
+
+	for (const FName Name : InNames)
+	{
+		OutputNames.Append(ExpandSetName(Name, SetNames));
+	}
+
+	return OutputNames;
+}
diff --git a/Engine/Plugins/Experimental/PhysicsControl/Source/PhysicsControl/Private/RigidBodyPoseData.cpp b/Engine/Plugins/Experimental/PhysicsControl/Source/PhysicsControl/Private/RigidBodyPoseData.cpp
new file mode 100644
--- /dev/null
+++ b/Engine/Plugins/Experimental/PhysicsControl/Source/PhysicsControl/Private/RigidBodyPoseData.cpp
@@ -0,0 +1,50 @@
+// Copyright Epic Games, Inc. All Rights Reserved.
+
+#include "RigidBodyPoseData.h"
+#include "Animation/AnimNodeBase.h"
+#include "Animation/AnimInstanceProxy.h"
+
+namespace RigidBodyWithControl
+{
+
+//======================================================================================================================
+void FRigidBodyPoseData::Update(
+	FComponentSpacePoseContext&   ComponentSpacePoseContext,
+	const TArray<FOutputBoneData> OutputBoneData,
+	const ESimulationSpace        SimulationSpace,
+	const FBoneReference&         BaseBoneRef,
+	const FGraphTraversalCounter& InUpdateCounter)
+{
+	ExpectedUpdateCounter = UpdateCounter;
+	ExpectedUpdateCounter.Increment();
+	UpdateCounter = InUpdateCounter;
+
+	const FTransform CompWorldSpaceTM = ComponentSpacePoseContext.AnimInstanceProxy->GetComponentTransform();
+	const FBoneContainer& BoneContainer = ComponentSpacePoseContext.Pose.GetPose().GetBoneContainer();
+	const FTransform BaseBoneTM = ComponentSpacePoseContext.Pose.GetComponentSpaceTransform(
+		BaseBoneRef.GetCompactPoseIndex(BoneContainer));
+
+	if (BoneData.Num() == OutputBoneData.Num())
+	{
+		for (const FOutputBoneData& OutputData : OutputBoneData)
+		{
+			const int32 BodyIndex = OutputData.BodyIndex;
+			const FTransform& ComponentSpaceTM = ComponentSpacePoseContext.Pose.GetComponentSpaceTransform(OutputData.CompactPoseBoneIndex);
+			const FTransform BodyTM = ConvertCSTransformToSimSpace(SimulationSpace, ComponentSpaceTM, CompWorldSpaceTM, BaseBoneTM);
+			BoneData[BodyIndex].Set(BodyTM.GetLocation(), BodyTM.GetRotation());
+		}
+	}
+	else
+	{
+		BoneData.Empty(OutputBoneData.Num());
+		for (const FOutputBoneData& OutputData : OutputBoneData)
+		{
+			const int32 BodyIndex = OutputData.BodyIndex;
+			const FTransform& ComponentSpaceTM = ComponentSpacePoseContext.Pose.GetComponentSpaceTransform(OutputData.CompactPoseBoneIndex);
+			const FTransform BodyTM = ConvertCSTransformToSimSpace(SimulationSpace, ComponentSpaceTM, CompWorldSpaceTM, BaseBoneTM);
+			BoneData.Emplace(BodyTM.GetLocation(), BodyTM.GetRotation());
+		}
+	}
+}
+
+}
\ No newline at end of file
diff --git a/Engine/Plugins/Experimental/PhysicsControl/Source/PhysicsControl/Private/RigidBodyWithControl.cpp b/Engine/Plugins/Experimental/PhysicsControl/Source/PhysicsControl/Private/RigidBodyWithControl.cpp
new file mode 100644
--- /dev/null
+++ b/Engine/Plugins/Experimental/PhysicsControl/Source/PhysicsControl/Private/RigidBodyWithControl.cpp
@@ -0,0 +1,1238 @@
+// Copyright Epic Games, Inc. All Rights Reserved.
+
+#include "AnimNode_RigidBodyWithControl.h"
+
+#include "Animation/AnimInstance.h"
+#include "Animation/AnimInstanceProxy.h"
+#include "Chaos/PBDJointConstraints.h"
+#include "Chaos/PBDJointConstraintTypes.h"
+#include "Components/StaticMeshComponent.h"
+#include "Physics/ImmediatePhysics/ImmediatePhysicsActorHandle.h"
+#include "Physics/ImmediatePhysics/ImmediatePhysicsJointHandle.h"
+#include "Physics/ImmediatePhysics/ImmediatePhysicsSimulation.h"
+#include "PhysicsEngine/PhysicsConstraintTemplate.h"
+#include "PhysicsEngine/PhysicsSettings.h"
+#include "Chaos/ChaosConstraintSettings.h"
+#include "PhysicsControlLog.h"
+
+//UE_DISABLE_OPTIMIZATION;
+
+int32 MaxNumControlsOrModifiersPerName = 16;
+
+constexpr int32 ConstraintChildIndex = 0;
+constexpr int32 ConstraintParentIndex = 1;
+
+//======================================================================================================================
+static void SetConstraintEnabled(ImmediatePhysics::FJointHandle* const JointHandle, const bool bIsEnabled)
+{
+	if (JointHandle)
+	{
+		if (ImmediatePhysics::FJointHandle::FChaosConstraintHandle* ConstraintHandle = JointHandle->GetConstraint())
+		{
+			ConstraintHandle->SetConstraintEnabled(bIsEnabled);
+		}
+	}
+}
+
+//======================================================================================================================
+template<typename TRecord, typename TData> void SetRecordParameters(
+	const FName Name, const TData& Data, TMap<FName, TRecord>& Records)
+{
+	if (TRecord* const RecordSearchResult = Records.Find(Name))
+	{
+		RecordSearchResult->CurrentData.UpdateFromSparseData(Data);
+	}
+	else
+	{
+		UE_LOG(LogRigidBodyWithControl, Warning,
+			TEXT("SetRecordParameters: Failed to find control/modifier with name %s"), *Name.ToString());
+	}
+}
+
+//======================================================================================================================
+template<typename TRecord, typename TParameters> void ApplyControlAndModifierParametersToRecords(
+	TMap<FName, TRecord>& Records, const TArray<TParameters>& AllParameters, const TMap<FName, TArray<FName>>& Sets)
+{
+	for (const TParameters& Parameters : AllParameters)
+	{
+		// Find the list of control records in the target set.
+		const TArray<FName>* const SetSearchResult = Sets.Find(Parameters.Name);
+
+		if (SetSearchResult)
+		{
+			for (const FName& Name : *SetSearchResult)
+			{
+				SetRecordParameters(Name, Parameters.Data, Records);
+			}
+		}
+		else
+		{
+			// No Set found with a matching name - try to find a control with a matching name.
+			SetRecordParameters(Parameters.Name, Parameters.Data, Records);
+		}
+	}
+}
+
+//======================================================================================================================
+const FTransform FAnimNode_RigidBodyWithControl::GetBodyTransform(const int32 BodyIndex) const
+{
+	if (Bodies.IsValidIndex(BodyIndex))
+	{
+		return Bodies[BodyIndex]->GetWorldTransform();
+	}
+	return FTransform::Identity;
+}
+
+//======================================================================================================================
+const FTransform FAnimNode_RigidBodyWithControl::GetWorldSpaceControlRootTransform() const
+{
+	if (WorldSpaceControlActorHandle)
+	{
+		return WorldSpaceControlActorHandle->GetWorldTransform();
+	}
+	return FTransform::Identity;
+}
+
+//======================================================================================================================
+void FAnimNode_RigidBodyWithControl::CreateWorldSpaceControlRootBody(UPhysicsAsset* const PhysicsAsset)
+{
+	WorldSpaceControlActorHandle = PhysicsSimulation->CreateActor(
+		ImmediatePhysics::EActorType::KinematicActor, nullptr, FTransform());
+	if (WorldSpaceControlActorHandle)
+	{
+		WorldSpaceControlActorHandle->SetName(FName("WorldSpaceControlHandle"));
+	}
+	else
+	{
+		UE_LOG(LogRigidBodyWithControl, Error, TEXT("Failed to create world space control root actor"));
+	}
+}
+
+//======================================================================================================================
+ImmediatePhysics::FJointHandle* FAnimNode_RigidBodyWithControl::CreateConstraint(
+	ImmediatePhysics::FActorHandle* ChildActorHandle, 
+	ImmediatePhysics::FActorHandle* ParentActorHandle)
+{
+	ImmediatePhysics::FJointHandle* JointHandle = nullptr;
+
+	if (PhysicsSimulation && ChildActorHandle && ParentActorHandle)
+	{
+		Chaos::FPBDJointSettings Settings;
+
+		Settings.LinearMotionTypes = { 
+			Chaos::EJointMotionType::Free, Chaos::EJointMotionType::Free, Chaos::EJointMotionType::Free };
+		Settings.AngularMotionTypes = { 
+			Chaos::EJointMotionType::Free, Chaos::EJointMotionType::Free, Chaos::EJointMotionType::Free };
+
+		Settings.bLinearPositionDriveEnabled = { true, true, true };
+		Settings.bLinearVelocityDriveEnabled = { true, true, true }; 
+		Settings.LinearDriveForceMode = Chaos::EJointForceMode::Acceleration;
+
+		Settings.bAngularSLerpPositionDriveEnabled = true;
+		Settings.bAngularSLerpVelocityDriveEnabled = true;
+
+		Settings.bAngularTwistPositionDriveEnabled = false;
+		Settings.bAngularTwistVelocityDriveEnabled = false;
+		Settings.bAngularSwingPositionDriveEnabled = false;
+		Settings.bAngularSwingVelocityDriveEnabled = false;
+		Settings.AngularDriveForceMode = Chaos::EJointForceMode::Acceleration;
+
+		Settings.bMassConditioningEnabled = false; // TODO needed/wanted?
+		Settings.bCollisionEnabled = false; // TODO needed?
+
+		FVector ChildCoMPositionOffset = ChildActorHandle->GetLocalCoMTransform().GetLocation();
+		Settings.ConnectorTransforms[ConstraintChildIndex].SetLocation(ChildCoMPositionOffset);
+
+		Settings.Sanitize();
+		JointHandle = PhysicsSimulation->CreateJoint(Settings, ChildActorHandle, ParentActorHandle);
+	}
+
+	if (JointHandle)
+	{
+		SetConstraintEnabled(JointHandle, false); // Disable constraints by default.
+	}
+
+	return JointHandle;
+}
+
+//======================================================================================================================
+int32 FAnimNode_RigidBodyWithControl::AddBody(ImmediatePhysics::FActorHandle* const BodyHandle)
+{
+	const int32 BodyIndex = Bodies.Add(BodyHandle);
+
+	if (BodyHandle != nullptr)
+	{
+		BodyNameToIndexMap.Add(BodyHandle->GetName(), BodyIndex);
+	}
+
+	return BodyIndex;
+}
+
+//======================================================================================================================
+int32 FAnimNode_RigidBodyWithControl::FindBodyIndexFromBoneName(const FName BoneName) const
+{
+	const int32* const FoundIndex = BodyNameToIndexMap.Find(BoneName);
+
+	return (FoundIndex != nullptr) ? *FoundIndex : INDEX_NONE;
+}
+
+//======================================================================================================================
+ImmediatePhysics::FActorHandle* FAnimNode_RigidBodyWithControl::FindBodyFromBoneName(const FName BoneName) const
+{
+	ImmediatePhysics::FActorHandle* BodyHandle = nullptr;
+	const int32 BodyIndex = FindBodyIndexFromBoneName(BoneName);
+	if (BodyIndex != INDEX_NONE)
+	{
+		BodyHandle = Bodies[BodyIndex];
+	}
+	return BodyHandle;
+}
+
+//======================================================================================================================
+FName FAnimNode_RigidBodyWithControl::FindParentBoneNameFromBoneName(const FName BoneName) const
+{
+	// TODO - would be better if we didn't need to access the physics structures to find the parent bone name.
+	FName ParentName;
+	const int32 JointIndex = FindBodyIndexFromBoneName(BoneName);
+	if (JointIndex != INDEX_NONE)
+	{
+		ImmediatePhysics::FJointHandle* const JointHandle = Joints[JointIndex];
+
+		if (JointHandle && JointHandle->GetActorHandles()[ConstraintParentIndex])
+		{
+			ParentName = JointHandle->GetActorHandles()[ConstraintParentIndex]->GetName();
+		}
+	}
+	return ParentName;
+}
+
+//======================================================================================================================
+FName FAnimNode_RigidBodyWithControl::CreateControl(
+	const FName ParentBoneName, const FName ChildBoneName, const FRigidBodyControlData& ControlData)
+{
+	FRigidBodyControl Control;
+	Control.ParentBoneName = ParentBoneName;
+	Control.ChildBoneName = ChildBoneName;
+	Control.ControlData = ControlData;
+
+	ImmediatePhysics::FJointHandle* JointHandle = nullptr;
+
+	ImmediatePhysics::FActorHandle* ParentBodyHandle;
+
+	if (ParentBoneName.IsNone())
+	{
+		ParentBodyHandle = WorldSpaceControlActorHandle;
+	}
+	else
+	{
+		ParentBodyHandle = FindBodyFromBoneName(ParentBoneName);
+	}
+	ImmediatePhysics::FActorHandle* const ChildBodyHandle = FindBodyFromBoneName(ChildBoneName);
+	JointHandle = CreateConstraint(ChildBodyHandle, ParentBodyHandle);
+
+	if (!JointHandle)
+	{
+		UE_LOG(LogRigidBodyWithControl, Warning,
+			TEXT("Unable to create world space control constraint for bone %s"), *ChildBoneName.ToString());
+		return FName();
+	}
+	
+	FName ControlName = GetUniqueControlName(Control.ParentBoneName, Control.ChildBoneName);
+	ControlRecords.Add(ControlName, FControlRecord(Control, JointHandle));
+
+	return ControlName;
+}
+
+//======================================================================================================================
+void FAnimNode_RigidBodyWithControl::CreateControlsFromLimbBones(
+	const FName                   LimbName,
+	const FRigidBodyLimbBones&    LimbBones,
+	const ERigidBodyControlType   ControlType,
+	const FRigidBodyControlData&  ControlData)
+{
+	for (const FName ChildBoneName : LimbBones.BoneNames)
+	{
+		FName ParentBoneName;
+
+		if (ControlType == ERigidBodyControlType::ParentSpace)
+		{
+			ParentBoneName = FindParentBoneNameFromBoneName(ChildBoneName);
+			if (ParentBoneName.IsNone())
+			{
+				// This happens for the pelvis, for example - we only create parent space
+				// controls if there's a parent!
+				continue;
+			}
+		}
+
+		FName ControlName = CreateControl(ParentBoneName, ChildBoneName, ControlData);
+		if (ControlName.IsNone())
+		{
+			UE_LOG(LogRigidBodyWithControl, Warning,
+				TEXT("Failed to make control for %s"), *ChildBoneName.ToString());
+		}
+		else
+		{
+			NameRecords.AddControl(ControlName, LimbName);
+			NameRecords.AddControl(ControlName, GetControlTypeName(ControlType));
+			NameRecords.AddControl(ControlName, FName(
+				GetControlTypeName(ControlType).ToString().Append("_").Append(LimbName.ToString())));
+		}
+	}
+}
+
+//======================================================================================================================
+FName FAnimNode_RigidBodyWithControl::GetBodyFromBoneName(const FName BoneName) const
+{
+	if (const FName* Name = BoneToBodyNameMap.Find(BoneName))
+	{
+		return *Name;
+	}
+	return BoneName;
+}
+
+//======================================================================================================================
+FName FAnimNode_RigidBodyWithControl::GetUniqueBodyModifierName(const FName BoneName)
+{
+	FName BodyName = GetBodyFromBoneName(BoneName);
+
+	FString NameBase = TEXT("");
+	if (!BodyName.IsNone())
+	{
+		NameBase += BodyName.ToString();
+	}
+	else
+	{
+		NameBase = TEXT("Body");
+	}
+
+	TSet<FName> Keys;
+	ModifierRecords.GetKeys(Keys);
+
+	if (!Keys.Find(FName(NameBase)))
+	{
+		return FName(NameBase);
+	}
+
+	// If the number gets too large, almost certainly we're in some nasty situation where this is
+	// getting called in a loop. Better to quit and fail, rather than allow the modifier set to
+	// increase without bound. 
+	for (int32 Index = 0; Index != MaxNumControlsOrModifiersPerName; ++Index)
+	{
+		FString NameStr = FString::Format(TEXT("{0}_{1}"), { NameBase, Index });
+		FName Name(NameStr);
+		if (!Keys.Find(Name))
+		{
+			return Name;
+		}
+	}
+	UE_LOG(LogRigidBodyWithControl, Warning,
+		TEXT("Unable to find a suitable Body Modifier name - the limit of MaxNumControlsOrModifiersPerName (%d) has been exceeded"),
+		MaxNumControlsOrModifiersPerName);
+	return FName();
+}
+
+//======================================================================================================================
+FName FAnimNode_RigidBodyWithControl::GetUniqueControlName(const FName ParentBoneName, const FName ChildBoneName)
+{
+	FName ParentBodyName = GetBodyFromBoneName(ParentBoneName);
+	FName ChildBodyName = GetBodyFromBoneName(ChildBoneName);
+
+	FString NameBase = TEXT("");
+	if (!ParentBodyName.IsNone())
+	{
+		NameBase += ParentBodyName.ToString() + TEXT("_");
+	}
+	if (!ChildBodyName.IsNone())
+	{
+		NameBase += ChildBodyName.ToString();
+	}
+
+	TSet<FName> Keys;
+	ControlRecords.GetKeys(Keys);
+
+	if (!Keys.Find(FName(NameBase)))
+	{
+		return FName(NameBase);
+	}
+
+	// If the number gets too large, almost certainly we're in some nasty situation where this is
+	// getting called in a loop. Better to quit and fail, rather than allow the constraint set to
+	// increase without bound. 
+	for (int32 Index = 0; Index < MaxNumControlsOrModifiersPerName; ++Index)
+	{
+		FString NameStr = FString::Format(TEXT("{0}_{1}"), { NameBase, Index });
+		FName Name(NameStr);
+		if (!Keys.Find(Name))
+		{
+			return Name;
+		}
+	}
+	UE_LOG(LogRigidBodyWithControl, Warning,
+		TEXT("Unable to find a suitable Control name - the limit of MaxNumControlsOrModifiersPerName (%d) has been exceeded"),
+		MaxNumControlsOrModifiersPerName);
+	return FName();
+}
+
+//======================================================================================================================
+FName FAnimNode_RigidBodyWithControl::CreateBodyModifier(FName BoneName, const FRigidBodyModifierData& ModifierData)
+{
+	FName Name;
+
+	ImmediatePhysics::FActorHandle* const ActorHandle = FindBodyFromBoneName(BoneName);
+	if (ActorHandle)
+	{
+		Name = GetUniqueBodyModifierName(BoneName);
+
+		FRigidBodyModifier BodyModifier;
+		BodyModifier.BoneName = BoneName;
+		BodyModifier.ModifierData = ModifierData;
+		FBodyModifierRecord& Modifier = ModifierRecords.Add(Name, FBodyModifierRecord(BodyModifier, ActorHandle));
+	}
+
+	return Name;
+}
+
+//======================================================================================================================
+void FAnimNode_RigidBodyWithControl::CreateBodyModifiersFromLimbBones(
+	const FName                   LimbName,
+	const FRigidBodyLimbBones&    LimbBones,
+	const FRigidBodyModifierData& DefaultModifierData)
+{
+	for (const FName BoneName : LimbBones.BoneNames)
+	{
+		const FName BodyModifierName = CreateBodyModifier(BoneName, DefaultModifierData);
+		if (BodyModifierName.IsNone())
+		{
+			UE_LOG(LogRigidBodyWithControl, Warning,
+				TEXT("Failed to make body modifier for %s"), *BoneName.ToString());
+		}
+		else
+		{
+			NameRecords.AddBodyModifier(BodyModifierName, LimbName);
+		}
+	}
+}
+
+//======================================================================================================================
+TMap<FName, FRigidBodyLimbBones> FAnimNode_RigidBodyWithControl::GetLimbBonesFromSkeletalMesh(
+	const TArray<FRigidBodyLimbSetupData>& LimbSetupData,
+	USkeletalMeshComponent* const SkeletalMeshComponent) const
+{
+	// Parse the skeleton tree to figure out which bones are associated with which limbs. 
+	
+	// TODO - Output limb bones are not in the order specified in the skeleton - would be better if they were
+	
+	TMap<FName, FRigidBodyLimbBones> Result;
+
+	UPhysicsAsset* const PhysicsAsset = GetPhysicsAsset();
+
+	if (!PhysicsAsset)
+	{
+		UE_LOG(LogRigidBodyWithControl, Warning, TEXT("Physics asset missing"));
+		return Result;
+	}
+
+	if (!SkeletalMeshComponent)
+	{
+		UE_LOG(LogRigidBodyWithControl, Warning, TEXT("Invalid Skeletal Mesh Component"));
+		return Result;
+	}
+
+	USkeletalMesh* const SkelMesh = SkeletalMeshComponent->GetSkeletalMeshAsset();
+
+	TSet<FName> AllBones;
+
+	for (const FRigidBodyLimbSetupData& LimbSetup : LimbSetupData)
+	{
+		FRigidBodyLimbBones& LimbBones = Result.Add(LimbSetup.LimbName);
+
+		LimbBones.bFirstBoneIsAdditional = false;
+		LimbBones.bCreateWorldSpaceControls = LimbSetup.bCreateWorldSpaceControls;
+		LimbBones.bCreateParentSpaceControls = LimbSetup.bCreateParentSpaceControls;
+		LimbBones.bCreateBodyModifiers = LimbSetup.bCreateBodyModifiers;
+
+		if (LimbSetup.bIncludeParentBone)
+		{		
+			const int32 StartBoneIndex = SkelMesh->GetRefSkeleton().FindBoneIndex(LimbSetup.StartBone);
+			const int32 ParentBodyIndex = PhysicsAsset->FindParentBodyIndex(SkelMesh, StartBoneIndex);
+			
+			FName ParentBoneName;
+
+			if (PhysicsAsset->SkeletalBodySetups.IsValidIndex(ParentBodyIndex) && 
+				PhysicsAsset->SkeletalBodySetups[ParentBodyIndex])
+			{
+				ParentBoneName = PhysicsAsset->SkeletalBodySetups[ParentBodyIndex]->BoneName;
+			}
+
+			if (!ParentBoneName.IsNone())
+			{
+				LimbBones.BoneNames.Add(ParentBoneName);
+				AllBones.Add(ParentBoneName);
+				LimbBones.bFirstBoneIsAdditional = true;
+			}
+		}
+
+		TArray<int32> ChildBodyIndices;
+		PhysicsAsset->GetBodyIndicesBelow(ChildBodyIndices, LimbSetup.StartBone, SkelMesh);
+
+		for (int32 ChildBodyIndex : ChildBodyIndices)
+		{
+			const FName BoneName = PhysicsAsset->SkeletalBodySetups[ChildBodyIndex]->BoneName;
+
+			if (!AllBones.Find(BoneName))
+			{
+				LimbBones.BoneNames.Add(BoneName);
+				AllBones.Add(BoneName);
+			}
+		}
+	}
+
+	return Result;
+}
+
+//======================================================================================================================
+void FAnimNode_RigidBodyWithControl::InitControlsAndBodyModifiers(USkeletalMeshComponent* const SkeletalMeshComponent)
+{
+	check(ControlRecords.IsEmpty()); // Controls should not exist when this function is called.
+
+	// These functions will create the base set of controls and modifiers from SetupData
+	TMap<FName, FRigidBodyLimbBones> AllLimbBones = 
+		GetLimbBonesFromSkeletalMesh(SetupData.LimbSetupData, SkeletalMeshComponent);
+
+	for (const TMap<FName, FRigidBodyLimbBones>::ElementType& LimbBoneEntry : AllLimbBones)
+	{
+		const FName LimbName = LimbBoneEntry.Key;
+		const FRigidBodyLimbBones& LimbBones = LimbBoneEntry.Value;
+		if (LimbBones.bCreateWorldSpaceControls)
+		{
+			CreateControlsFromLimbBones(
+				LimbName, LimbBones, ERigidBodyControlType::WorldSpace, SetupData.DefaultWorldSpaceControlData);
+		}
+		if (LimbBones.bCreateParentSpaceControls)
+		{
+			CreateControlsFromLimbBones(
+				LimbName, LimbBones, ERigidBodyControlType::ParentSpace, SetupData.DefaultParentSpaceControlData);
+		}
+		if (LimbBones.bCreateBodyModifiers)
+		{
+			CreateBodyModifiersFromLimbBones(LimbName, LimbBones, SetupData.DefaultBodyModifierData);
+		}
+	}
+
+	// Create any additional controls that have been requested
+	CreateAdditionalControls();
+
+	// Create any additional sets that have been requested
+	CreateAdditionalSets();
+
+	// Apply control and modifier parameters on a single-use basis
+	ApplyControlAndBodyModifierDatas(
+		InitialControlAndBodyModifierUpdates.ControlParameters,
+		InitialControlAndBodyModifierUpdates.ModifierParameters);
+
+	// Tell the poor user what we've done
+	LogControlsModifiersAndSets();
+}
+
+//======================================================================================================================
+void FAnimNode_RigidBodyWithControl::DestroyControlsAndBodyModifiers()
+{
+	// TODO - make sure this fn is complete
+
+	ControlRecords.Reset();
+	ModifierRecords.Reset();
+	NameRecords.Reset();
+
+	bHaveRunControlSeup = false;
+}
+
+//======================================================================================================================
+// TODO This isn't ideal as it only dumps out the "original" values, not including the updates
+void FAnimNode_RigidBodyWithControl::LogControlsModifiersAndSets()
+{
+#define RBWC_LOG_LEVEL Display
+
+	UE_LOG(LogRigidBodyWithControl, RBWC_LOG_LEVEL, TEXT("Controls:"));
+	for (TMap<FName, FControlRecord>::ElementType& NameRecordPair : ControlRecords)
+	{
+		UE_LOG(LogRigidBodyWithControl, RBWC_LOG_LEVEL, TEXT("  %s:"), *NameRecordPair.Key.ToString());
+		const FControlRecord& Record = NameRecordPair.Value;
+		UE_LOG(LogRigidBodyWithControl, RBWC_LOG_LEVEL, TEXT("    Parent bone: %s Child bone: %s"), 
+			*Record.Control.ParentBoneName.ToString(), *Record.Control.ChildBoneName.ToString());
+		UE_LOG(LogRigidBodyWithControl, RBWC_LOG_LEVEL, TEXT("    Enabled %d"),
+			Record.IsEnabled());
+		UE_LOG(LogRigidBodyWithControl, RBWC_LOG_LEVEL, TEXT("    Linear: Strength %f DampingRatio %f ExtraDamping %f"),
+			Record.Control.ControlData.LinearStrength, 
+			Record.Control.ControlData.LinearDampingRatio, 
+			Record.Control.ControlData.LinearExtraDamping);
+		UE_LOG(LogRigidBodyWithControl, RBWC_LOG_LEVEL, TEXT("    Angular: Strength %f DampingRatio %f ExtraDamping %f"),
+			Record.Control.ControlData.AngularStrength, 
+			Record.Control.ControlData.AngularDampingRatio, 
+			Record.Control.ControlData.AngularExtraDamping);
+	}
+
+	UE_LOG(LogRigidBodyWithControl, RBWC_LOG_LEVEL, TEXT("Body Modifiers:"));
+	for (TMap<FName, FBodyModifierRecord>::ElementType& NameRecordPair : ModifierRecords)
+	{
+		UE_LOG(LogRigidBodyWithControl, RBWC_LOG_LEVEL, TEXT("  %s:"), *NameRecordPair.Key.ToString());
+		const FBodyModifierRecord& Record = NameRecordPair.Value;
+		UE_LOG(LogRigidBodyWithControl, RBWC_LOG_LEVEL, TEXT("    Bone: %s Body: %s"),
+			*Record.Modifier.BoneName.ToString(), *GetBodyFromBoneName(Record.Modifier.BoneName).ToString());
+		UE_LOG(LogRigidBodyWithControl, RBWC_LOG_LEVEL, TEXT("    Movement: %s GravityMultiplier: %f"),
+			*GetControlTypeName(Record.Modifier.ModifierData.MovementType).ToString(),
+			Record.Modifier.ModifierData.GravityMultiplier);
+	}
+
+	UE_LOG(LogRigidBodyWithControl, RBWC_LOG_LEVEL, TEXT("Control sets:"));
+	for (TMap<FName, TArray<FName>>::ElementType& Pair : NameRecords.ControlSets)
+	{
+		UE_LOG(LogRigidBodyWithControl, RBWC_LOG_LEVEL, TEXT("  %s:"), *Pair.Key.ToString());
+		const TArray<FName>& Names = Pair.Value;
+		for (FName Name : Names)
+		{
+			UE_LOG(LogRigidBodyWithControl, RBWC_LOG_LEVEL, TEXT("    %s:"), *Name.ToString());
+		}
+	}
+
+	UE_LOG(LogRigidBodyWithControl, RBWC_LOG_LEVEL, TEXT("Body Modifier sets:"));
+	for (TMap<FName, TArray<FName>>::ElementType& Pair : NameRecords.BodyModifierSets)
+	{
+		UE_LOG(LogRigidBodyWithControl, RBWC_LOG_LEVEL, TEXT("  %s:"), *Pair.Key.ToString());
+		const TArray<FName>& Names = Pair.Value;
+		for (FName Name : Names)
+		{
+			UE_LOG(LogRigidBodyWithControl, RBWC_LOG_LEVEL, TEXT("    %s:"), *Name.ToString());
+		}
+	}
+#undef RBWC_LOG_LEVEL
+}
+
+
+//======================================================================================================================
+void FAnimNode_RigidBodyWithControl::CreateAdditionalControls()
+{
+	// Create additional controls.
+	for (const FRigidBodyControlCreation& Control : AdditionalControlsAndBodyModifiers.Controls)
+	{	
+		FName ControlName = CreateControl(
+			Control.Control.ParentBoneName, Control.Control.ChildBoneName, Control.Control.ControlData);
+		if (ControlName.IsNone())
+		{
+			UE_LOG(LogRigidBodyWithControl, Warning,
+				TEXT("CreateAdditionalControls: Failed to make control between %s and %s"), 
+				*Control.Control.ParentBoneName.ToString(), *Control.Control.ChildBoneName.ToString());
+		}
+		else
+		{
+			UE_LOG(LogRigidBodyWithControl, Verbose,
+				TEXT("Made control %s between %s and %s"),
+				*ControlName.ToString(), 
+				*Control.Control.ParentBoneName.ToString(), 
+				*Control.Control.ChildBoneName.ToString());
+			NameRecords.AddControl(ControlName, Control.Sets);
+		}
+	}
+
+	// Create additional modifiers.
+	for (const FRigidBodyModifierCreation& Modifier : AdditionalControlsAndBodyModifiers.Modifiers)
+	{
+		FName ModifierName = CreateBodyModifier(Modifier.Modifier.BoneName, Modifier.Modifier.ModifierData);
+		if (ModifierName.IsNone())
+		{
+			UE_LOG(LogRigidBodyWithControl, Warning,
+				TEXT("CreateAdditionalControls: Failed to make body modifier for %s"), *Modifier.Modifier.BoneName.ToString());
+		}
+		else
+		{
+			UE_LOG(LogRigidBodyWithControl, Verbose,
+				TEXT("Made modifier %s for %s"),
+				*ModifierName.ToString(),
+				*Modifier.Modifier.BoneName.ToString());
+			NameRecords.AddControl(ModifierName, Modifier.Sets);
+		}
+	}
+}
+
+//======================================================================================================================
+// Slightly annoying to have to add the names individually, but we want to check they exist
+void FAnimNode_RigidBodyWithControl::CreateAdditionalSets()
+{
+	for (const FSetUpdate& Set : AdditionalSets.ControlSetUpdates)
+	{
+		TArray<FName> Names = ExpandSetNames(Set.Names, NameRecords.ControlSets);
+		for (FName Name : Names)
+		{
+			if (ControlRecords.Find(Name))
+			{
+				NameRecords.AddControl(Name, Set.SetName);
+			}
+			else
+			{
+				UE_LOG(LogRigidBodyWithControl, Warning,
+					TEXT("CreateAdditionalSets: Failed to find control with name %s to add to set %s"), 
+					*Name.ToString(), *Set.SetName.ToString());
+			}
+		}
+	}
+
+	for (const FSetUpdate& Set : AdditionalSets.ModifierSetUpdates)
+	{
+		TArray<FName> Names = ExpandSetNames(Set.Names, NameRecords.BodyModifierSets);
+		for (FName Name : Names)
+		{
+			if (ModifierRecords.Find(Name))
+			{
+				NameRecords.AddBodyModifier(Name, Set.SetName);
+			}
+			else
+			{
+				UE_LOG(LogRigidBodyWithControl, Warning,
+					TEXT("CreateAdditionalSets: Failed to find body modifier with name %s to add to set %s"),
+					*Name.ToString(), *Set.SetName.ToString());
+			}
+		}
+	}
+}
+
+
+//======================================================================================================================
+template<typename TOut>
+static void ConvertStrengthToSpringParams(
+	TOut& OutSpring, TOut& OutDamping,
+	double InStrength, double InDampingRatio, double InExtraDamping)
+{
+	TOut AngularFrequency = TOut(InStrength * UE_DOUBLE_TWO_PI);
+	TOut Stiffness = AngularFrequency * AngularFrequency;
+
+	OutSpring = Stiffness;
+	OutDamping = TOut(InExtraDamping + 2 * InDampingRatio * AngularFrequency);
+}
+
+//======================================================================================================================
+void FAnimNode_RigidBodyWithControl::UpdateDriveSpringDamperSettings(
+	Chaos::FPBDJointSettings&             Settings, 
+	const FRigidBodyControlData&          ControlData)
+{
+	float LinearSpring;
+	float LinearDamping;
+	float AngularSpring;
+	float AngularDamping;
+
+	ConvertStrengthToSpringParams(LinearSpring, LinearDamping, 
+		ControlData.LinearStrength, ControlData.LinearDampingRatio, ControlData.LinearExtraDamping);
+	ConvertStrengthToSpringParams(AngularSpring, AngularDamping, 
+		ControlData.AngularStrength, ControlData.AngularDampingRatio, ControlData.AngularExtraDamping);
+
+	Settings.LinearDriveStiffness = { LinearSpring, LinearSpring, LinearSpring };
+	Settings.LinearDriveDamping = { LinearDamping, LinearDamping, LinearDamping };
+	Settings.LinearDriveMaxForce = { 0, 0, 0};
+
+	Settings.AngularDriveStiffness = { AngularSpring, AngularSpring, AngularSpring };
+	Settings.AngularDriveDamping = { AngularDamping, AngularDamping, AngularDamping };
+	Settings.AngularDriveMaxTorque = { 0, 0, 0 };
+}
+
+//======================================================================================================================
+static FTransform CalculateTargetTM(
+	const Chaos::FPBDJointSettings&                 JointSettings, 
+	const RigidBodyWithControl::FRigidBodyPoseData& PoseData,
+	const int32                                     ParentBodyIndex, 
+	const int32                                     ChildBodyIndex)
+{
+	FTransform ChildTargetTM = 
+		JointSettings.ConnectorTransforms[ConstraintChildIndex] * PoseData.GetTM(ChildBodyIndex);
+	if (ParentBodyIndex >= 0)
+	{
+		const FTransform ParentTargetTM = 
+			JointSettings.ConnectorTransforms[ConstraintParentIndex] * PoseData.GetTM(ParentBodyIndex);
+		return ChildTargetTM * ParentTargetTM.Inverse();
+	}
+	return ChildTargetTM;
+}
+
+//======================================================================================================================
+void FAnimNode_RigidBodyWithControl::ApplyControl(FControlRecord& ControlRecord, float DeltaTime)
+{
+	using namespace ImmediatePhysics;
+	FJointHandle* const JointHandle = ControlRecord.JointHandle;
+
+	if (JointHandle)
+	{
+		Chaos::FPBDJointConstraintHandle* const Constraint = JointHandle->GetConstraint();
+		if (Constraint)
+		{
+			if (ControlRecord.ExpectedUpdateCounter.Get() != PoseData.UpdateCounter.Get())
+			{
+				DeltaTime = 0.0f;
+			}
+
+			const FActorHandle* ChildActorHandle = JointHandle->GetActorHandles()[ConstraintChildIndex];
+			const FActorHandle* ParentActorHandle = JointHandle->GetActorHandles()[ConstraintParentIndex];
+
+			if (ChildActorHandle && ParentActorHandle)
+			{
+				Chaos::FPBDJointSettings JointSettings = Constraint->GetSettings();
+
+				// TODO
+				// - cache settings / previous input parameters to avoid unnecessary repeating
+				//   calculations and making physics API calls every update.
+
+				// Update the target point on the child
+				if (ControlRecord.ControlTarget.bUseTargetPoint)
+				{
+					JointSettings.ConnectorTransforms[ConstraintChildIndex].SetLocation(
+						ControlRecord.ControlTarget.TargetPoint);
+				}
+				else
+				{
+					FVector ChildCoMPositionOffset = ChildActorHandle->GetLocalCoMTransform().GetLocation();
+					JointSettings.ConnectorTransforms[ConstraintChildIndex].SetLocation(ChildCoMPositionOffset);
+				}
+
+				const int32 ChildBodyIndex = FindBodyIndexFromBoneName(ControlRecord.Control.ChildBoneName);
+				const int32 ParentBodyIndex = ControlRecord.Control.ParentBoneName.IsNone() ? 
+					-1 : FindBodyIndexFromBoneName(ControlRecord.Control.ParentBoneName);
+
+				FTransform TargetTM(
+					ControlRecord.ControlTarget.TargetOrientation, 
+					ControlRecord.ControlTarget.TargetPosition);
+
+				if (ControlRecord.ControlTarget.bUseSkeletalAnimation)
+				{
+					FTransform AnimTargetTM = CalculateTargetTM(
+						JointSettings, PoseData, ParentBodyIndex, ChildBodyIndex);
+					TargetTM = TargetTM * AnimTargetTM;
+				}
+
+				JointSettings.LinearDrivePositionTarget = TargetTM.GetTranslation();
+				JointSettings.AngularDrivePositionTarget = TargetTM.GetRotation();
+
+				if (DeltaTime != 0)
+				{
+					FTransform PrevTargetTM = ControlRecord.PrevTargetTM;
+
+					FVector Velocity = (TargetTM.GetTranslation() - PrevTargetTM.GetTranslation()) / DeltaTime;
+					// Note that quats multiply in the opposite order to TMs, and must be in the same hemisphere.
+					const FQuat Q = TargetTM.GetRotation();
+					FQuat PrevQ = PrevTargetTM.GetRotation();
+					PrevQ.EnforceShortestArcWith(Q);
+
+					const FQuat DeltaQ = Q * PrevQ.Inverse();
+					const FVector AngularVelocity = DeltaQ.ToRotationVector() / DeltaTime;
+
+					JointSettings.LinearDriveVelocityTarget = 
+						Velocity * ControlRecord.CurrentData.LinearTargetVelocityMultiplier;
+					JointSettings.AngularDriveVelocityTarget = 
+						AngularVelocity * ControlRecord.CurrentData.AngularTargetVelocityMultiplier;
+				}
+				else
+				{
+					JointSettings.LinearDriveVelocityTarget = FVector::ZeroVector;
+					JointSettings.AngularDriveVelocityTarget = FVector::ZeroVector;
+				}
+
+				UpdateDriveSpringDamperSettings(JointSettings, ControlRecord.CurrentData);
+
+				Constraint->SetSettings(JointSettings);
+
+				ControlRecord.PrevTargetTM = TargetTM;
+				ControlRecord.ExpectedUpdateCounter = PoseData.UpdateCounter;
+				ControlRecord.ExpectedUpdateCounter.Increment();
+			}
+		}
+	}
+}
+
+//======================================================================================================================
+void FAnimNode_RigidBodyWithControl::ApplyModifier(
+	const FBodyModifierRecord& BodyModifierRecord, 
+	const FVector&             SimSpaceGravity)
+{
+	if (BodyModifierRecord.ActorHandle)
+	{
+		// Note that there's an early out if there's no change needed, so this should be OK.
+		BodyModifierRecord.ActorHandle->SetIsKinematic(
+			BodyModifierRecord.CurrentData.MovementType == ERigidBodyMovementType::Kinematic);
+
+		// Note that the actual kinematic targets will be set separately, since they need to be set
+		// for all kinematics whether or not they were under a modifier.
+
+		// Scale gravity
+		float GravityMultiplier = BodyModifierRecord.CurrentData.GravityMultiplier;
+		if (GravityMultiplier != 1 && BodyModifierRecord.ActorHandle->IsGravityEnabled())
+		{
+			float Mass = (float) BodyModifierRecord.ActorHandle->GetMass();
+			FVector AntiGravityForce = SimSpaceGravity * (-Mass * (1.0f - GravityMultiplier));
+			BodyModifierRecord.ActorHandle->AddForce(AntiGravityForce);
+		}
+	}
+}
+
+//======================================================================================================================
+void FAnimNode_RigidBodyWithControl::ApplyControlAndModifierUpdatesAndParametersToRecords(
+	const FRigidBodyControlAndModifierUpdates&    Updates,
+	const FRigidBodyControlAndModifierParameters& Parameters)
+{
+	// Apply control and modifier parameters on a single-use basis
+	ApplyControlAndBodyModifierDatas(
+		Updates.ControlParameters, Updates.ModifierParameters);
+
+	// This goes through the records, resetting the update parts.
+	// Then the update structures get adjusted based on the parameters.
+	// The results don't get applied to the actual constraints yet - that happens in ApplyControlsAndModifiers
+	for (TMap<FName, FControlRecord>::ElementType& NameRecordPair : ControlRecords)
+	{
+		FName ControlName = NameRecordPair.Key;
+		if (const FRigidBodyControlTarget* ControlTarget = ControlTargets.Targets.Find(ControlName))
+		{
+			NameRecordPair.Value.ControlTarget = *ControlTarget;
+			NameRecordPair.Value.ResetCurrent(false);
+		}
+		else
+		{
+			NameRecordPair.Value.ResetCurrent(true);
+		}
+	}
+
+	for (TMap<FName, FBodyModifierRecord>::ElementType& NameRecordPair : ModifierRecords)
+	{
+		NameRecordPair.Value.ResetCurrent();
+	}
+
+	::ApplyControlAndModifierParametersToRecords(
+		ControlRecords, Parameters.ControlParameters, NameRecords.ControlSets);
+	::ApplyControlAndModifierParametersToRecords(
+		ModifierRecords, Parameters.ModifierParameters, NameRecords.BodyModifierSets);
+}
+
+//======================================================================================================================
+void FAnimNode_RigidBodyWithControl::ApplyControlAndBodyModifierDatas(
+	const TArray<FRigidBodyNamedControlParameters>& InControlParameters,
+	const TArray<FRigidBodyNamedModifierParameters>& InModifierParameters)
+{
+	// This updates the "original" controls and modifiers based on the parameters.
+	for (const FRigidBodyNamedControlParameters& ControlParameters : InControlParameters)
+	{
+		TArray<FName> Names = ExpandSetName(ControlParameters.Name, NameRecords.ControlSets);
+		for (FName Name : Names)
+		{
+			const FRigidBodyControlSparseData& ControlData = ControlParameters.Data;
+			if (FControlRecord* ControlRecord = ControlRecords.Find(Name))
+			{
+				ControlRecord->Control.ControlData.UpdateFromSparseData(ControlData);
+			}
+			else
+			{
+				UE_LOG(LogRigidBodyWithControl, Warning,
+					TEXT("ApplyControlAndBodyModifierDatas: Failed to find control with name %s"), *Name.ToString());
+			}
+		}
+	}
+
+	for (const FRigidBodyNamedModifierParameters& ModifierParameters : InModifierParameters)
+	{
+		TArray<FName> Names = ExpandSetName(ModifierParameters.Name, NameRecords.BodyModifierSets);
+		for (FName Name : Names)
+		{
+			const FRigidBodyModifierSparseData& ModifierData = ModifierParameters.Data;
+			if (FBodyModifierRecord* ModifierRecord = ModifierRecords.Find(Name))
+			{
+				ModifierRecord->Modifier.ModifierData.UpdateFromSparseData(ModifierData);
+			}
+			else
+			{
+				UE_LOG(LogRigidBodyWithControl, Warning,
+					TEXT("ApplyControlAndBodyModifierDatas: Failed to find modifier with name %s"), *Name.ToString());
+			}
+		}
+	}
+}
+
+//======================================================================================================================
+void FAnimNode_RigidBodyWithControl::ApplyControlsAndModifiers(const FVector& SimSpaceGravity, float DeltaTime)
+{
+	// If we've skipped a frame then we need to avoid doing any velocity calculations. Simplest
+	// method is to set DeltaTime to zero.
+	{
+		if (PoseData.UpdateCounter.Get() == INDEX_NONE || 
+			PoseData.UpdateCounter.Get() != PoseData.ExpectedUpdateCounter.Get())
+		{
+			DeltaTime = 0.0f;
+		}
+	}
+
+	// Update world space control root object transform.
+	if (WorldSpaceControlActorHandle)
+	{
+		// This should be ok as we have no collisions on our root body
+		WorldSpaceControlActorHandle->SetWorldTransform(PoseData.WorldCoordinateFrame); 
+	}
+
+	if (!PoseData.IsEmpty())
+	{
+		// Apply Controls.
+		for (TMap<FName, FControlRecord>::ElementType& NameRecordPair : ControlRecords)
+		{
+			FControlRecord& ControlRecord = NameRecordPair.Value;
+			if (ControlRecord.IsEnabled())
+			{
+				ApplyControl(ControlRecord, DeltaTime);
+			}
+
+			SetConstraintEnabled(ControlRecord.JointHandle, ControlRecord.IsEnabled());
+		}
+
+		// Apply Body Modifiers.
+		for (const TMap<FName, FBodyModifierRecord>::ElementType& NameRecordPair : ModifierRecords)
+		{
+			ApplyModifier(NameRecordPair.Value, SimSpaceGravity);
+		}
+	}
+}
+
+//======================================================================================================================
+// Note that this will be called AFTER normal kinematic targets have been set
+void FAnimNode_RigidBodyWithControl::ApplyKinematicTargets()
+{
+	if (!PoseData.IsEmpty())
+	{
+		// Apply any kinematic targets.
+		for (TMap<FName, FRigidBodyKinematicTarget>::ElementType& KinematicTargetPair : KinematicTargets.Targets)
+		{
+			FName BodyModifierName = KinematicTargetPair.Key;
+			FRigidBodyKinematicTarget& Target = KinematicTargetPair.Value;
+			FBodyModifierRecord* ModifierRecord = ModifierRecords.Find(BodyModifierName);
+			if (ModifierRecord && ModifierRecord->ActorHandle)
+			{
+				ImmediatePhysics::FActorHandle* ActorHandle = ModifierRecord->ActorHandle;
+				// TODO It might be worth storing this index in a way that doesn't need a lookup
+				const int32 BodyIndex = FindBodyIndexFromBoneName(ModifierRecord->Modifier.BoneName);
+				if (ActorHandle->GetIsKinematic() && BodyIndex != INDEX_NONE)
+				{
+					FTransform TM(Target.TargetOrientation, Target.TargetPosition);
+					if (Target.bUseSkeletalAnimation)
+					{
+						FTransform PoseTM = PoseData.GetTM(BodyIndex);
+						TM = TM * PoseTM;
+					}
+					ActorHandle->SetKinematicTarget(TM);
+				}
+			}
+		}
+	}
+}
+
+//======================================================================================================================
+static Chaos::EJointMotionType ConvertMotionType(ELinearConstraintMotion InEngineType)
+{
+	switch (InEngineType)
+	{
+	case ELinearConstraintMotion::LCM_Free: return Chaos::EJointMotionType::Free;
+	case ELinearConstraintMotion::LCM_Limited: return Chaos::EJointMotionType::Limited;
+	case ELinearConstraintMotion::LCM_Locked : return Chaos::EJointMotionType::Locked;
+	default: ensure(false); return Chaos::EJointMotionType::Locked;
+	}
+};
+
+//======================================================================================================================
+static Chaos::EJointMotionType ConvertMotionType(EAngularConstraintMotion InEngineType)
+{
+	switch(InEngineType)
+	{
+	case EAngularConstraintMotion::ACM_Free: return Chaos::EJointMotionType::Free;
+	case EAngularConstraintMotion::ACM_Limited: return Chaos::EJointMotionType::Limited;
+	case EAngularConstraintMotion::ACM_Locked: return Chaos::EJointMotionType::Locked;
+	default: ensure(false); return Chaos::EJointMotionType::Locked;
+	}
+};
+
+//======================================================================================================================
+static Chaos::EPlasticityType ConvertPlasticityType(EConstraintPlasticityType InType)
+{
+	switch (InType)
+	{
+	case EConstraintPlasticityType::CCPT_Free: return Chaos::EPlasticityType::Free;
+	case EConstraintPlasticityType::CCPT_Shrink: return Chaos::EPlasticityType::Shrink;
+	case EConstraintPlasticityType::CCPT_Grow: return Chaos::EPlasticityType::Grow;
+	default: ensure(false); return Chaos::EPlasticityType::Free;
+	}
+}
+
+//======================================================================================================================
+void FAnimNode_RigidBodyWithControl::ApplyCurrentConstraintProfile()
+{
+	using namespace Chaos;
+	// Go through each joint (in the ragdoll that's been created) in turn...
+	for (int32 JointIndex = 0; JointIndex != Joints.Num(); ++JointIndex)
+	{
+		if (ImmediatePhysics::FJointHandle* JointHandle = Joints[JointIndex])
+		{
+			ImmediatePhysics::FActorHandle* ParentActorHandle = JointHandle->GetActorHandles()[ConstraintParentIndex];
+			ImmediatePhysics::FActorHandle* ChildActorHandle = JointHandle->GetActorHandles()[ConstraintChildIndex];
+
+			// We need to associate this with the constraint setup to get the profile. At the moment
+			// we have to do a brute force search, because our joints will not necessarily be in the
+			// same order. TODO store the map from joint indices to constraint setups in the physics asset.
+			FName ParentActorName = ParentActorHandle ? ParentActorHandle->GetName() : FName();
+			FName ChildActorName = ChildActorHandle ? ChildActorHandle->GetName() : FName();
+
+			for (UPhysicsConstraintTemplate* ConstraintSetup : PhysicsAssetToUse->ConstraintSetup)
+			{
+				// All sorts of problems with comparing names
+				if (ConstraintSetup->DefaultInstance.GetParentBoneName() == ParentActorName &&
+					ConstraintSetup->DefaultInstance.GetChildBoneName() == ChildActorName)
+				{
+					FPBDJointConstraintHandle* ConstraintHandle = JointHandle->GetConstraint();
+
+					const FConstraintProfileProperties& Profile = 
+						ConstraintSetup->GetConstraintProfilePropertiesOrDefault(ConstraintProfile);
+
+					FPBDJointSettings JointSettings = ConstraintHandle->GetSettings();
+					
+					// See ImmediatePhysics_Chaos::TransferJointSettings
+					// TODO avoid code duplication with that function.
+
+					JointSettings.Stiffness = ConstraintSettings::JointStiffness();
+					JointSettings.LinearProjection = Profile.bEnableProjection ? Profile.ProjectionLinearAlpha : 0.0f;
+					JointSettings.AngularProjection = Profile.bEnableProjection ? Profile.ProjectionAngularAlpha : 0.0f;
+					JointSettings.ShockPropagation = Profile.bEnableShockPropagation ? Profile.ShockPropagationAlpha : 0.0f;
+					JointSettings.TeleportDistance = Profile.bEnableProjection ? Profile.ProjectionLinearTolerance : -1.0f;
+					JointSettings.TeleportAngle = Profile.bEnableProjection ? 
+						FMath::DegreesToRadians(Profile.ProjectionAngularTolerance) : -1.0f;
+					JointSettings.ParentInvMassScale = Profile.bParentDominates ? (FReal)0 : (FReal)1;
+
+					JointSettings.bCollisionEnabled = !Profile.bDisableCollision;
+					JointSettings.bProjectionEnabled = Profile.bEnableProjection;
+					JointSettings.bShockPropagationEnabled = Profile.bEnableShockPropagation;
+					JointSettings.bMassConditioningEnabled = Profile.bEnableMassConditioning;
+
+					JointSettings.LinearMotionTypes[0] = ConvertMotionType(Profile.LinearLimit.XMotion);
+					JointSettings.LinearMotionTypes[1] = ConvertMotionType(Profile.LinearLimit.YMotion);
+					JointSettings.LinearMotionTypes[2] = ConvertMotionType(Profile.LinearLimit.ZMotion);
+
+					JointSettings.LinearLimit = Profile.LinearLimit.Limit; // Is float to vector the best way?!
+
+					// Order is twist, swing1, swing2 and in degrees
+					JointSettings.AngularMotionTypes[0] = ConvertMotionType(Profile.TwistLimit.TwistMotion);
+					JointSettings.AngularMotionTypes[1] = ConvertMotionType(Profile.ConeLimit.Swing1Motion);
+					JointSettings.AngularMotionTypes[2] = ConvertMotionType(Profile.ConeLimit.Swing2Motion);
+
+					JointSettings.AngularLimits[0] = FMath::DegreesToRadians(Profile.TwistLimit.TwistLimitDegrees);
+					JointSettings.AngularLimits[1] = FMath::DegreesToRadians(Profile.ConeLimit.Swing1LimitDegrees);
+					JointSettings.AngularLimits[2] = FMath::DegreesToRadians(Profile.ConeLimit.Swing2LimitDegrees);
+
+					JointSettings.bSoftLinearLimitsEnabled = Profile.LinearLimit.bSoftConstraint;
+					JointSettings.bSoftTwistLimitsEnabled = Profile.TwistLimit.bSoftConstraint;
+					JointSettings.bSoftSwingLimitsEnabled = Profile.ConeLimit.bSoftConstraint;
+
+					JointSettings.LinearSoftForceMode = (ConstraintSettings::SoftLinearForceMode() == 0) ? 
+						EJointForceMode::Acceleration : EJointForceMode::Force;
+					JointSettings.AngularSoftForceMode = (ConstraintSettings::SoftAngularForceMode() == 0) ? 
+						EJointForceMode::Acceleration : EJointForceMode::Force;
+
+					JointSettings.SoftLinearStiffness = Profile.LinearLimit.Stiffness;
+					JointSettings.SoftLinearDamping = Profile.LinearLimit.Damping;
+					JointSettings.SoftTwistStiffness = Profile.TwistLimit.Stiffness;
+					JointSettings.SoftTwistDamping = Profile.TwistLimit.Damping;
+					JointSettings.SoftSwingStiffness = Profile.ConeLimit.Stiffness;
+					JointSettings.SoftSwingDamping = Profile.ConeLimit.Damping;
+
+					JointSettings.LinearRestitution = Profile.LinearLimit.Restitution;
+					JointSettings.TwistRestitution = Profile.TwistLimit.Restitution;
+					JointSettings.SwingRestitution = Profile.ConeLimit.Restitution;
+
+					JointSettings.LinearContactDistance = Profile.LinearLimit.ContactDistance;
+					JointSettings.TwistContactDistance = Profile.TwistLimit.ContactDistance;
+					JointSettings.SwingContactDistance = Profile.ConeLimit.ContactDistance;
+
+					JointSettings.LinearDrivePositionTarget = Profile.LinearDrive.PositionTarget;
+					JointSettings.LinearDriveVelocityTarget = Profile.LinearDrive.VelocityTarget;
+					JointSettings.bLinearPositionDriveEnabled[0] = Profile.LinearDrive.XDrive.bEnablePositionDrive;
+					JointSettings.bLinearPositionDriveEnabled[1] = Profile.LinearDrive.YDrive.bEnablePositionDrive;
+					JointSettings.bLinearPositionDriveEnabled[2] = Profile.LinearDrive.ZDrive.bEnablePositionDrive;
+					JointSettings.bLinearVelocityDriveEnabled[0] = Profile.LinearDrive.XDrive.bEnableVelocityDrive;
+					JointSettings.bLinearVelocityDriveEnabled[1] = Profile.LinearDrive.YDrive.bEnableVelocityDrive;
+					JointSettings.bLinearVelocityDriveEnabled[2] = Profile.LinearDrive.ZDrive.bEnableVelocityDrive;
+					JointSettings.LinearDriveForceMode = EJointForceMode::Acceleration; // hardcoded!
+					JointSettings.LinearDriveStiffness[0] = 
+						ConstraintSettings::LinearDriveStiffnessScale() * Profile.LinearDrive.XDrive.Stiffness;
+					JointSettings.LinearDriveStiffness[1] = 
+						ConstraintSettings::LinearDriveStiffnessScale() * Profile.LinearDrive.YDrive.Stiffness;
+					JointSettings.LinearDriveStiffness[2] = 
+						ConstraintSettings::LinearDriveStiffnessScale() * Profile.LinearDrive.ZDrive.Stiffness;
+					JointSettings.LinearDriveDamping[0] = 
+						ConstraintSettings::LinearDriveDampingScale() * Profile.LinearDrive.XDrive.Stiffness;
+					JointSettings.LinearDriveDamping[1] = 
+						ConstraintSettings::LinearDriveDampingScale() * Profile.LinearDrive.YDrive.Stiffness;
+					JointSettings.LinearDriveDamping[2] = 
+						ConstraintSettings::LinearDriveDampingScale() * Profile.LinearDrive.ZDrive.Stiffness;
+					JointSettings.LinearDriveMaxForce[0] = Profile.LinearDrive.XDrive.MaxForce;
+					JointSettings.LinearDriveMaxForce[1] = Profile.LinearDrive.YDrive.MaxForce;
+					JointSettings.LinearDriveMaxForce[2] = Profile.LinearDrive.ZDrive.MaxForce;
+
+					JointSettings.AngularDrivePositionTarget = FQuat(Profile.AngularDrive.OrientationTarget);
+					JointSettings.AngularDriveVelocityTarget = Profile.AngularDrive.AngularVelocityTarget * UE_TWO_PI; // rev/s to rad/s
+
+					JointSettings.AngularDriveForceMode = EJointForceMode::Acceleration; // hardcoded!
+					if (Profile.AngularDrive.AngularDriveMode == EAngularDriveMode::SLERP)
+					{
+						JointSettings.AngularDriveStiffness = FVec3(
+							ConstraintSettings::LinearDriveStiffnessScale() * Profile.AngularDrive.SlerpDrive.Stiffness);
+						JointSettings.AngularDriveDamping = FVec3(
+							ConstraintSettings::LinearDriveDampingScale() * Profile.AngularDrive.SlerpDrive.Damping);
+						JointSettings.AngularDriveMaxTorque = FVec3(
+							Profile.AngularDrive.SlerpDrive.MaxForce);
+						JointSettings.bAngularSLerpPositionDriveEnabled = Profile.AngularDrive.SlerpDrive.bEnablePositionDrive;
+						JointSettings.bAngularSLerpVelocityDriveEnabled = Profile.AngularDrive.SlerpDrive.bEnableVelocityDrive;
+						JointSettings.bAngularTwistPositionDriveEnabled = false;
+						JointSettings.bAngularTwistVelocityDriveEnabled = false;
+						JointSettings.bAngularSwingPositionDriveEnabled = false;
+						JointSettings.bAngularSwingVelocityDriveEnabled = false;
+					}
+					else
+					{
+						JointSettings.AngularDriveStiffness[0] = 
+							ConstraintSettings::LinearDriveStiffnessScale() * Profile.AngularDrive.TwistDrive.Stiffness;
+						JointSettings.AngularDriveStiffness[1] = 
+							ConstraintSettings::LinearDriveStiffnessScale() * Profile.AngularDrive.SwingDrive.Stiffness;
+						JointSettings.AngularDriveStiffness[2] = 
+							ConstraintSettings::LinearDriveStiffnessScale() * Profile.AngularDrive.SwingDrive.Stiffness;
+						JointSettings.AngularDriveDamping[0] = 
+							ConstraintSettings::LinearDriveDampingScale() * Profile.AngularDrive.TwistDrive.Damping;
+						JointSettings.AngularDriveDamping[1] = 
+							ConstraintSettings::LinearDriveDampingScale() * Profile.AngularDrive.SwingDrive.Damping;
+						JointSettings.AngularDriveDamping[2] = 
+							ConstraintSettings::LinearDriveDampingScale() * Profile.AngularDrive.SwingDrive.Damping;
+						JointSettings.AngularDriveMaxTorque[0] = Profile.AngularDrive.TwistDrive.MaxForce;
+						JointSettings.AngularDriveMaxTorque[1] = Profile.AngularDrive.SwingDrive.MaxForce;
+						JointSettings.AngularDriveMaxTorque[2] = Profile.AngularDrive.SwingDrive.MaxForce;
+						JointSettings.bAngularSLerpPositionDriveEnabled = false;
+						JointSettings.bAngularSLerpVelocityDriveEnabled = false;
+						JointSettings.bAngularTwistPositionDriveEnabled = Profile.AngularDrive.TwistDrive.bEnablePositionDrive;
+						JointSettings.bAngularTwistVelocityDriveEnabled = Profile.AngularDrive.TwistDrive.bEnableVelocityDrive;
+						JointSettings.bAngularSwingPositionDriveEnabled = Profile.AngularDrive.SwingDrive.bEnablePositionDrive;
+						JointSettings.bAngularSwingVelocityDriveEnabled = Profile.AngularDrive.SwingDrive.bEnableVelocityDrive;
+					}
+
+					JointSettings.LinearBreakForce = Profile.bLinearBreakable ? 
+						Chaos::ConstraintSettings::LinearBreakScale() * Profile.LinearBreakThreshold : FLT_MAX;
+					JointSettings.LinearPlasticityLimit = Profile.bLinearPlasticity ? 
+						FMath::Clamp((float)Profile.LinearPlasticityThreshold, 0.0f, 1.0f) : FLT_MAX;
+					JointSettings.LinearPlasticityType = ConvertPlasticityType(Profile.LinearPlasticityType);
+					// JointSettings.LinearPlasticityInitialDistanceSquared = ; // What do we do with this?
+
+					JointSettings.AngularBreakTorque = Profile.bAngularBreakable ? 
+						Chaos::ConstraintSettings::AngularBreakScale() * Profile.AngularBreakThreshold : FLT_MAX;
+					JointSettings.AngularPlasticityLimit = Profile.bAngularPlasticity ? 
+						FMath::Clamp((float)Profile.AngularPlasticityThreshold, 0.0f, 1.0f) : FLT_MAX;;
+
+					JointSettings.ContactTransferScale = Profile.ContactTransferScale;
+
+					ConstraintHandle->SetSettings(JointSettings);
+				}
+			}
+		}
+	}
+}
+
diff --git a/Engine/Plugins/Experimental/PhysicsControl/Source/PhysicsControl/Public/AnimNode_RigidBodyWithControl.h b/Engine/Plugins/Experimental/PhysicsControl/Source/PhysicsControl/Public/AnimNode_RigidBodyWithControl.h
new file mode 100644
--- /dev/null
+++ b/Engine/Plugins/Experimental/PhysicsControl/Source/PhysicsControl/Public/AnimNode_RigidBodyWithControl.h
@@ -0,0 +1,579 @@
+// Copyright Epic Games, Inc. All Rights Reserved.
+
+/* 
+ * This is essentially a copy + paste of the RigidBody file
+ */
+
+#pragma once
+
+#include "CoreMinimal.h"
+#include "RigidBodyPoseData.h"
+#include "RigidBodyControlData.h"
+#include "RigidBodyNameRecords.h"
+#include "BoneControllers/AnimNode_SkeletalControlBase.h"
+#include "Components/SkeletalMeshComponent.h"
+#include "Physics/ImmediatePhysics/ImmediatePhysicsDeclares.h"
+#include "PhysicsEngine/PhysicsAsset.h"
+#include "PhysicsProxy/PerSolverFieldSystem.h"
+#include "Tasks/Task.h"
+
+ // Use the simulation space functions from the RBAN
+#include "BoneControllers/AnimNode_RigidBody.h" 
+
+#include "AnimNode_RigidBodyWithControl.generated.h"
+
+struct FBodyInstance;
+struct FConstraintInstance;
+class FEvent;
+
+extern TAutoConsoleVariable<int32> CVarEnableRigidBodyNodeWithControl;
+extern TAutoConsoleVariable<int32> CVarEnableRigidBodyNodeWithControlSimulation;
+extern TAutoConsoleVariable<int32> CVarRigidBodyNodeWithControlLODThreshold;
+
+/**
+ * Controller that simulates physics based on the physics asset of the skeletal mesh component
+ */
+USTRUCT()
+struct PHYSICSCONTROL_API FAnimNode_RigidBodyWithControl : public FAnimNode_SkeletalControlBase
+{
+	GENERATED_BODY()
+
+	FAnimNode_RigidBodyWithControl();
+	~FAnimNode_RigidBodyWithControl();
+
+	// FAnimNode_Base interface
+	virtual void GatherDebugData(FNodeDebugData& DebugData) override;
+	virtual void Initialize_AnyThread(const FAnimationInitializeContext& Context) override;
+	// End of FAnimNode_Base interface
+
+	// FAnimNode_SkeletalControlBase interface
+	virtual void UpdateComponentPose_AnyThread(const FAnimationUpdateContext& Context) override;
+	virtual void EvaluateComponentPose_AnyThread(FComponentSpacePoseContext& Output) override;
+	virtual void EvaluateSkeletalControl_AnyThread(FComponentSpacePoseContext& Output, TArray<FBoneTransform>& OutBoneTransforms) override;
+	virtual void OnInitializeAnimInstance(const FAnimInstanceProxy* InProxy, const UAnimInstance* InAnimInstance) override;
+	virtual bool NeedsOnInitializeAnimInstance() const override { return true; }
+	virtual void PreUpdate(const UAnimInstance* InAnimInstance) override;
+	virtual void UpdateInternal(const FAnimationUpdateContext& Context) override;
+	virtual bool HasPreUpdate() const override { return true; }
+	virtual bool IsValidToEvaluate(const USkeleton* Skeleton, const FBoneContainer& RequiredBones) override;
+	virtual bool NeedsDynamicReset() const override;
+	virtual void ResetDynamics(ETeleportType InTeleportType) override;
+	virtual int32 GetLODThreshold() const override;
+	// End of FAnimNode_SkeletalControlBase interface
+
+	virtual void AddImpulseAtLocation(FVector Impulse, FVector Location, FName BoneName = NAME_None);
+
+	// TEMP: Exposed for use in PhAt as a quick way to get drag handles working with Chaos
+	virtual ImmediatePhysics::FSimulation* GetSimulation() { return PhysicsSimulation; }
+
+	UPhysicsAsset* GetPhysicsAsset() const { return PhysicsAssetToUse; }
+
+public:
+	/** Physics asset to use. If empty use the skeletal mesh's default physics asset */
+	UPROPERTY(EditAnywhere, Category = Settings)
+	TObjectPtr<UPhysicsAsset> OverridePhysicsAsset;
+
+private:
+	FTransform PreviousCompWorldSpaceTM;
+	FTransform CurrentTransform;
+	FTransform PreviousTransform;
+
+	UPhysicsAsset* PhysicsAssetToUse;
+public:
+
+	/** Override gravity*/
+	UPROPERTY(EditAnywhere, Category = Settings, meta = (PinHiddenByDefault, editcondition = "bOverrideWorldGravity"))
+	FVector OverrideWorldGravity;
+
+	/** Applies a uniform external force in world space. This allows for easily faking inertia of movement while still simulating in component space for example */
+	UPROPERTY(EditAnywhere, Category = Settings, meta = (PinShownByDefault))
+	FVector ExternalForce;
+
+	/** When using non-world-space sim, this controls how much of the components world-space acceleration is passed on to the local-space simulation. */
+	UPROPERTY(EditAnywhere, Category = Settings, meta = (PinHiddenByDefault))
+	FVector ComponentLinearAccScale;
+
+	/** When using non-world-space sim, this applies a 'drag' to the bodies in the local space simulation, based on the components world-space velocity. */
+	UPROPERTY(EditAnywhere, Category = Settings, meta = (PinHiddenByDefault))
+	FVector ComponentLinearVelScale;
+
+	/** When using non-world-space sim, this is an overall clamp on acceleration derived from ComponentLinearAccScale and ComponentLinearVelScale, to ensure it is not too large. */
+	UPROPERTY(EditAnywhere, Category = Settings)
+	FVector	ComponentAppliedLinearAccClamp;
+
+	/**
+	 * Settings for the system which passes motion of the simulation's space
+	 * into the simulation. This allows the simulation to pass a
+	 * fraction of the world space motion onto the bodies which allows Bone-Space
+	 * and Component-Space simulations to react to world-space movement in a
+	 * controllable way.
+	 * This system is a superset of the functionality provided by ComponentLinearAccScale,
+	 * ComponentLinearVelScale, and ComponentAppliedLinearAccClamp. In general
+	 * you should not have both systems enabled.
+	 */
+	UPROPERTY(EditAnywhere, Category = Settings, meta = (PinHiddenByDefault))
+	FSimSpaceSettings SimSpaceSettings;
+
+
+	/**
+	 * Scale of cached bounds (vs. actual bounds).
+	 * Increasing this may improve performance, but overlaps may not work as well.
+	 * (A value of 1.0 effectively disables cached bounds).
+	 */
+	UPROPERTY(EditAnywhere, Category = Settings, meta = (ClampMin="1.0", ClampMax="2.0"))
+	float CachedBoundsScale;
+
+	/** Matters if SimulationSpace is BaseBone */
+	UPROPERTY(EditAnywhere, Category = Settings)
+	FBoneReference BaseBoneRef;
+
+	/** The channel we use to find static geometry to collide with */
+	UPROPERTY(EditAnywhere, Category = Settings, meta = (editcondition = "bEnableWorldGeometry"))
+	TEnumAsByte<ECollisionChannel> OverlapChannel;
+
+	/** What space to simulate the bodies in. This affects how velocities are generated */
+	UPROPERTY(EditAnywhere, Category = Settings)
+	ESimulationSpace SimulationSpace;
+
+	/** Whether to allow collisions between two bodies joined by a constraint  */
+	UPROPERTY(EditAnywhere, Category = Settings)
+	bool bForceDisableCollisionBetweenConstraintBodies;
+
+	/** If true, kinematic objects will be added to the simulation at runtime to represent any cloth colliders defined for the parent object. */
+	UPROPERTY(EditAnywhere, Category = Settings)
+	bool bUseExternalClothCollision;
+
+private:
+	ETeleportType ResetSimulatedTeleportType;
+
+public:
+	UPROPERTY(EditAnywhere, Category = Settings, meta = (InlineEditConditionToggle))
+	uint8 bEnableWorldGeometry : 1;
+
+	UPROPERTY(EditAnywhere, Category = Settings, meta = (InlineEditConditionToggle))
+	uint8 bOverrideWorldGravity : 1;
+
+	/** 
+		When simulation starts, transfer previous bone velocities (from animation)
+		to make transition into simulation seamless.
+	*/
+	UPROPERTY(EditAnywhere, Category = Settings, meta=(PinHiddenByDefault))
+	uint8 bTransferBoneVelocities : 1;
+
+	/**
+		When simulation starts, freeze incoming pose.
+		This is useful for ragdolls, when we want the simulation to take over.
+		It prevents non simulated bones from animating.
+	*/
+	UPROPERTY(EditAnywhere, Category = Settings)
+	uint8 bFreezeIncomingPoseOnStart : 1;
+
+	/**
+		Correct for linear tearing on bodies with all axes Locked.
+		This only works if all axes linear translation are locked
+	*/
+	UPROPERTY(EditAnywhere, Category = Settings)
+	uint8 bClampLinearTranslationLimitToRefPose : 1;
+
+	/**
+		For world-space simulations, if the magnitude of the component's 3D scale is less than WorldSpaceMinimumScale, do not update the node.
+	*/
+	UPROPERTY(EditAnywhere, Category = Settings)
+	float WorldSpaceMinimumScale;
+
+	/**
+		If the node is not evaluated for this amount of time (seconds), either because a lower LOD was in use for a while or the component was
+		not visible, reset the simulation to the default pose on the next evaluation. Set to 0 to disable time-based reset.
+	*/
+	UPROPERTY(EditAnywhere, Category = Settings)
+	float EvaluationResetTime;
+
+	/**
+	 * If false, then controls will not be created. Note that this can be exposed as a pin/bound, and then control 
+	 * creation can be deferred to when the value is set to true as the node runs.
+	 */
+	UPROPERTY(EditAnywhere, Category = ControlSetup, meta = (PinHiddenByDefault))
+	bool bEnableControls;
+
+	/**
+	 * A map of bone names to "body" names, the latter being used to assign names to controls/modifiers. 
+	 * This is optional - so if there is no mapping for a bone, then its name will be used directly when 
+	 * creating controls. The two main benefits of this are (1) to generate consistently named controls 
+	 * even on different skeletons and (2) to make it easier to refer to individual controls, without needing 
+	 * to refer to the skeleton.
+	 */
+	UPROPERTY(EditAnywhere, Category = ControlSetup, meta = (PinHiddenByDefault))
+	TMap<FName, FName> BoneToBodyNameMap;
+
+	/**
+	 * Setup data for creating the main controls (world- and parent-space) and modifiers, based on splitting the 
+	 * skeleton up into limbs.
+	 */
+	UPROPERTY(EditAnywhere, Category = ControlSetup, meta = (PinHiddenByDefault))
+	FRigidBodySetupData SetupData;
+
+	/**
+	 * Controls and modifiers that should be created, in addition to those made as part of the limb setup.
+	 */
+	UPROPERTY(EditAnywhere, Category = ControlSetup, meta = (PinHiddenByDefault))
+	FRigidBodyControlAndBodyModifierCreations AdditionalControlsAndBodyModifiers;
+
+	/**
+	 * Allows additional sets of controls or modifiers to be created, and existing sets to be modified
+	 */
+	UPROPERTY(EditAnywhere, Category = ControlSetup, meta = (PinHiddenByDefault))
+	FRigidBodySetUpdates AdditionalSets;
+
+	/**
+	 * An initial set of controls that should be applied immediately after setup. This allows individual or 
+	 * sets of controls/modifiers etc to be adjusted. Note that these will then be "baked" into the 
+	 * controls.
+	 */
+	UPROPERTY(EditAnywhere, Category = ControlSetup, meta = (PinHiddenByDefault))
+	FRigidBodyControlAndModifierParameters InitialControlAndBodyModifierUpdates;
+
+	/**
+	 * Controls that should be applied each frame, and can be expected to change. Note that if these
+	 * stop being passed in then the controls and modifiers will return to their normal/original state.
+	 */
+	UPROPERTY(EditAnywhere, Category = Controls, meta = (PinShownByDefault))
+	FRigidBodyControlAndModifierParameters ControlAndModifierParameters;
+
+	/**
+	 * Updates to controls that can be applied. Note that these update the normal/original state.
+	 */
+	UPROPERTY(EditAnywhere, Category = Controls, meta = (PinShownByDefault))
+	FRigidBodyControlAndModifierUpdates ControlAndModifierUpdates;
+
+	/**
+	 * Targets that should be applied to the controls. 
+	 */
+	UPROPERTY(EditAnywhere, Category = Controls, meta = (PinShownByDefault))
+	FRigidBodyControlTargets ControlTargets;
+
+	/**
+	 * Targets that should be applied to kinematic bodies that are under the influence of a body modifier. 
+	 */
+	UPROPERTY(EditAnywhere, Category = Controls, meta = (PinShownByDefault))
+	FRigidBodyKinematicTargets KinematicTargets;
+
+	/**
+	 * The constraint profile to use on all the joints in the physics asset
+	 */
+	UPROPERTY(EditAnywhere, Category = Controls, meta = (PinHiddenByDefault))
+	FName ConstraintProfile;
+
+private:
+	uint8 bEnabled : 1;
+	uint8 bSimulationStarted : 1;
+	uint8 bCheckForBodyTransformInit : 1;
+	uint8 bHaveRunControlSeup : 1;
+
+public:
+	void PostSerialize(const FArchive& Ar);
+
+	const int32 GetNumBodies() const;
+
+	const FRigidBodyNameRecords& GetNameRecords() const { return NameRecords; }
+
+private:
+
+	// FAnimNode_SkeletalControlBase interface
+	virtual void InitializeBoneReferences(const FBoneContainer& RequiredBones) override;
+	// End of FAnimNode_SkeletalControlBase interface
+
+	void InitPhysics(const UAnimInstance* InAnimInstance);
+	void UpdateWorldGeometry(const UWorld& World, const USkeletalMeshComponent& SKC);
+	void UpdateWorldForces(const FTransform& ComponentToWorld, const FTransform& RootBoneTM, const float DeltaSeconds);
+
+public:
+	const FTransform GetBodyTransform(const int32 BodyIndex) const;
+	const FTransform GetWorldSpaceControlRootTransform() const;
+
+private:
+
+	// Add a new physics body and update the body name to index map
+	int32 AddBody(ImmediatePhysics::FActorHandle* const BodyHandle); 
+	int32 FindBodyIndexFromBoneName(const FName BoneName) const;
+	FName FindParentBoneNameFromBoneName(const FName BoneName) const;
+	ImmediatePhysics::FActorHandle* FindBodyFromBoneName(const FName BoneName) const;
+
+	ImmediatePhysics::FJointHandle* CreateConstraint(
+		ImmediatePhysics::FActorHandle* const ChildBodyHandle, ImmediatePhysics::FActorHandle* const ParentBodyHandle);
+	void CreateWorldSpaceControlRootBody(UPhysicsAsset* const PhysicsAsset);
+
+	void ApplyControl(FControlRecord& ControlRecord, float DeltaTime);
+	void ApplyModifier(const FBodyModifierRecord& BodyModifierRecord, const FVector& SimSpaceGravity);
+
+	void ApplyControlsAndModifiers(const FVector& SimSpaceGravity, float DeltaTime);
+
+	// This applies kinematic targets to any that have been set through the body modifiers
+	void ApplyKinematicTargets();
+
+	// This applies the desired constraint profile, if necessary
+	void ApplyCurrentConstraintProfile();
+
+	// Adjusts the spring drive settings to reflect the control data
+	void UpdateDriveSpringDamperSettings(Chaos::FPBDJointSettings& Settings, const FRigidBodyControlData& ControlData);
+
+	TMap<FName, int32> BodyNameToIndexMap;
+	ImmediatePhysics::FActorHandle* WorldSpaceControlActorHandle;
+
+	FName GetBodyFromBoneName(const FName BoneName) const;
+
+	FName GetUniqueBodyModifierName(const FName BoneName);
+
+	FName GetUniqueControlName(const FName ParentBoneName, const FName ChildBoneName);
+
+	FName CreateControl(const FName ParentBoneName, const FName ChildBoneName, const FRigidBodyControlData& ControlData);
+
+	FName CreateBodyModifier(const FName BoneName, const FRigidBodyModifierData& ModifierData);
+
+	void CreateControlsFromLimbBones(
+		const FName                  LimbName,
+		const FRigidBodyLimbBones&   LimbBones,
+		const ERigidBodyControlType  ControlType,
+		const FRigidBodyControlData& ControlData);
+
+	void CreateBodyModifiersFromLimbBones(
+		const FName                   LimbName,
+		const FRigidBodyLimbBones&    LimbBones,
+		const FRigidBodyModifierData& DefaultModifierData);
+
+	TMap<FName, FRigidBodyLimbBones> GetLimbBonesFromSkeletalMesh(
+		const TArray<FRigidBodyLimbSetupData>& LimbSetupData,
+		USkeletalMeshComponent* const SkeletalMeshComponent) const;
+
+	// This will walk through the skeleton, create controls and body modifiers, and create the sets.
+	void InitControlsAndBodyModifiers(USkeletalMeshComponent* const SkeletalMeshComponent);
+
+	void DestroyControlsAndBodyModifiers();
+
+	// Helper to log overything that has been created to help the user
+	void LogControlsModifiersAndSets();
+
+	// Creates any controls as set in AdditionalControls
+	void CreateAdditionalControls();
+	
+	// Creates/updates sets - possibly making new ones, or adding to existing ones
+	void CreateAdditionalSets();
+
+	// Applies the overrides to the underlying controls and modifiers
+	void ApplyControlAndBodyModifierDatas(
+		const TArray<FRigidBodyNamedControlParameters>& ControlParameters,
+		const TArray<FRigidBodyNamedModifierParameters>& ModifierParameters);
+
+	// Applies the parameters to control and modifier records
+	void ApplyControlAndModifierUpdatesAndParametersToRecords(
+		const FRigidBodyControlAndModifierUpdates&    Updates,
+		const FRigidBodyControlAndModifierParameters& Parameters);
+
+	void InitializeNewBodyTransformsDuringSimulation(
+		FComponentSpacePoseContext& Output, const FTransform& ComponentTransform, const FTransform& BaseBoneTM);
+
+	void InitSimulationSpace(const FTransform& ComponentToWorld, const FTransform& BoneToComponent);
+
+	// Calculate simulation space transform, velocity etc to pass into the solver
+	void CalculateSimulationSpace(
+		ESimulationSpace Space,
+		const FTransform& ComponentToWorld,
+		const FTransform& BoneToComponent,
+		const float Dt,
+		const FSimSpaceSettings& Settings,
+		FTransform& SpaceTransform,
+		FVector& SpaceLinearVel,
+		FVector& SpaceAngularVel,
+		FVector& SpaceLinearAcc,
+		FVector& SpaceAngularAcc);
+
+	// Gather cloth collision sources from the supplied Skeltal Mesh and add a kinematic actor representing each one of them to the sim.
+	void CollectClothColliderObjects(const USkeletalMeshComponent* SkeletalMeshComp);
+	
+	// Remove all cloth collider objects from the sim.
+	void RemoveClothColliderObjects();
+
+	// Update the sim-space transforms of all cloth collider objects.
+	void UpdateClothColliderObjects(const FTransform& SpaceTransform);
+
+	// Gather nearby world objects and add them to the sim
+	void CollectWorldObjects();
+
+	// Flag invalid world objects to be removed from the sim
+	void ExpireWorldObjects();
+
+	// Remove simulation objects that are flagged as expired
+	void PurgeExpiredWorldObjects();
+
+	// Update sim-space transforms of world objects
+	void UpdateWorldObjects(const FTransform& SpaceTransform);
+
+	// Advances the simulation by a given timestep
+	void RunPhysicsSimulation(float DeltaSeconds, const FVector& SimSpaceGravity);
+
+	// Waits for the deferred simulation task to complete if it's not already finished
+	void FlushDeferredSimulationTask();
+
+	// Destroy the simulation and free related structures
+	void DestroyPhysicsSimulation();
+
+public:
+
+	/* Whether the physics simulation runs synchronously with the node's evaluation or is run in the background until the next frame. */
+	UPROPERTY(EditAnywhere, Category=Settings, AdvancedDisplay)
+	ESimulationTiming SimulationTiming;
+
+private:
+
+	float WorldTimeSeconds;
+	float LastEvalTimeSeconds;
+
+	float AccumulatedDeltaTime;
+	float AnimPhysicsMinDeltaTime;
+	bool bSimulateAnimPhysicsAfterReset;
+	/** This should only be used for removing the delegate during termination. Do NOT use this for any per frame work */
+	TWeakObjectPtr<USkeletalMeshComponent> SkelMeshCompWeakPtr;
+
+	ImmediatePhysics::FSimulation* PhysicsSimulation;
+	FPhysicsAssetSolverSettings SolverSettings;
+	FSolverIterations SolverIterations;	// to be deprecated
+
+	friend class FRigidBodyNodeWithControlSimulationTask;
+	UE::Tasks::FTask SimulationTask;
+
+	struct FBodyAnimData
+	{
+		FBodyAnimData()
+			: TransferedBoneAngularVelocity(ForceInit)
+			, TransferedBoneLinearVelocity(ForceInitToZero)
+			, LinearXMotion(ELinearConstraintMotion::LCM_Locked)
+			, LinearYMotion(ELinearConstraintMotion::LCM_Locked)
+			, LinearZMotion(ELinearConstraintMotion::LCM_Locked)
+			, LinearLimit(0.0f)
+			, RefPoseLength (0.f)
+			, bIsSimulated(false)
+			, bBodyTransformInitialized(false)
+		{}
+
+		FQuat TransferedBoneAngularVelocity;
+		FVector TransferedBoneLinearVelocity;
+
+		ELinearConstraintMotion LinearXMotion;
+		ELinearConstraintMotion LinearYMotion;
+		ELinearConstraintMotion LinearZMotion;
+		float LinearLimit;
+		// we don't use linear limit but use default length to limit the bodies
+		// linear limits are defined per constraint - it can be any two joints that can limit
+		// this is just default length of the local space from parent, and we use that info to limit
+		// the translation
+		float RefPoseLength;
+
+		bool bIsSimulated : 1;
+		bool bBodyTransformInitialized : 1;
+	};
+
+	struct FWorldObject
+	{
+		FWorldObject() : ActorHandle(nullptr), LastSeenTick(0), bExpired(false) {}
+		FWorldObject(ImmediatePhysics::FActorHandle* InActorHandle, int32 InLastSeenTick) : ActorHandle(InActorHandle), LastSeenTick(InLastSeenTick), bExpired(false) {}
+
+		ImmediatePhysics::FActorHandle* ActorHandle;
+		int32 LastSeenTick;
+		bool bExpired;
+	};
+
+	TArray<RigidBodyWithControl::FOutputBoneData> OutputBoneData;
+	// Note that the Bodies and Joints arrays will be the same size - i.e. there will be a
+	// correspondence between every joint and body. The joint will correspond to the parent of the
+	// body. Note that some joints will not be set.
+	TArray<ImmediatePhysics::FActorHandle*> Bodies;
+	TArray<ImmediatePhysics::FJointHandle*> Joints;
+	TArray<int32> SkeletonBoneIndexToBodyIndex;
+	TArray<FBodyAnimData> BodyAnimData;
+
+	friend class UAnimGraphNode_RigidBodyWithControl;
+
+	// Each update we cache the incoming pose transforms in whatever space the simulation is running in
+	RigidBodyWithControl::FRigidBodyPoseData PoseData;
+
+	// Map of control records - they will be referenced by name
+	TMap<FName, FControlRecord> ControlRecords;
+
+	// Map of body modifier records - they will be referenced by name
+	TMap<FName, FBodyModifierRecord> ModifierRecords;
+
+	// Details about sets etc
+	FRigidBodyNameRecords NameRecords;
+
+	FName CurrentConstraintProfile;
+
+	TArray<USkeletalMeshComponent::FPendingRadialForces> PendingRadialForces;
+
+	FPerSolverFieldSystem PerSolverField;
+
+	// Information required to identify and update a kinematic object representing a cloth collision source in the sim.
+	struct FClothCollider
+	{
+		FClothCollider(ImmediatePhysics::FActorHandle* const InActorHandle, const USkeletalMeshComponent* const InSkeletalMeshComponent, const uint32 InBoneIndex)
+			: ActorHandle(InActorHandle)
+			, SkeletalMeshComponent(InSkeletalMeshComponent)
+			, BoneIndex(InBoneIndex)
+		{}
+
+		ImmediatePhysics::FActorHandle* ActorHandle; // Identifies the physics actor in the sim.
+		const USkeletalMeshComponent* SkeletalMeshComponent; // Parent skeleton.
+		uint32 BoneIndex; // Bone within parent skeleton that drives physics actors transform.
+	};
+
+	// List of actors in the sim that represent objects collected from other parts of this character.
+	TArray<FClothCollider> ClothColliders; 
+	
+	TMap<const UPrimitiveComponent*, FWorldObject> ComponentsInSim;
+	int32 ComponentsInSimTick;
+
+	FVector WorldSpaceGravity;
+
+	float TotalMass;
+
+	// Bounds used to gather world objects copied into the simulation
+	FSphere CachedBounds;
+
+	FCollisionQueryParams QueryParams;
+
+	FPhysScene* PhysScene;
+
+	// Used by CollectWorldObjects and UpdateWorldGeometry in Task Thread
+	// Typically, World should never be accessed off the Game Thread.
+	// However, since we're just doing overlaps this should be OK.
+	const UWorld* UnsafeWorld;
+
+	// Used by CollectWorldObjects and UpdateWorldGeometry in Task Thread
+	// Only used for a pointer comparison.
+	const AActor* UnsafeOwner;
+
+	FBoneContainer CapturedBoneVelocityBoneContainer;
+	FCSPose<FCompactHeapPose> CapturedBoneVelocityPose;
+	FCSPose<FCompactHeapPose> CapturedFrozenPose;
+	FBlendedHeapCurve CapturedFrozenCurves;
+
+	FVector PreviousComponentLinearVelocity;
+
+	// Used by the world-space to simulation-space motion transfer system in Component- or Bone-Space sims
+	FTransform SimSpacePreviousComponentToWorld;
+	FTransform SimSpacePreviousBoneToComponent;
+	FVector SimSpacePreviousComponentLinearVelocity;
+	FVector SimSpacePreviousComponentAngularVelocity;
+	FVector SimSpacePreviousBoneLinearVelocity;
+	FVector SimSpacePreviousBoneAngularVelocity;
+};
+
+#if WITH_EDITORONLY_DATA
+template<>
+struct TStructOpsTypeTraits<FAnimNode_RigidBodyWithControl> : public TStructOpsTypeTraitsBase2<FAnimNode_RigidBodyWithControl>
+{
+	enum
+	{
+		WithPostSerialize = true,
+	};
+};
+#endif
+
diff --git a/Engine/Plugins/Experimental/PhysicsControl/Source/PhysicsControl/Public/PhysicsControlBPLibrary.h b/Engine/Plugins/Experimental/PhysicsControl/Source/PhysicsControl/Public/PhysicsControlBPLibrary.h
new file mode 100644
--- /dev/null
+++ b/Engine/Plugins/Experimental/PhysicsControl/Source/PhysicsControl/Public/PhysicsControlBPLibrary.h
@@ -0,0 +1,119 @@
+// Copyright Epic Games, Inc. All Rights Reserved.
+
+#pragma once
+
+#if 1
+
+#include "Kismet/BlueprintFunctionLibrary.h"
+#include "Animation/AnimNodeReference.h"
+#include "SkeletalControlLibrary.h"
+
+#include "PhysicsControlBPLibrary.generated.h"
+
+struct FRigidBodyControlParameters;
+struct FAnimNode_RigidBodyWithControl;
+
+USTRUCT(BlueprintType)
+struct FRigidBodyWithControlReference : public FAnimNodeReference
+{
+	GENERATED_BODY()
+
+	typedef FAnimNode_RigidBodyWithControl FInternalNodeType;
+};
+
+
+UCLASS()
+class PHYSICSCONTROL_API UPhysicsControlBPLibrary : public UBlueprintFunctionLibrary
+{
+	GENERATED_BODY()
+
+	/** Add a single control parameter Parameters to a set of parameter Parameters. */
+	UFUNCTION(BlueprintPure, Category = "Animation|PhysicsControl", meta = (BlueprintThreadSafe, AutoCreateRefTerm = "LinearStrengthMultiplier, LinearExtraDampingMultiplier, AngularStrengthMultiplier, AngularExtraDampingMultiplier, bEnabled"))
+	static void AddControlParameters(
+		UPARAM(ref, DisplayName = "Parameters") FRigidBodyControlAndModifierParameters& InParameters, 
+		UPARAM(DisplayName = "Parameters") FRigidBodyControlAndModifierParameters& OutParameters,
+		const FName                        Name,
+		const FRigidBodyControlSparseData& ControlData);
+
+	/** Add a single control parameter Parameters to a set of parameter Parameters. */
+	UFUNCTION(BlueprintPure, Category = "Animation|PhysicsControl", meta = (BlueprintThreadSafe, AutoCreateRefTerm = "LinearStrengthMultiplier, LinearExtraDampingMultiplier, AngularStrengthMultiplier, AngularExtraDampingMultiplier, bEnabled"))
+	static void AddMultipleControlParameters(
+		UPARAM(ref, DisplayName = "Parameters") FRigidBodyControlAndModifierParameters& InParameters, 
+		UPARAM(DisplayName = "Parameters") FRigidBodyControlAndModifierParameters& OutParameters,
+		const TArray<FName>&               Names,
+		const FRigidBodyControlSparseData& ControlData);
+
+	/** Add a single body modifier parameter Parameters to a set of parameter Parameters. */
+	UFUNCTION(BlueprintPure, Category = "Animation|PhysicsControl", meta = (BlueprintThreadSafe))
+	static void AddModifierParameters(
+		UPARAM(ref, DisplayName = "Parameters") FRigidBodyControlAndModifierParameters& InParameters,
+		UPARAM(DisplayName = "Parameters") FRigidBodyControlAndModifierParameters& OutParameters,
+		const FName                         Name, 
+		const FRigidBodyModifierSparseData& ModifierData);
+
+	/** Add a single body modifier parameter Parameters to a set of parameter Parameters. */
+	UFUNCTION(BlueprintPure, Category = "Animation|PhysicsControl", meta = (BlueprintThreadSafe))
+	static void AddMultipleModifierParameters(
+		UPARAM(ref, DisplayName = "Parameters") FRigidBodyControlAndModifierParameters& InParameters, 
+		UPARAM(DisplayName = "Parameters") FRigidBodyControlAndModifierParameters& OutParameters,
+		const TArray<FName>&                Names,
+		const FRigidBodyModifierSparseData& ModifierData);
+
+	/** Returns the linear interpolation of two sets of parameter Parameters. Any Parameters that exists in one of the input sets but not the other will be added to the output with a weight of 1. */
+	UFUNCTION(BlueprintPure, Category = "Animation|PhysicsControl", meta = (BlueprintThreadSafe))
+	static void BlendParameters(
+		UPARAM(ref, DisplayName = "From Parameters") FRigidBodyControlAndModifierParameters& InParametersA,
+		UPARAM(ref, DisplayName = "To Parameters") FRigidBodyControlAndModifierParameters& InParametersB,
+		UPARAM(DisplayName = "Weight") const float InInterpolationWeight,
+		UPARAM(DisplayName = "Parameters") FRigidBodyControlAndModifierParameters& OutParameters);
+	
+	/** Adds an Parameters to the output parameters for each control name in the supplied set. The values in each Parameters will be a linear interpolation of the two supplied Parameters, blending from the start Parameters to the end Parameters across the elements in the set. */
+	UFUNCTION(BlueprintPure, Category = "Animation|PhysicsControl", meta = (BlueprintThreadSafe))
+	static void BlendControlParametersThroughSet(
+		UPARAM(ref, DisplayName = "Parameters") FRigidBodyControlAndModifierParameters& InParameters,
+		UPARAM(ref, DisplayName = "Start Parameters") const FRigidBodyNamedControlParameters& InStartControlParameters,
+		UPARAM(ref, DisplayName = "End Parameters") const FRigidBodyNamedControlParameters& InEndControlParameters,
+		const TArray<FName>& ControlNames,
+		FRigidBodyControlAndModifierParameters& OutParameters);
+
+	/** Adds an Parameters to the output parameters for each modifier name in the supplied set. The values in each Parameters will be a linear interpolation of the two supplied Parameters, blending from the start Parameters to the end Parameters across the elements in the set. */
+	UFUNCTION(BlueprintPure, Category = "Animation|PhysicsControl", meta = (BlueprintThreadSafe))
+	static void BlendModifierParametersThroughSet(
+		UPARAM(ref) FRigidBodyControlAndModifierParameters& InParameters,
+		UPARAM(ref) const FRigidBodyNamedModifierParameters& InStartModifierParameters,
+		UPARAM(ref) const FRigidBodyNamedModifierParameters& InEndModifierParameters,
+		const TArray<FName>& ModifierNames,
+		UPARAM(DisplayName = "Parameters") FRigidBodyControlAndModifierParameters& OutParameters);
+
+	/** Get a Rigid Body With Control node reference from an anim node reference */
+	UFUNCTION(BlueprintCallable, Category = "Animation|PhysicsControl", meta = (BlueprintThreadSafe, ExpandEnumAsExecs = "Result"))
+	static FRigidBodyWithControlReference ConvertToRigidBodyWithControl(
+		const FAnimNodeReference& Node,
+		EAnimNodeReferenceConversionResult& Result);
+
+	/** Get a Rigid Body With Control node from an anim node (pure) */
+	UFUNCTION(BlueprintPure, Category = "Animation|PhysicsControl", meta = (BlueprintThreadSafe, DisplayName = "Convert to Rigid Body With Control"))
+	static void ConvertToRigidBodyWithControlPure(
+		const FAnimNodeReference& Node, 
+		FRigidBodyWithControlReference& RigidBodyWithControl, 
+		bool& Result)
+	{
+		EAnimNodeReferenceConversionResult ConversionResult;
+		RigidBodyWithControl = ConvertToRigidBodyWithControl(Node, ConversionResult);
+		Result = (ConversionResult == EAnimNodeReferenceConversionResult::Succeeded);
+	}
+
+	/** Get the names of all the controls in a specified set managed by this Rigid Body With Control node. */
+	UFUNCTION(BlueprintPure, Category = "Animation|PhysicsControl", meta = (BlueprintThreadSafe))
+	static TArray<FName> GetControlNamesInSet(
+		UPARAM(DisplayName = "Node") const FRigidBodyWithControlReference& RigidBodyWithControl, 
+		const FName SetName);
+
+	/** Get the names of all the body modifiers in a specified set managed by this Rigid Body With Control node. */
+	UFUNCTION(BlueprintPure, Category = "Animation|PhysicsControl", meta = (BlueprintThreadSafe))
+	static TArray<FName> GetBodyModifierNamesInSet(
+		UPARAM(DisplayName = "Node") const FRigidBodyWithControlReference& RigidBodyWithControl, 
+		const FName SetName);
+};
+
+#endif
\ No newline at end of file
diff --git a/Engine/Plugins/Experimental/PhysicsControl/Source/PhysicsControl/Public/RigidBodyControlData.h b/Engine/Plugins/Experimental/PhysicsControl/Source/PhysicsControl/Public/RigidBodyControlData.h
new file mode 100644
--- /dev/null
+++ b/Engine/Plugins/Experimental/PhysicsControl/Source/PhysicsControl/Public/RigidBodyControlData.h
@@ -0,0 +1,764 @@
+// Copyright Epic Games, Inc. All Rights Reserved.
+#pragma once
+
+#include "CoreMinimal.h"
+#include "Physics/ImmediatePhysics/ImmediatePhysicsDeclares.h"
+#include "Animation/AnimTypes.h"
+#include "RigidBodyControlData.generated.h"
+
+/**
+ * Setup data that is used to create the representation of a single limb. A limb is an array of 
+ * contiguous bones (e.g. left arm, or the spine etc). We can define it has the set of bones that
+ * are children of a start bone (plus the start bone itself), plus optionally the parent of that start 
+ * bone (this is useful when defining the spine, since you will want to include the pelvis, but you 
+ * don't to include all the children of the pelvis since that would include the legs), but excluding 
+ * any bones that are already part of another limb. This implies limbs should be constructed in order
+ * from leaf to root.
+ */
+USTRUCT(BlueprintType)
+struct PHYSICSCONTROL_API FRigidBodyLimbSetupData
+{
+	GENERATED_BODY();
+
+	FRigidBodyLimbSetupData()
+		: bIncludeParentBone(false), bCreateWorldSpaceControls(true)
+		, bCreateParentSpaceControls(true), bCreateBodyModifiers(true)
+	{}
+
+	/** The name of the limb that this will be used to create */
+	UPROPERTY(EditAnywhere, BlueprintReadWrite, Category = RigidBodyControl)
+	FName LimbName;
+
+	/* 
+	 * Normally the root-most bone of the limb (e.g. left clavicle when defining the left arm) - so the
+	 * limb will contain children of this bone (plus this bone itself).
+	 */
+	UPROPERTY(EditAnywhere, BlueprintReadWrite, Category = RigidBodyControl)
+	FName StartBone;
+
+	/* 
+	 * Whether or not to include the parent of the start bone. This is intended to be used for limbs like 
+	 * the spine, where you would set StartBone = spine_01 but also expect to include the pelvis (parent 
+	 * of spine_01) in the spine limb.
+	 */
+	UPROPERTY(EditAnywhere, BlueprintReadWrite, Category = RigidBodyControl)
+	uint8 bIncludeParentBone : 1;
+
+	/** Whether to create-world space controls for this limb */
+	UPROPERTY(EditAnywhere, BlueprintReadWrite, Category = RigidBodyControl)
+	uint8 bCreateWorldSpaceControls : 1;
+
+	/** Whether to create parent-space controls for this limb */
+	UPROPERTY(EditAnywhere, BlueprintReadWrite, Category = RigidBodyControl)
+	uint8 bCreateParentSpaceControls : 1;
+
+	/** Whether to create body modifiers for this limb */
+	UPROPERTY(EditAnywhere, BlueprintReadWrite, Category = RigidBodyControl)
+	uint8 bCreateBodyModifiers : 1;
+};
+
+/**
+ * Helper structure when creating the limb setup
+ */
+struct FRigidBodyLimbBones
+{
+	FRigidBodyLimbBones() 
+		: bFirstBoneIsAdditional(false), bCreateWorldSpaceControls(true)
+		, bCreateParentSpaceControls(true), bCreateBodyModifiers(true)
+	{}
+
+	/** The names of the bones in the limb */
+	TArray<FName> BoneNames;
+
+	/** Indicates if the first bone in this limb was included due to "IncludeParentBone". */
+	uint8 bFirstBoneIsAdditional : 1;
+
+	uint8 bCreateWorldSpaceControls : 1;
+	uint8 bCreateParentSpaceControls : 1;
+	uint8 bCreateBodyModifiers : 1;
+};
+
+/**
+ * Used by Body Modifiers to specify how the physical bodies should move.
+ */
+UENUM(BlueprintType)
+enum class ERigidBodyMovementType : uint8
+{
+	// Kinematic means that the object won't be simulated, but will be moved according to the
+	// kinematic target set in the Body Modifier.
+	Kinematic,
+	// Simulated means that the object will be controlled by the physics solver
+	Simulated
+};
+FName GetControlTypeName(const ERigidBodyMovementType MovementType);
+
+/**
+ * Analogous to the ControlData, this indicates how an individual controlled body should move, with flags indicating
+ * whether each element should get used.
+ */
+USTRUCT(BlueprintType)
+struct PHYSICSCONTROL_API FRigidBodyModifierSparseData
+{
+	GENERATED_BODY();
+
+	FRigidBodyModifierSparseData(
+		const ERigidBodyMovementType InMovementType = ERigidBodyMovementType::Simulated,
+		const float                  InGravityMultiplier = 1.0f)
+		: MovementType(InMovementType)
+		, GravityMultiplier(InGravityMultiplier)
+		, bEnableMovementType(1)
+		, bEnableGravityMultiplier(1)
+	{
+	}
+
+	UPROPERTY(EditAnywhere, BlueprintReadWrite, Category = PhysicsControl, meta = (editcondition = "bEnableMovementType"))
+	ERigidBodyMovementType MovementType;
+
+	UPROPERTY(EditAnywhere, BlueprintReadWrite, Category = PhysicsControl, meta = (editcondition = "bEnableGravityMultiplier"))
+	float GravityMultiplier;
+
+	UPROPERTY(EditAnywhere, Category = PhysicsControl, meta = (InlineEditConditionToggle))
+	uint8 bEnableMovementType : 1;
+
+	UPROPERTY(EditAnywhere, Category = PhysicsControl, meta = (InlineEditConditionToggle))
+	uint8 bEnableGravityMultiplier : 1;
+};
+
+/**
+ * Analogous to the ControlData, this indicates how an individual controlled body should move
+ */
+USTRUCT(BlueprintType)
+struct PHYSICSCONTROL_API FRigidBodyModifierData
+{
+	GENERATED_BODY();
+
+	FRigidBodyModifierData(
+		const ERigidBodyMovementType InMovementType = ERigidBodyMovementType::Simulated,
+		const float                  InGravityMultiplier = 1.0f)
+		: MovementType(InMovementType)
+		, GravityMultiplier(InGravityMultiplier)
+	{
+	}
+
+	void UpdateFromSparseData(const FRigidBodyModifierSparseData& SparseData);
+
+	UPROPERTY(EditAnywhere, BlueprintReadWrite, Category = PhysicsControl)
+	ERigidBodyMovementType MovementType;
+
+	UPROPERTY(EditAnywhere, BlueprintReadWrite, Category = PhysicsControl)
+	float GravityMultiplier;
+};
+
+/**
+ * Specifies the type of control that is created when making controls from a skeleton or a set of limbs.
+ * Note that if controls are made individually then other options are available - i.e. in a character,
+ * any body part can be controlled relative to any other part, or indeed any other object.
+ */
+UENUM(BlueprintType)
+enum class ERigidBodyControlType : uint8
+{
+	/** Control is done in world space, so each object/part is driven independently */
+	WorldSpace,
+	/** Control is done in the space of the parent of each object */
+	ParentSpace,
+};
+FName GetControlTypeName(const ERigidBodyControlType ControlType);
+
+/**
+ * Strength and damping etc parameters that will affect a control, with flags indicating
+ * whether each element should get used.
+ */
+USTRUCT(BlueprintType)
+struct PHYSICSCONTROL_API FRigidBodyControlSparseData
+{
+	GENERATED_BODY();
+
+	FRigidBodyControlSparseData(
+		float InLinearStrength = 0.0f, float InLinearDampingRatio = 1.0f, float InLinearExtraDamping = 0.0f,
+		float InAngularStrength = 0.0f, float InAngularDampingRatio = 1.0f, float InAngularExtraDamping = 0.0f,
+		float InLinearTargetVelocityMultiplier = 1.0f, float InAngularTargetVelocityMultiplier = 1.0f,
+		bool bInEnabled = true)
+		: LinearStrength(InLinearStrength)
+		, LinearDampingRatio(InLinearDampingRatio)
+		, LinearExtraDamping(InLinearExtraDamping)
+		, AngularStrength(InAngularStrength)
+		, AngularDampingRatio(InAngularDampingRatio)
+		, AngularExtraDamping(InAngularExtraDamping)
+		, LinearTargetVelocityMultiplier(InLinearTargetVelocityMultiplier)
+		, AngularTargetVelocityMultiplier(InAngularTargetVelocityMultiplier)
+		, bEnabled(bInEnabled)
+		, bEnableLinearStrength(true)
+		, bEnableLinearDampingRatio(true)
+		, bEnableLinearExtraDamping(true)
+		, bEnableAngularStrength(true)
+		, bEnableAngularDampingRatio(true)
+		, bEnableAngularExtraDamping(true)
+		, bEnableLinearTargetVelocityMultiplier(true)
+		, bEnableAngularTargetVelocityMultiplier(true)
+		, bEnablebEnabled(true)
+	{
+	}
+
+	/** The strength used to drive linear motion */
+	UPROPERTY(EditAnywhere, BlueprintReadWrite, Category = PhysicsControl, meta = (ClampMin = "0.0", editcondition = "bEnableLinearStrength"))
+	float LinearStrength;
+
+	/** 
+	 * The amount of damping associated with the linear strength. A value of 1 Results in critically 
+	 * damped motion where the control drives as quickly as possible to the target without overshooting. 
+	 * Values > 1 result in more damped motion, and values below 1 result in faster, but more "wobbly" motion.
+	 */
+	UPROPERTY(EditAnywhere, BlueprintReadWrite, Category = PhysicsControl, meta = (ClampMin = "0.0", editcondition = "bEnableLinearDampingRatio"))
+	float LinearDampingRatio;
+
+	/** 
+	 * The amount of additional linear damping. This is added to the damping that comes from LinearDampingRatio
+	 * and can be useful when you want damping even when LinearStrength is zero.
+	 */
+	UPROPERTY(EditAnywhere, BlueprintReadWrite, Category = PhysicsControl, meta = (ClampMin = "0.0", editcondition = "bEnableLinearExtraDamping"))
+	float LinearExtraDamping;
+
+	/** The strength used to drive angular motion */
+	UPROPERTY(EditAnywhere, BlueprintReadWrite, Category = PhysicsControl, meta = (ClampMin = "0.0", editcondition = "bEnableAngularStrength"))
+	float AngularStrength;
+
+	/** 
+	 * The amount of damping associated with the angular strength. A value of 1 Results in critically 
+	 * damped motion where the control drives as quickly as possible to the target without overshooting. 
+	 * Values > 1 result in more damped motion, and values below 1 result in faster, but more "wobbly" motion.
+	 */
+	UPROPERTY(EditAnywhere, BlueprintReadWrite, Category = PhysicsControl, meta = (ClampMin = "0.0", editcondition = "bEnableAngularDampingRatio"))
+	float AngularDampingRatio;
+
+	/** 
+	 * The amount of additional angular damping. This is added to the damping that comes from AngularDampingRatio
+	 * and can be useful when you want damping even when AngularStrength is zero.
+	 */
+	UPROPERTY(EditAnywhere, BlueprintReadWrite, Category = PhysicsControl, meta = (ClampMin = "0.0", editcondition = "bEnableAngularExtraDamping"))
+	float AngularExtraDamping;
+
+	/**
+	 * Multiplier on the velocity, which gets applied to the damping. A value of 1 means the animation target
+	 * velocity is used, which helps it track the animation. A value of 0 means damping happens in "world space" 
+	 * - so damping acts like drag on the movement.
+	 */
+	UPROPERTY(EditAnywhere, BlueprintReadWrite, Category = PhysicsControl, meta = (ClampMin = "0.0", editcondition = "bEnableLinearTargetVelocityMultiplier"))
+	float LinearTargetVelocityMultiplier;
+
+	/**
+	 * Multiplier on the angular velocity, which gets applied to the damping. A value of 1 means the animation target
+	 * velocity is used, which helps it track the animation. A value of 0 means damping happens in "world space" 
+	 * - so damping acts like drag on the movement.
+	 */
+	UPROPERTY(EditAnywhere, BlueprintReadWrite, Category = PhysicsControl, meta = (ClampMin = "0.0", editcondition = "bEnableAngularTargetVelocityMultiplier"))
+	float AngularTargetVelocityMultiplier;
+
+	/**
+	 * Whether this control should be enabled
+	 */
+	UPROPERTY(EditAnywhere, BlueprintReadWrite, Category = PhysicsControl, meta = (editcondition = "bEnablebEnabled"))
+	uint8 bEnabled : 1;
+
+	UPROPERTY(EditAnywhere, Category = PhysicsControl, meta = (InlineEditConditionToggle))
+	uint8 bEnableLinearStrength : 1;
+
+	UPROPERTY(EditAnywhere, Category = PhysicsControl, meta = (InlineEditConditionToggle))
+	uint8 bEnableLinearDampingRatio : 1;
+
+	UPROPERTY(EditAnywhere, Category = PhysicsControl, meta = (InlineEditConditionToggle))
+	uint8 bEnableLinearExtraDamping : 1;
+
+	UPROPERTY(EditAnywhere, Category = PhysicsControl, meta = (InlineEditConditionToggle))
+	uint8 bEnableAngularStrength : 1;
+
+	UPROPERTY(EditAnywhere, Category = PhysicsControl, meta = (InlineEditConditionToggle))
+	uint8 bEnableAngularDampingRatio : 1;
+
+	UPROPERTY(EditAnywhere, Category = PhysicsControl, meta = (InlineEditConditionToggle))
+	uint8 bEnableAngularExtraDamping : 1;
+
+	UPROPERTY(EditAnywhere, Category = PhysicsControl, meta = (InlineEditConditionToggle))
+	uint8 bEnableLinearTargetVelocityMultiplier : 1;
+
+	UPROPERTY(EditAnywhere, Category = PhysicsControl, meta = (InlineEditConditionToggle))
+	uint8 bEnableAngularTargetVelocityMultiplier : 1;
+
+	UPROPERTY(EditAnywhere, Category = PhysicsControl, meta = (InlineEditConditionToggle))
+	uint8 bEnablebEnabled : 1;
+};
+
+
+
+/**
+ * Strength and damping etc parameters that will affect a control.
+ */
+USTRUCT(BlueprintType)
+struct PHYSICSCONTROL_API FRigidBodyControlData
+{
+	GENERATED_BODY();
+
+	FRigidBodyControlData(
+		float InLinearStrength = 0.0f, float InLinearDampingRatio = 1.0f, float InLinearExtraDamping = 0.0f,
+		float InAngularStrength = 0.0f, float InAngularDampingRatio = 1.0f, float InAngularExtraDamping = 0.0f,
+		float InLinearTargetVelocityMultiplier = 1.0f, float InAngularTargetVelocityMultiplier = 1.0f,
+		bool bInEnabled = true)
+		: LinearStrength(InLinearStrength)
+		, LinearDampingRatio(InLinearDampingRatio)
+		, LinearExtraDamping(InLinearExtraDamping)
+		, AngularStrength(InAngularStrength)
+		, AngularDampingRatio(InAngularDampingRatio)
+		, AngularExtraDamping(InAngularExtraDamping)
+		, LinearTargetVelocityMultiplier(InLinearTargetVelocityMultiplier)
+		, AngularTargetVelocityMultiplier(InAngularTargetVelocityMultiplier)
+		, bEnabled(bInEnabled)
+	{
+	}
+
+	void UpdateFromSparseData(const FRigidBodyControlSparseData& SparseData);
+
+	/** The strength used to drive linear motion */
+	UPROPERTY(EditAnywhere, BlueprintReadWrite, Category = PhysicsControl, meta = (ClampMin = "0.0"))
+	float LinearStrength;
+
+	/** 
+	 * The amount of damping associated with the linear strength. A value of 1 Results in critically 
+	 * damped motion where the control drives as quickly as possible to the target without overshooting. 
+	 * Values > 1 result in more damped motion, and values below 1 result in faster, but more "wobbly" motion.
+	 */
+	UPROPERTY(EditAnywhere, BlueprintReadWrite, Category = PhysicsControl, meta = (ClampMin = "0.0"))
+	float LinearDampingRatio;
+
+	/** 
+	 * The amount of additional linear damping. This is added to the damping that comes from LinearDampingRatio
+	 * and can be useful when you want damping even when LinearStrength is zero.
+	 */
+	UPROPERTY(EditAnywhere, BlueprintReadWrite, Category = PhysicsControl, meta = (ClampMin = "0.0"))
+	float LinearExtraDamping;
+
+	/** The strength used to drive angular motion */
+	UPROPERTY(EditAnywhere, BlueprintReadWrite, Category = PhysicsControl, meta = (ClampMin = "0.0"))
+	float AngularStrength;
+
+	/** 
+	 * The amount of damping associated with the angular strength. A value of 1 Results in critically 
+	 * damped motion where the control drives as quickly as possible to the target without overshooting. 
+	 * Values > 1 result in more damped motion, and values below 1 result in faster, but more "wobbly" motion.
+	 */
+	UPROPERTY(EditAnywhere, BlueprintReadWrite, Category = PhysicsControl, meta = (ClampMin = "0.0"))
+	float AngularDampingRatio;
+
+	/** 
+	 * The amount of additional angular damping. This is added to the damping that comes from AngularDampingRatio
+	 * and can be useful when you want damping even when AngularStrength is zero.
+	 */
+	UPROPERTY(EditAnywhere, BlueprintReadWrite, Category = PhysicsControl, meta = (ClampMin = "0.0"))
+	float AngularExtraDamping;
+
+	/**
+	 * Multiplier on the velocity, which gets applied to the damping. A value of 1 means the animation target
+	 * velocity is used, which helps it track the animation. A value of 0 means damping happens in "world space" 
+	 * - so damping acts like drag on the movement.
+	 */
+	UPROPERTY(EditAnywhere, BlueprintReadWrite, Category = PhysicsControl, meta = (ClampMin = "0.0"))
+	float LinearTargetVelocityMultiplier;
+
+	/**
+	 * Multiplier on the angular velocity, which gets applied to the damping. A value of 1 means the animation target
+	 * velocity is used, which helps it track the animation. A value of 0 means damping happens in "world space" 
+	 * - so damping acts like drag on the movement.
+	 */
+	UPROPERTY(EditAnywhere, BlueprintReadWrite, Category = PhysicsControl, meta = (ClampMin = "0.0"))
+	float AngularTargetVelocityMultiplier;
+
+	/**
+	 * Whether this control should be enabled
+	 */
+	UPROPERTY(EditAnywhere, BlueprintReadWrite, Category = PhysicsControl)
+	uint8 bEnabled : 1;
+};
+
+/**
+ * Setup data for all the bodies controlled by the node. Contains info to split the skeleton up into limbs, 
+ * and default control and modifier settings for each of them.
+ */
+USTRUCT(BlueprintType)
+struct FRigidBodySetupData
+{
+	GENERATED_BODY();
+
+	UPROPERTY(EditAnywhere, Category = ControlSetup)
+	TArray<FRigidBodyLimbSetupData> LimbSetupData;
+
+	UPROPERTY(EditAnywhere, Category = ControlSetup)
+	FRigidBodyControlData DefaultWorldSpaceControlData;
+
+	UPROPERTY(EditAnywhere, Category = ControlSetup)
+	FRigidBodyControlData DefaultParentSpaceControlData;
+
+	UPROPERTY(EditAnywhere, Category = ControlSetup)
+	FRigidBodyModifierData DefaultBodyModifierData;
+};
+
+
+/** 
+ * Specifies a single/individual control between parent/child bodies, with the control data
+ */
+USTRUCT(BlueprintType)
+struct PHYSICSCONTROL_API FRigidBodyControl
+{
+	GENERATED_BODY();
+
+	UPROPERTY(EditAnywhere, BlueprintReadWrite, Category = PhysicsControl)
+	FName ParentBoneName;
+
+	UPROPERTY(EditAnywhere, BlueprintReadWrite, Category = PhysicsControl)
+	FName ChildBoneName;
+
+	UPROPERTY(EditAnywhere, BlueprintReadWrite, Category = PhysicsControl)
+	FRigidBodyControlData ControlData;
+
+	bool IsEnabled() const { return ControlData.bEnabled; }
+};
+
+/**
+ * Used on creation, to allow requesting the control to be in certain sets
+ */
+USTRUCT(BlueprintType)
+struct PHYSICSCONTROL_API FRigidBodyControlCreation
+{
+	GENERATED_BODY();
+
+	UPROPERTY(EditAnywhere, BlueprintReadWrite, Category = PhysicsControl)
+	FRigidBodyControl Control;
+
+	UPROPERTY(EditAnywhere, BlueprintReadWrite, Category = PhysicsControl)
+	TArray<FName> Sets;
+};
+
+/**
+ * Used to create a single/individual body modifier, with the modifier data,
+ * and what sets it should be added to.
+ */
+USTRUCT(BlueprintType)
+struct PHYSICSCONTROL_API FRigidBodyModifier
+{
+	GENERATED_BODY();
+
+	UPROPERTY(EditAnywhere, BlueprintReadWrite, Category = PhysicsControl)
+	FName BoneName;
+
+	UPROPERTY(EditAnywhere, BlueprintReadWrite, Category = PhysicsControl)
+	FRigidBodyModifierData ModifierData;
+
+	UPROPERTY(EditAnywhere, BlueprintReadWrite, Category = PhysicsControl)
+	TArray<FName> Sets;
+};
+
+/**
+ * Used on creation, to allow requesting the modifier to be in certain sets
+ */
+USTRUCT(BlueprintType)
+struct PHYSICSCONTROL_API FRigidBodyModifierCreation
+{
+	GENERATED_BODY();
+
+	UPROPERTY(EditAnywhere, BlueprintReadWrite, Category = PhysicsControl)
+	FRigidBodyModifier Modifier;
+
+	UPROPERTY(EditAnywhere, BlueprintReadWrite, Category = PhysicsControl)
+	TArray<FName> Sets;
+};
+
+/**
+ * Collection of controls and body modifiers, used for creation
+ */
+USTRUCT(BlueprintType)
+struct PHYSICSCONTROL_API FRigidBodyControlAndBodyModifierCreations
+{
+	GENERATED_BODY();
+
+	UPROPERTY(EditAnywhere, BlueprintReadWrite, Category = PhysicsControl)
+	TArray<FRigidBodyControlCreation> Controls;
+
+	UPROPERTY(EditAnywhere, BlueprintReadWrite, Category = PhysicsControl)
+	TArray<FRigidBodyModifierCreation> Modifiers;
+};
+
+/**
+ * Update an existing set, or add to it
+ */
+USTRUCT(BlueprintType)
+struct PHYSICSCONTROL_API FSetUpdate
+{
+	GENERATED_BODY();
+
+	UPROPERTY(EditAnywhere, BlueprintReadWrite, Category = PhysicsControl)
+	FName SetName;
+
+	/** The names of either controls or body modifiers (depending on context), or sets of controls/body modifiers */
+	UPROPERTY(EditAnywhere, BlueprintReadWrite, Category = PhysicsControl)
+	TArray<FName> Names;
+};
+
+/**
+ * Combines updates for control and modifier sets
+ */
+USTRUCT(BlueprintType)
+struct PHYSICSCONTROL_API FRigidBodySetUpdates
+{
+	GENERATED_BODY();
+
+	UPROPERTY(EditAnywhere, BlueprintReadWrite, Category = PhysicsControl)
+	TArray<FSetUpdate> ControlSetUpdates;
+
+	UPROPERTY(EditAnywhere, BlueprintReadWrite, Category = PhysicsControl)
+	TArray<FSetUpdate> ModifierSetUpdates;
+};
+
+FRigidBodyControlData Interpolate(
+	const FRigidBodyControlData& A, const FRigidBodyControlData& B, const float Weight);
+FRigidBodyControlSparseData Interpolate(
+	const FRigidBodyControlSparseData& A, const FRigidBodyControlSparseData& B, const float Weight);
+
+/**
+ * Data that can be used to parameterize (modify/update) a control 
+ */
+USTRUCT(BlueprintType)
+struct PHYSICSCONTROL_API FRigidBodyNamedControlParameters
+{
+	GENERATED_BODY();
+
+	FRigidBodyNamedControlParameters() {}
+
+	FRigidBodyNamedControlParameters(FName InName, const FRigidBodyControlSparseData& InData) 
+		: Name(InName), Data(InData) {}
+
+	// The name of the control (or set of controls) to update
+	UPROPERTY(EditAnywhere, BlueprintReadWrite, Category = PhysicsControl)
+	FName Name;
+
+	UPROPERTY(EditAnywhere, BlueprintReadWrite, Category = PhysicsControl)
+	FRigidBodyControlSparseData Data;
+};
+
+FRigidBodyModifierData Interpolate(
+	const FRigidBodyModifierData& A, const FRigidBodyModifierData& B, const float Weight);
+FRigidBodyModifierSparseData Interpolate(
+	const FRigidBodyModifierSparseData& A, const FRigidBodyModifierSparseData& B, const float Weight);
+
+/**
+ * Data that can be used to parameterize(modify / update) a control
+ */
+USTRUCT(BlueprintType)
+struct PHYSICSCONTROL_API FRigidBodyNamedModifierParameters
+{
+	GENERATED_BODY();
+
+	FRigidBodyNamedModifierParameters() {}
+
+	FRigidBodyNamedModifierParameters(FName InName, const FRigidBodyModifierSparseData& InData)
+		: Name(InName), Data(InData) {}
+
+	// The name of the modifier (or set of modifiers) to update
+	UPROPERTY(EditAnywhere, BlueprintReadWrite, Category = PhysicsControl)
+	FName Name;
+
+	UPROPERTY(EditAnywhere, BlueprintReadWrite, Category = PhysicsControl)
+	FRigidBodyModifierSparseData Data;
+};
+
+/**
+ * These apply temporary/ephemeral changes to the controls that only persist for one tick.
+ */
+USTRUCT(BlueprintType)
+struct PHYSICSCONTROL_API FRigidBodyControlAndModifierParameters
+{
+	GENERATED_BODY();
+
+	/**
+	 * Parameters for existing controls. Each name can be the name of a control, or the name of a 
+	 * set of controls. They will only apply for one tick/update. They will be applied in order (so 
+	 * subsequent entries will override earlier ones if they apply to the same control).
+	 */ 
+	UPROPERTY(EditAnywhere, BlueprintReadWrite, Category = Settings)
+	TArray<FRigidBodyNamedControlParameters> ControlParameters;
+	
+	/**
+	 *  Parameters for existing modifiers. Each name can be the name of a modifier, or the name of a 
+	 * set of modifiers. They will only apply for one tick/update.
+	 */
+	UPROPERTY(EditAnywhere, BlueprintReadWrite, Category = Settings)
+	TArray<FRigidBodyNamedModifierParameters> ModifierParameters;
+
+	void Add(const FRigidBodyNamedControlParameters& InParameters) { ControlParameters.Add(InParameters); }
+	void Add(const FRigidBodyNamedModifierParameters& InParameters) { ModifierParameters.Add(InParameters); }
+};
+
+/**
+ * These apply permanent changes to the controls and modifiers, allowing all the settings to be changed
+ * (apart from the actual bodies that are being controlled/affected)
+ */
+USTRUCT(BlueprintType)
+struct PHYSICSCONTROL_API FRigidBodyControlAndModifierUpdates
+{
+	GENERATED_BODY();
+
+	/** Modifications to the underlying controls - these will persist */
+	UPROPERTY(EditAnywhere, BlueprintReadWrite, Category = PhysicsControl)
+	TArray<FRigidBodyNamedControlParameters> ControlParameters;
+
+	/** Modifications to the underlying modifiers - these will persist */
+	UPROPERTY(EditAnywhere, BlueprintReadWrite, Category = PhysicsControl)
+	TArray<FRigidBodyNamedModifierParameters> ModifierParameters;
+};
+
+/**
+ * A single target for a control, which may be defined as an offset from the (implicit) animation target.
+ */
+USTRUCT(BlueprintType)
+struct PHYSICSCONTROL_API FRigidBodyControlTarget
+{
+	GENERATED_BODY();
+
+	FRigidBodyControlTarget()
+		: TargetPosition(ForceInitToZero)
+		, TargetOrientation(ForceInitToZero)
+		, TargetPoint(ForceInitToZero)
+		, bUseSkeletalAnimation(true)
+		, bUseTargetPoint(false)
+	{
+	}
+
+	/** The target position of the child body, relative to the parent body */
+	UPROPERTY(EditAnywhere, BlueprintReadWrite, Category = PhysicsControl)
+	FVector TargetPosition;
+
+	/** The target orientation of the child body, relative to the parent body */
+	UPROPERTY(EditAnywhere, BlueprintReadWrite, Category = PhysicsControl)
+	FRotator TargetOrientation;
+
+	/** 
+	 * The point on the controlled (child) object that should be driven towards the target position.
+	 * Note that if this is not set (i.e. if UseTargetPoint is false) then the centre of mass will be used. 
+	 */
+	UPROPERTY(EditAnywhere, BlueprintReadWrite, Category = PhysicsControl, meta = (editcondition = "bUseTargetPoint"))
+	FVector TargetPoint;
+
+	/** If true then the target will be applied on top of the skeletal animation */
+	UPROPERTY(EditAnywhere, BlueprintReadWrite, Category = PhysicsControl)
+	uint8 bUseSkeletalAnimation : 1;
+
+	/** If true then the target will be applied on top of the skeletal animation */
+	UPROPERTY(EditAnywhere, BlueprintReadWrite, Category = PhysicsControl)
+	uint8 bUseTargetPoint : 1;
+};
+
+/**
+ * A set of targets for controls
+ */
+USTRUCT(BlueprintType)
+struct PHYSICSCONTROL_API FRigidBodyControlTargets
+{
+	GENERATED_BODY();
+
+	/** Targets to apply to the named control */
+	UPROPERTY(EditAnywhere, BlueprintReadWrite, Category = PhysicsControl)
+	TMap<FName, FRigidBodyControlTarget> Targets;
+};
+
+/**
+ * A single kinematic target, which may be defined as an offset from the (implicit) animation target.
+ */
+USTRUCT(BlueprintType)
+struct PHYSICSCONTROL_API FRigidBodyKinematicTarget
+{
+	GENERATED_BODY();
+
+	FRigidBodyKinematicTarget()
+		: TargetPosition(ForceInitToZero)
+		, TargetOrientation(ForceInitToZero)
+		, bUseSkeletalAnimation(true)
+	{
+	}
+
+	/** The target position of the body */
+	UPROPERTY(EditAnywhere, BlueprintReadWrite, Category = PhysicsControl)
+	FVector TargetPosition;
+
+	/** The target orientation of the body */
+	UPROPERTY(EditAnywhere, BlueprintReadWrite, Category = PhysicsControl)
+	FRotator TargetOrientation;
+
+	/** If true then the target will be applied on top of the skeletal animation */
+	UPROPERTY(EditAnywhere, BlueprintReadWrite, Category = PhysicsControl)
+	uint8 bUseSkeletalAnimation : 1;
+};
+
+/**
+ * A set of kinematic targets
+ */
+USTRUCT(BlueprintType)
+struct PHYSICSCONTROL_API FRigidBodyKinematicTargets
+{
+	GENERATED_BODY();
+
+	/** Targets to apply to the named body modifier */
+	UPROPERTY(EditAnywhere, BlueprintReadWrite, Category = PhysicsControl)
+	TMap<FName, FRigidBodyKinematicTarget> Targets;
+};
+
+/**
+ * This is the record for a control. It stores the "original" data, which will normally stay fixed.
+ * In addition, it will also store some updates to these original data which will determine how the actual
+ * controls get used.
+ */
+struct FControlRecord
+{
+	FControlRecord(const FRigidBodyControl& InControl, ImmediatePhysics::FJointHandle* InJointHandle);
+
+	void ResetCurrent(bool bResetTarget);
+
+	// Note that this is only correct when called during or after the update has been done
+	bool IsEnabled() const { return CurrentData.bEnabled; }
+
+	// TODO - might benefit from smaller - non-blueprint data members here ?
+	FRigidBodyControl               Control;
+	ImmediatePhysics::FJointHandle* JointHandle;
+
+	// This contains the currently active control data. It will be updated just prior to
+	// applying the controls, by setting it to the default, and then updating it with any parameters.
+	FRigidBodyControlData CurrentData;
+
+	// Contains any control target that has been set
+	FRigidBodyControlTarget ControlTarget;
+
+	// The previous control target. This will have been set at the end of a previous update (but
+	// only if the control was enabled etc), so to check if it is valid, check the update counter.
+	FTransform PrevTargetTM;
+
+	// Update counter set when the control was last updated.
+	// TODO just store the count we're interested in rather than the whole structure
+	FGraphTraversalCounter ExpectedUpdateCounter;
+};
+
+/**
+ * This is the record for a modifier. It stores the "original" data, which will normally stay fixed.
+ * In addition, it will also store some updates to these original data which will determine how the actual
+ * body properties get set.
+ */
+struct FBodyModifierRecord
+{
+	FBodyModifierRecord(const FRigidBodyModifier& InModifier, ImmediatePhysics::FActorHandle* InActorHandle);
+
+	void ResetCurrent();
+
+	// TODO - might benefit from smaller - non-blueprint data members here ?
+	FRigidBodyModifier              Modifier;
+	ImmediatePhysics::FActorHandle* ActorHandle;
+
+	// This contains the currently active modifier data. It will be updated just prior to
+	// applying the controls, by setting it to the default, and then updating it with any parameters.
+	FRigidBodyModifierData CurrentData;
+};
+
+
diff --git a/Engine/Plugins/Experimental/PhysicsControl/Source/PhysicsControl/Public/RigidBodyNameRecords.h b/Engine/Plugins/Experimental/PhysicsControl/Source/PhysicsControl/Public/RigidBodyNameRecords.h
new file mode 100644
--- /dev/null
+++ b/Engine/Plugins/Experimental/PhysicsControl/Source/PhysicsControl/Public/RigidBodyNameRecords.h
@@ -0,0 +1,55 @@
+// Copyright Epic Games, Inc. All Rights Reserved.
+
+#pragma once
+
+#include "CoreMinimal.h"
+
+/**
+ * All the controls and body modifiers we create get given names, so we can keep a record in a form
+ * that will be useful to users. The three general uses are:
+ * 1. All controls or modifiers
+ * 2. The type of control - the obvious ones are "world space" and "parent space", but users might make others
+ * 3. Sets of control/modifiers - the obvious ones are limbs (e.g. a set of world-space controls on the leg)
+ *    but users might want to make other sets - e.g. "UpperBody"
+ */
+struct FRigidBodyNameRecords
+{
+	/** Adds Name to SetName as well as adding it to the set "All" */
+	void AddControl(FName Name, FName SetName);
+	/** Adds Name to SetNames as well as adding it to the set "All" */
+	void AddControl(FName Name, const TArray<FName>& SetNames);
+	/** Adds a collection of Names to SetName as well as adding them to the set "All"  */
+	void AddControls(const TArray<FName>& ControlNames, const FName SetName);
+	/** Removes Name from all Sets */
+	void RemoveControl(FName Name);
+	const TArray<FName>& GetControlNamesInSet(FName SetName) const;
+
+	/** Adds Name to SetName as well as adding it to the set "All" */
+	void AddBodyModifier(FName Name, FName SetName);
+	/** Adds Name to SetNames as well as adding it to the set "All" */
+	void AddBodyModifier(FName Name, const TArray<FName>& SetNames);
+	/** Adds a collection of Names to SetName as well as adding them to the set "All"  */
+	void AddBodyModifiers(const TArray<FName>& BodyModifierNames, const FName SetName);
+	/** Removes Name from all Sets */
+	void RemoveBodyModifier(FName Name);
+	const TArray<FName>& GetBodyModifierNamesInSet(FName SetName) const;
+
+	/** Remove all control and modifier sets. */
+	void Reset();
+
+	/**
+	 * All the control sets we've created, arranged by set name. Note that some controls will
+	 * not necessarily have been added to a set.
+	 */
+	TMap<FName, TArray<FName>> ControlSets;
+
+	/**
+	 * All the body modifiers we've created, arranged by set name. Note that some modifiers will
+	 * not necessarily have been added to a set.
+	 * */
+	TMap<FName, TArray<FName>> BodyModifierSets;
+};
+
+/** Returns a copy of the supplied list of names in which any names that match keys in the supplied 'Sets' map are replaced with all the names in the array associated with that key. */
+TArray<FName> ExpandSetName(const FName InName, const TMap<FName, TArray<FName>>& Sets);
+TArray<FName> ExpandSetNames(const TArray<FName>& InNames, const TMap<FName, TArray<FName>>& Sets);
diff --git a/Engine/Plugins/Experimental/PhysicsControl/Source/PhysicsControl/Public/RigidBodyPoseData.h b/Engine/Plugins/Experimental/PhysicsControl/Source/PhysicsControl/Public/RigidBodyPoseData.h
new file mode 100644
--- /dev/null
+++ b/Engine/Plugins/Experimental/PhysicsControl/Source/PhysicsControl/Public/RigidBodyPoseData.h
@@ -0,0 +1,85 @@
+// Copyright Epic Games, Inc. All Rights Reserved.
+#pragma once
+
+#include "CoreMinimal.h"
+#include "Engine/EngineTypes.h"
+#include "BonePose.h"
+
+// Use the simulation space functions from the RBAN
+#include "BoneControllers/AnimNode_RigidBody.h" 
+
+struct FComponentSpacePoseContext;
+
+namespace RigidBodyWithControl
+{
+
+//======================================================================================================================
+struct FOutputBoneData
+{
+	FOutputBoneData()
+		: CompactPoseBoneIndex(INDEX_NONE), CompactPoseParentBoneIndex(INDEX_NONE)
+	{}
+
+	TArray<FCompactPoseBoneIndex> BoneIndicesToParentBody;
+	FCompactPoseBoneIndex CompactPoseBoneIndex;
+	FCompactPoseBoneIndex CompactPoseParentBoneIndex;
+	int32 BodyIndex; // Index into Bodies - will be the same index as into the joints
+	int32 ParentBodyIndex;
+};
+
+//======================================================================================================================
+struct FBoneData
+{
+	FBoneData()
+		: Position(FVector::ZeroVector), Orientation(FQuat::Identity) 
+	{}
+	FBoneData(const FVector& InPosition, const FQuat& InOrientation)
+		: Position(InPosition)
+		, Orientation(InOrientation)
+	{}
+
+	/**
+	 * Sets position/orientation in both current and previous (i.e. implying zero velocity)
+	 */
+	void Set(const FVector& InPosition, const FQuat& InOrientation) { 
+		Position = InPosition; Orientation = InOrientation; }
+
+	FTransform GetTM() const { return FTransform(Orientation, Position); }
+
+	FVector Position;
+	FQuat   Orientation;
+};
+
+//======================================================================================================================
+struct FRigidBodyPoseData
+{
+	void Update(
+		FComponentSpacePoseContext&   ComponentSpacePoseContext,
+		const TArray<FOutputBoneData> OutputBoneData,
+		const ESimulationSpace        SimulationSpace,
+		const FBoneReference&         BaseBoneRef,
+		const FGraphTraversalCounter& InUpdateCounter);
+
+	FTransform GetTM(int32 Index) const { return BoneData[Index].GetTM(); }
+	bool IsValidIndex(const int32 Index) const { return BoneData.IsValidIndex(Index); }
+	bool IsEmpty() const { return BoneData.IsEmpty(); }
+
+	/**
+	 * The cached skeletal data, updated at the start of each tick
+	 */
+	TArray<FBoneData> BoneData;
+
+	/**
+	 * The origin (in world space). BoneData transforms are relative to this.
+	 */
+	FTransform WorldCoordinateFrame;
+
+	// Track when we were currently/last updated so the user can detect missing updates if calculating
+	// velocity etc
+	FGraphTraversalCounter UpdateCounter;
+	// When the update is called we'll take the current counter, increment it, and store here so it
+	// can be compared.
+	FGraphTraversalCounter ExpectedUpdateCounter;
+};
+
+}
\ No newline at end of file
diff --git a/Engine/Plugins/Experimental/PhysicsControl/Source/PhysicsControlEditor/PhysicsControlEditor.Build.cs b/Engine/Plugins/Experimental/PhysicsControl/Source/PhysicsControlEditor/PhysicsControlEditor.Build.cs
--- a/Engine/Plugins/Experimental/PhysicsControl/Source/PhysicsControlEditor/PhysicsControlEditor.Build.cs
+++ b/Engine/Plugins/Experimental/PhysicsControl/Source/PhysicsControlEditor/PhysicsControlEditor.Build.cs
@@ -35,21 +35,38 @@ public class PhysicsControlEditor : ModuleRules
 		PrivateDependencyModuleNames.AddRange(
 			new string[]
 			{
+				"AnimationCore",
+				"AnimationEditMode",
+				"AnimGraph",
+				"AnimGraphRuntime",
+				"ComponentVisualizers",
+				"Core",
 				"CoreUObject",
 				"EditorFramework",
 				"Engine",
+				"Persona",
+				"PhysicsControl",
 				"Slate",
 				"SlateCore",
 				"UnrealEd",
-				"Persona",
-				"AnimationEditMode",
-				"ComponentVisualizers",
-				"AnimGraph",
-				"PhysicsControl"
 			}
 			);
-		
-		
+
+		if (Target.bBuildEditor == true)
+		{
+			PrivateDependencyModuleNames.AddRange(
+				new string[]
+				{
+						"AnimationEditor",
+						"AnimGraph",
+						"BlueprintGraph",
+						"EditorFramework",
+						"Kismet",
+						"UnrealEd",
+				}
+			);
+		}
+
 		DynamicallyLoadedModuleNames.AddRange(
 			new string[]
 			{
diff --git a/Engine/Plugins/Experimental/PhysicsControl/Source/PhysicsControlEditor/Private/PhysicsControlEditorModule.cpp b/Engine/Plugins/Experimental/PhysicsControl/Source/PhysicsControlEditor/Private/PhysicsControlEditorModule.cpp
--- a/Engine/Plugins/Experimental/PhysicsControl/Source/PhysicsControlEditor/Private/PhysicsControlEditorModule.cpp
+++ b/Engine/Plugins/Experimental/PhysicsControl/Source/PhysicsControlEditor/Private/PhysicsControlEditorModule.cpp
@@ -63,6 +63,6 @@ void FPhysicsControlEditorModule::ShutdownModule()
 
 #undef LOCTEXT_NAMESPACE
 	
-IMPLEMENT_MODULE(FPhysicsControlEditorModule, PhysicsControl)
+IMPLEMENT_MODULE(FPhysicsControlEditorModule, PhysicsControlEditor)
 
 
diff --git a/Engine/Plugins/Experimental/PhysicsControl/Source/PhysicsControlUncookedOnly/PhysicsControlUncookedOnly.Build.cs b/Engine/Plugins/Experimental/PhysicsControl/Source/PhysicsControlUncookedOnly/PhysicsControlUncookedOnly.Build.cs
new file mode 100644
--- /dev/null
+++ b/Engine/Plugins/Experimental/PhysicsControl/Source/PhysicsControlUncookedOnly/PhysicsControlUncookedOnly.Build.cs
@@ -0,0 +1,41 @@
+// Copyright Epic Games, Inc. All Rights Reserved.
+
+using UnrealBuildTool;
+public class PhysicsControlUncookedOnly : ModuleRules
+{
+	public PhysicsControlUncookedOnly(ReadOnlyTargetRules Target) : base(Target)
+	{
+		PCHUsage = ModuleRules.PCHUsageMode.UseExplicitOrSharedPCHs;
+		UnsafeTypeCastWarningLevel = WarningLevel.Error;
+
+		PrivateDependencyModuleNames.AddRange(
+			new string[]
+			{
+				"AnimationCore",
+				"AnimGraphRuntime",
+				"Core",
+				"CoreUObject",
+				"Engine",
+				"Persona",
+				"PhysicsControl",
+				"Slate",
+				"SlateCore",
+			}
+			);
+
+		if (Target.bBuildEditor == true)
+		{
+			PrivateDependencyModuleNames.AddRange(
+				new string[]
+				{
+						"AnimationEditor",
+						"AnimGraph",
+						"BlueprintGraph",
+						"EditorFramework",
+						"Kismet",
+						"UnrealEd",
+				}
+			);
+		}
+	}
+}
diff --git a/Engine/Plugins/Experimental/PhysicsControl/Source/PhysicsControlUncookedOnly/Private/AnimGraphNode_RigidBodyWithControl.cpp b/Engine/Plugins/Experimental/PhysicsControl/Source/PhysicsControlUncookedOnly/Private/AnimGraphNode_RigidBodyWithControl.cpp
new file mode 100644
--- /dev/null
+++ b/Engine/Plugins/Experimental/PhysicsControl/Source/PhysicsControlUncookedOnly/Private/AnimGraphNode_RigidBodyWithControl.cpp
@@ -0,0 +1,246 @@
+// Copyright Epic Games, Inc. All Rights Reserved.
+
+/* <<<< THIS IS TEMPORARY PROTOTYPE CODE >>>>
+ *
+ * This is essentially a copy + paste of the 
+ * equivalent file in engine code and is 
+ * intended for prototyping work as part of 
+ * the Ch5 locomotion initiative.
+ * 
+ */
+
+#include "AnimGraphNode_RigidBodyWithControl.h"
+#include "AnimNodeEditModes.h"
+#include "Features/IModularFeatures.h"
+#include "Kismet2/CompilerResultsLog.h"
+#include "AnimNode_RigidBodyWithControl.h"
+#include "EditorModeManager.h"
+#include "IPhysicsAssetRenderInterface.h"
+// #include "IAnimNodeEditMode.h" <- removed as it can't be found during build and doesn't appear to be necessary
+
+// Details includes
+#include "PropertyHandle.h"
+#include "DetailLayoutBuilder.h"
+#include "DetailWidgetRow.h"
+#include "DetailCategoryBuilder.h"
+#include "Widgets/Input/SButton.h"
+#include "Widgets/Text/STextBlock.h"
+
+
+#include "SceneManagement.h"
+
+/////////////////////////////////////////////////////
+// UAnimGraphNode_RigidBodyWithControl
+
+#define LOCTEXT_NAMESPACE "RigidBodyWithControl"
+
+UAnimGraphNode_RigidBodyWithControl::UAnimGraphNode_RigidBodyWithControl(const FObjectInitializer& ObjectInitializer)
+	: Super(ObjectInitializer)
+{
+}
+
+FText UAnimGraphNode_RigidBodyWithControl::GetControllerDescription() const
+{
+	return LOCTEXT("AnimGraphNode_RigidBodyWithControl_ControllerDescription", "Rigid body simulation with Control for physics asset");
+}
+
+FText UAnimGraphNode_RigidBodyWithControl::GetTooltipText() const
+{
+	return LOCTEXT("AnimGraphNode_RigidBodyWithControl_Tooltip", "This simulates based on the skeletal mesh component's physics asset with control options");
+}
+
+FText UAnimGraphNode_RigidBodyWithControl::GetNodeTitle(ENodeTitleType::Type TitleType) const
+{
+	return FText(LOCTEXT("AnimGraphNode_RigidBodyWithControl_NodeTitle", "RigidBodyWithControl"));
+}
+
+FLinearColor UAnimGraphNode_RigidBodyWithControl::GetNodeTitleColor() const
+{
+	return FLinearColor(1.0f, 0.0f, 1.0f); // <- Magenta - as a warning
+}
+
+FString UAnimGraphNode_RigidBodyWithControl::GetNodeCategory() const
+{
+	return TEXT("Animation|Dynamics");
+}
+
+void UAnimGraphNode_RigidBodyWithControl::ValidateAnimNodeDuringCompilation(USkeleton* ForSkeleton, FCompilerResultsLog& MessageLog)
+{
+	Super::ValidateAnimNodeDuringCompilation(ForSkeleton, MessageLog);
+}
+
+void UAnimGraphNode_RigidBodyWithControl::Draw(FPrimitiveDrawInterface* PDI, USkeletalMeshComponent* PreviewSkelMeshComp, const bool bIsSelected, const bool bIsPoseWatchEnabled) const
+{
+	if (const FAnimNode_RigidBodyWithControl* const RuntimeRigidBodyNode = GetDebuggedAnimNode<FAnimNode_RigidBodyWithControl>())
+	{
+		if (UPhysicsAsset* const PhysicsAsset = RuntimeRigidBodyNode->GetPhysicsAsset())
+		{
+			IPhysicsAssetRenderInterface& PhysicsAssetRenderInterface = IModularFeatures::Get().GetModularFeature<IPhysicsAssetRenderInterface>("PhysicsAssetRenderInterface");
+
+			// Draw Bodies.
+			if (bIsSelected || (bIsPoseWatchEnabled && PoseWatchElementBodies.IsValid() && PoseWatchElementBodies->GetIsVisible()))
+			{
+				FColor PrimitiveColorOverride = FColor::Transparent;
+
+				// Get primitive color from pose watch component.
+				if (!bIsSelected)
+				{
+					PrimitiveColorOverride = PoseWatchElementBodies->GetColor();
+					PrimitiveColorOverride.A = 255;
+				}
+
+				PhysicsAssetRenderInterface.DebugDrawBodies(PreviewSkelMeshComp, PhysicsAsset, PDI, PrimitiveColorOverride);
+			}
+
+			// Draw Constraints.
+			if (bIsSelected || (bIsPoseWatchEnabled && PoseWatchElementConstraints.IsValid() && PoseWatchElementConstraints->GetIsVisible()))
+			{
+				PhysicsAssetRenderInterface.DebugDrawConstraints(PreviewSkelMeshComp, PhysicsAsset, PDI);
+			}
+		}
+		
+		
+		// Draw Space Controls
+		if (bIsSelected || (bIsPoseWatchEnabled && PoseWatchElementWorldSpaceControls.IsValid() && PoseWatchElementWorldSpaceControls->GetIsVisible()))
+		{
+			for (const TMap<FName, FControlRecord>::ElementType& NameRecordPair : RuntimeRigidBodyNode->ControlRecords)
+			{
+				const FControlRecord& ControlRecord = NameRecordPair.Value;
+				if (!ControlRecord.JointHandle || !ControlRecord.IsEnabled())
+				{
+					continue;
+				}
+#if 0
+				// TODO get the targets from the control record and draw them with respect to the physical bodies
+				PDI->DrawPoint(ControlRootTransform.GetTranslation(), FLinearColor::Blue, 5.0f, SDPG_Foreground);
+				PDI->DrawLine(BodyPosition, TargetPosition, FLinearColor::Yellow, SDPG_Foreground);
+#endif
+			}
+		}
+	}
+}
+
+void UAnimGraphNode_RigidBodyWithControl::OnPoseWatchChanged(const bool IsPoseWatchEnabled, TObjectPtr<UPoseWatch> InPoseWatch, FEditorModeTools& InModeTools, FAnimNode_Base* InRuntimeNode)
+{
+	Super::OnPoseWatchChanged(IsPoseWatchEnabled, InPoseWatch, InModeTools, InRuntimeNode);
+
+	UPoseWatch* const PoseWatch = InPoseWatch.Get();
+
+	if (PoseWatch)
+	{
+		// A new pose watch has been created for this node - add node specific pose watch components.
+		PoseWatchElementBodies = InPoseWatch.Get()->FindOrAddElement(FText(LOCTEXT("PoseWatchElementLabel_RigidBodyWithControl_PhysicsBodies", "Physics Bodies")), TEXT("PhysicsAssetEditor.Tree.Body"));
+		PoseWatchElementConstraints = InPoseWatch.Get()->FindOrAddElement(FText(LOCTEXT("PoseWatchElementLabel__RigidBodyWithControl_PhysicsConstraints", "Physics Constraints")), TEXT("PhysicsAssetEditor.Tree.Constraint"));
+		PoseWatchElementParentSpaceControls = InPoseWatch.Get()->FindOrAddElement(FText(LOCTEXT("PoseWatchElementLabel_RigidBodyWithControl_ParentSpaceControls", "Parent Space Controls")), TEXT("PhysicsAssetEditor.Tree.Body"));
+		PoseWatchElementWorldSpaceControls = InPoseWatch.Get()->FindOrAddElement(FText(LOCTEXT("PoseWatchElementLabel_RigidBodyWithControl_WorldSpaceControls", "World Space Controls")), TEXT("PhysicsAssetEditor.Tree.Body"));
+
+		check(PoseWatchElementConstraints.IsValid()); // Expect to find a valid component;
+		if (PoseWatchElementConstraints.IsValid())
+		{
+			PoseWatchElementConstraints->SetHasColor(false);
+		}
+	}
+}
+
+void UAnimGraphNode_RigidBodyWithControl::CustomizeDetails(IDetailLayoutBuilder& DetailBuilder)
+{
+	Super::CustomizeDetails(DetailBuilder);
+
+	IDetailCategoryBuilder& ViewportCategory = DetailBuilder.EditCategory(TEXT("Debug Visualization"));
+	FDetailWidgetRow& WidgetRow = ViewportCategory.AddCustomRow(LOCTEXT("ToggleDebugVisualizationButtonRow", "DebugVisualization"));
+	FAnimNode_RigidBodyWithControl* const RigidBodyNode = static_cast<FAnimNode_RigidBodyWithControl*>(GetDebuggedAnimNode());
+
+	WidgetRow
+		[
+			SNew(SHorizontalBox)
+			// Show/Hide Bodies button.
+			+ SHorizontalBox::Slot()
+			[
+				SNew(SButton)
+				.HAlign(HAlign_Center)
+				.VAlign(VAlign_Center)
+				.OnClicked_Lambda([this](){ this->ToggleBodyVisibility(); return FReply::Handled(); })
+				.ButtonColorAndOpacity_Lambda([this](){ return (AreAnyBodiesHidden()) ? FAppStyle::Get().GetSlateColor("Colors.AccentRed") : FAppStyle::Get().GetSlateColor("Colors.AccentGreen"); })
+				.Content()
+				[
+					SNew(STextBlock)
+					.Text_Lambda([this]() { return (AreAnyBodiesHidden()) ? LOCTEXT("ShowAllBodiesButtonText", "Show All Bodies") : LOCTEXT("HideAllBodiesButtonText", "Hide All Bodies"); })
+					.ToolTipText(LOCTEXT("ToggleBodyVisibilityButtonToolTip", "Toggle debug visualization of all physics bodies"))
+				]
+			]
+			// Show/Hide Constraints button.
+			+ SHorizontalBox::Slot()
+			[
+				SNew(SButton)
+				.HAlign(HAlign_Center)
+				.VAlign(VAlign_Center)
+				.OnClicked_Lambda([this](){ this->ToggleConstraintVisibility(); return FReply::Handled(); })
+				.ButtonColorAndOpacity_Lambda([this]() { return (AreAnyConstraintsHidden()) ? FAppStyle::Get().GetSlateColor("Colors.AccentRed") : FAppStyle::Get().GetSlateColor("Colors.AccentGreen"); })
+				.Content()
+				[
+					SNew(STextBlock)
+					.Text_Lambda([this]() { return (AreAnyConstraintsHidden()) ? LOCTEXT("ShowAllConstraintsButtonText", "Show All Constraints") : LOCTEXT("HideAllConstraintsButtonText", "Hide All Constraints"); })
+					.ToolTipText(LOCTEXT("ToggleConstraintVisibilityButtonToolTip", "Toggle debug visualization of all physics constriants"))
+				]
+			]
+		];
+}
+
+void UAnimGraphNode_RigidBodyWithControl::PostEditChangeProperty(struct FPropertyChangedEvent& PropertyChangedEvent)
+{
+	Super::PostEditChangeProperty(PropertyChangedEvent);
+
+	IPhysicsAssetRenderInterface& PhysicsAssetRenderInterface = IModularFeatures::Get().GetModularFeature<IPhysicsAssetRenderInterface>("PhysicsAssetRenderInterface");
+	PhysicsAssetRenderInterface.SaveConfig();
+}
+
+void UAnimGraphNode_RigidBodyWithControl::ToggleBodyVisibility()
+{
+	FAnimNode_RigidBodyWithControl* const RigidBodyNode = static_cast<FAnimNode_RigidBodyWithControl*>(GetDebuggedAnimNode());
+	
+	IPhysicsAssetRenderInterface& PhysicsAssetRenderInterface = IModularFeatures::Get().GetModularFeature<IPhysicsAssetRenderInterface>("PhysicsAssetRenderInterface");
+
+	if (RigidBodyNode)
+	{
+		PhysicsAssetRenderInterface.ToggleShowAllBodies(RigidBodyNode->GetPhysicsAsset());
+	}
+}
+
+void UAnimGraphNode_RigidBodyWithControl::ToggleConstraintVisibility()
+{
+	FAnimNode_RigidBodyWithControl* const RigidBodyNode = static_cast<FAnimNode_RigidBodyWithControl*>(GetDebuggedAnimNode());
+	IPhysicsAssetRenderInterface& PhysicsAssetRenderInterface = IModularFeatures::Get().GetModularFeature<IPhysicsAssetRenderInterface>("PhysicsAssetRenderInterface");
+
+	if (RigidBodyNode)
+	{
+		PhysicsAssetRenderInterface.ToggleShowAllConstraints(RigidBodyNode->GetPhysicsAsset());
+	}
+}
+
+bool UAnimGraphNode_RigidBodyWithControl::AreAnyBodiesHidden() const
+{
+	FAnimNode_RigidBodyWithControl* const RigidBodyNode = static_cast<FAnimNode_RigidBodyWithControl*>(GetDebuggedAnimNode());
+	IPhysicsAssetRenderInterface& PhysicsAssetRenderInterface = IModularFeatures::Get().GetModularFeature<IPhysicsAssetRenderInterface>("PhysicsAssetRenderInterface");
+
+	if (RigidBodyNode)
+	{
+		return PhysicsAssetRenderInterface.AreAnyBodiesHidden(RigidBodyNode->GetPhysicsAsset());
+	}
+
+	return false;
+}
+
+bool UAnimGraphNode_RigidBodyWithControl::AreAnyConstraintsHidden() const
+{
+	FAnimNode_RigidBodyWithControl* const RigidBodyNode = static_cast<FAnimNode_RigidBodyWithControl*>(GetDebuggedAnimNode());
+	IPhysicsAssetRenderInterface& PhysicsAssetRenderInterface = IModularFeatures::Get().GetModularFeature<IPhysicsAssetRenderInterface>("PhysicsAssetRenderInterface");
+
+	if (RigidBodyNode)
+	{
+		return PhysicsAssetRenderInterface.AreAnyConstraintsHidden(RigidBodyNode->GetPhysicsAsset());
+	}
+
+	return false;
+}
+
+#undef LOCTEXT_NAMESPACE
diff --git a/Engine/Plugins/Experimental/PhysicsControl/Source/PhysicsControlUncookedOnly/Private/PhysicsControlUncookedOnlyModule.cpp b/Engine/Plugins/Experimental/PhysicsControl/Source/PhysicsControlUncookedOnly/Private/PhysicsControlUncookedOnlyModule.cpp
new file mode 100644
--- /dev/null
+++ b/Engine/Plugins/Experimental/PhysicsControl/Source/PhysicsControlUncookedOnly/Private/PhysicsControlUncookedOnlyModule.cpp
@@ -0,0 +1,5 @@
+// Copyright Epic Games, Inc. All Rights Reserved.
+
+#include "Modules/ModuleManager.h"
+
+IMPLEMENT_MODULE(FDefaultModuleImpl, PhysicsControlUncookedOnly)
diff --git a/Engine/Plugins/Experimental/PhysicsControl/Source/PhysicsControlUncookedOnly/Public/AnimGraphNode_RigidBodyWithControl.h b/Engine/Plugins/Experimental/PhysicsControl/Source/PhysicsControlUncookedOnly/Public/AnimGraphNode_RigidBodyWithControl.h
new file mode 100644
--- /dev/null
+++ b/Engine/Plugins/Experimental/PhysicsControl/Source/PhysicsControlUncookedOnly/Public/AnimGraphNode_RigidBodyWithControl.h
@@ -0,0 +1,60 @@
+// Copyright Epic Games, Inc. All Rights Reserved.
+
+/* <<<< THIS IS TEMPORARY PROTOTYPE CODE >>>>
+ *
+ * This is essentially a copy + paste of the
+ * equivalent file in engine code and is
+ * intended for prototyping work as part of
+ * the Ch5 locomotion initiative.
+ *
+ */
+
+#pragma once
+
+#include "AnimGraphNode_SkeletalControlBase.h"
+#include "AnimNode_RigidBodyWithControl.h"
+#include "Engine/PoseWatch.h"
+#include "AnimGraphNode_RigidBodyWithControl.generated.h"
+
+UCLASS(meta=(Keywords = "Simulate Rigid Body Physics Ragdoll"))
+class PHYSICSCONTROLUNCOOKEDONLY_API UAnimGraphNode_RigidBodyWithControl : public UAnimGraphNode_SkeletalControlBase
+{
+	GENERATED_UCLASS_BODY()
+
+	UPROPERTY(EditAnywhere, Category=Settings)
+	FAnimNode_RigidBodyWithControl Node;
+
+public:
+	// UEdGraphNode interface
+	virtual FText GetNodeTitle(ENodeTitleType::Type TitleType) const override;
+	virtual FLinearColor GetNodeTitleColor() const;
+	virtual FText GetTooltipText() const override;
+	virtual FString GetNodeCategory() const override;
+	virtual void ValidateAnimNodeDuringCompilation(USkeleton* ForSkeleton, FCompilerResultsLog& MessageLog) override;
+	// End of UEdGraphNode interface
+
+	// UAnimGraphNode_Base
+	virtual void Draw(FPrimitiveDrawInterface* PDI, USkeletalMeshComponent* PreviewSkelMeshComp, const bool bIsSelected, const bool bIsPoseWatchEnabled) const override;
+	virtual void OnPoseWatchChanged(const bool IsPoseWatchActive, TObjectPtr<UPoseWatch> InPoseWatch, FEditorModeTools& InModeTools, FAnimNode_Base* InRuntimeNode) override;
+
+	virtual void CustomizeDetails(IDetailLayoutBuilder& DetailBuilder) override;
+	virtual void PostEditChangeProperty(struct FPropertyChangedEvent& PropertyChangedEvent) override;
+
+	void ToggleBodyVisibility();
+	void ToggleConstraintVisibility();
+
+protected:
+	// UAnimGraphNode_SkeletalControlBase interface
+	virtual FText GetControllerDescription() const override;
+	virtual const FAnimNode_SkeletalControlBase* GetNode() const override { return &Node; }
+	// End of UAnimGraphNode_SkeletalControlBase interface
+
+	bool AreAnyBodiesHidden() const;
+	bool AreAnyConstraintsHidden() const;
+
+private:
+	TWeakObjectPtr<UPoseWatchElement> PoseWatchElementBodies;
+	TWeakObjectPtr<UPoseWatchElement> PoseWatchElementConstraints;
+	TWeakObjectPtr<UPoseWatchElement> PoseWatchElementParentSpaceControls;
+	TWeakObjectPtr<UPoseWatchElement> PoseWatchElementWorldSpaceControls;
+};
