[UE] Add profiling user markers to the RigidBodyWithControl node.

#rnx

diff --git a/Engine/Plugins/Experimental/PhysicsControl/Source/PhysicsControl/Private/AnimNode_RigidBodyWithControl.cpp b/Engine/Plugins/Experimental/PhysicsControl/Source/PhysicsControl/Private/AnimNode_RigidBodyWithControl.cpp
--- a/Engine/Plugins/Experimental/PhysicsControl/Source/PhysicsControl/Private/AnimNode_RigidBodyWithControl.cpp
+++ b/Engine/Plugins/Experimental/PhysicsControl/Source/PhysicsControl/Private/AnimNode_RigidBodyWithControl.cpp
@@ -40,6 +40,12 @@ LLM_DEFINE_TAG(Animation_RigidBodyWithControl);
 #define LOCTEXT_NAMESPACE "ImmediatePhysicsWithControl"
 
 //DEFINE_STAT(STAT_RigidBodyNodeWithControlInitTime);
+DECLARE_CYCLE_STAT(TEXT("RigidBodyNodeWithControl_Eval"), STAT_RigidBodyNodeWithControl_Eval, STATGROUP_Anim);
+DECLARE_CYCLE_STAT(TEXT("RigidBodyNodeWithControl_Simulation"), STAT_RigidBodyNodeWithControl_Simulation, STATGROUP_Anim);
+DECLARE_CYCLE_STAT(TEXT("RigidBodyNodeWithControl_SimulationWait"), STAT_RigidBodyNodeWithControl_SimulationWait, STATGROUP_Anim);
+DECLARE_CYCLE_STAT(TEXT("RigidBodyNodeWithControl_PreUpdate"), STAT_RigidBodyNodeWithControl_PreUpdate, STATGROUP_Anim);
+DECLARE_CYCLE_STAT(TEXT("RigidBodyNodeWithControl_Update"), STAT_RigidBodyNodeWithControl_Update, STATGROUP_Anim);
+DECLARE_CYCLE_STAT(TEXT("FAnimNode_RigidBodyWithControl::EvaluateSkeletalControl_AnyThread"), STAT_ImmediateEvaluateSkeletalControl, STATGROUP_ImmediatePhysics);
 
 CSV_DECLARE_CATEGORY_MODULE_EXTERN(ENGINE_API, Animation);
 
@@ -425,12 +431,6 @@ void FAnimNode_RigidBodyWithControl::CalculateSimulationSpace(
 	}
 }
 
-
-DECLARE_CYCLE_STAT(TEXT("RigidBody_Eval"), STAT_RigidBody_Eval, STATGROUP_Anim);
-DECLARE_CYCLE_STAT(TEXT("RigidBodyNodeWithControl_Simulation"), STAT_RigidBodyNodeWithControl_Simulation, STATGROUP_Anim);
-DECLARE_CYCLE_STAT(TEXT("RigidBodyNodeWithControl_SimulationWait"), STAT_RigidBodyNodeWithControl_SimulationWait, STATGROUP_Anim);
-DECLARE_CYCLE_STAT(TEXT("FAnimNode_RigidBodyWithControl::EvaluateSkeletalControl_AnyThread"), STAT_ImmediateEvaluateSkeletalControl, STATGROUP_ImmediatePhysics);
-
 void FAnimNode_RigidBodyWithControl::RunPhysicsSimulation(float DeltaSeconds, const FVector& SimSpaceGravity)
 {
 	SCOPE_CYCLE_COUNTER(STAT_RigidBodyNodeWithControl_Simulation);
@@ -493,7 +493,7 @@ static FTransform CalculateJointTargetTransform(
 void FAnimNode_RigidBodyWithControl::EvaluateSkeletalControl_AnyThread(FComponentSpacePoseContext& Output, TArray<FBoneTransform>& OutBoneTransforms)
 {
 	DECLARE_SCOPE_HIERARCHICAL_COUNTER_ANIMNODE(EvaluateSkeletalControl_AnyThread)
-	SCOPE_CYCLE_COUNTER(STAT_RigidBody_Eval);
+	SCOPE_CYCLE_COUNTER(STAT_RigidBodyNodeWithControl_Eval);
 	CSV_SCOPED_TIMING_STAT(Animation, RigidBodyEval);
 	SCOPE_CYCLE_COUNTER(STAT_ImmediateEvaluateSkeletalControl);
 	//SCOPED_NAMED_EVENT_TEXT("FAnimNode_RigidBodyWithControl::EvaluateSkeletalControl_AnyThread", FColor::Magenta);
@@ -1467,8 +1467,6 @@ void FAnimNode_RigidBodyWithControl::ResetDynamics(ETeleportType InTeleportType)
 	ResetSimulatedTeleportType = ((InTeleportType > ResetSimulatedTeleportType) ? InTeleportType : ResetSimulatedTeleportType);
 }
 
-DECLARE_CYCLE_STAT(TEXT("RigidBody_PreUpdate"), STAT_RigidBody_PreUpdate, STATGROUP_Anim);
-
 void FAnimNode_RigidBodyWithControl::PreUpdate(const UAnimInstance* InAnimInstance)
 {
 	// Don't update geometry if RBN is disabled
@@ -1477,7 +1475,7 @@ void FAnimNode_RigidBodyWithControl::PreUpdate(const UAnimInstance* InAnimInstan
 		return;
 	}
 
-	SCOPE_CYCLE_COUNTER(STAT_RigidBody_PreUpdate);
+	SCOPE_CYCLE_COUNTER(STAT_RigidBodyNodeWithControl_PreUpdate);
 
 	USkeletalMeshComponent* SKC = InAnimInstance->GetSkelMeshComponent();
 	APawn* PawnOwner = InAnimInstance->TryGetPawnOwner();
@@ -1549,8 +1547,6 @@ int32 FAnimNode_RigidBodyWithControl::GetLODThreshold() const
 	}
 }
 
-DECLARE_CYCLE_STAT(TEXT("RigidBody_Update"), STAT_RigidBody_Update, STATGROUP_Anim);
-
 void FAnimNode_RigidBodyWithControl::UpdateInternal(const FAnimationUpdateContext& Context)
 {
 	LLM_SCOPE_BYNAME(TEXT("Animation/RigidBodyWithControl")); 
@@ -1561,7 +1557,7 @@ void FAnimNode_RigidBodyWithControl::UpdateInternal(const FAnimationUpdateContex
 		return;
 	}
 
-	SCOPE_CYCLE_COUNTER(STAT_RigidBody_Update);
+	SCOPE_CYCLE_COUNTER(STAT_RigidBodyNodeWithControl_Update);
 	
 	// Must flush the simulation since we may be making changes to the scene
 	FlushDeferredSimulationTask();
diff --git a/Engine/Plugins/Experimental/PhysicsControl/Source/PhysicsControl/Private/RigidBodyWithControl.cpp b/Engine/Plugins/Experimental/PhysicsControl/Source/PhysicsControl/Private/RigidBodyWithControl.cpp
--- a/Engine/Plugins/Experimental/PhysicsControl/Source/PhysicsControl/Private/RigidBodyWithControl.cpp
+++ b/Engine/Plugins/Experimental/PhysicsControl/Source/PhysicsControl/Private/RigidBodyWithControl.cpp
@@ -17,6 +17,14 @@
 
 //UE_DISABLE_OPTIMIZATION;
 
+DECLARE_CYCLE_STAT(TEXT("RigidBodyNodeWithControl_InitControlsAndBodyModifiers"), STAT_RigidBodyNodeWithControl_InitControlsAndBodyModifiers, STATGROUP_Anim);
+DECLARE_CYCLE_STAT(TEXT("RigidBodyNodeWithControl_LogControlsModifiersAndSets"), STAT_RigidBodyNodeWithControl_LogControlsModifiersAndSets, STATGROUP_Anim);
+DECLARE_CYCLE_STAT(TEXT("RigidBodyNodeWithControl_ApplyControlAndModifierUpdatesAndParametersToRecords"), STAT_RigidBodyNodeWithControl_ApplyControlAndModifierUpdatesAndParametersToRecords, STATGROUP_Anim);
+DECLARE_CYCLE_STAT(TEXT("RigidBodyNodeWithControl_ApplyControlsAndModifierDatas"), STAT_RigidBodyNodeWithControl_ApplyControlsAndModifierDatas, STATGROUP_Anim);
+DECLARE_CYCLE_STAT(TEXT("RigidBodyNodeWithControl_ApplyControlsAndModifiers"), STAT_RigidBodyNodeWithControl_ApplyControlsAndModifiers, STATGROUP_Anim);
+DECLARE_CYCLE_STAT(TEXT("RigidBodyNodeWithControl_ApplyKinematicTargets"), STAT_RigidBodyNodeWithControl_ApplyKinematicTargets, STATGROUP_Anim);
+DECLARE_CYCLE_STAT(TEXT("RigidBodyNodeWithControl_ApplyCurrentConstraintProfile"), STAT_RigidBodyNodeWithControl_ApplyCurrentConstraintProfile, STATGROUP_Anim);
+
 int32 MaxNumControlsOrModifiersPerName = 16;
 
 constexpr int32 ConstraintChildIndex = 0;
@@ -494,6 +502,8 @@ TMap<FName, FRigidBodyLimbBones> FAnimNode_RigidBodyWithControl::GetLimbBonesFro
 //======================================================================================================================
 void FAnimNode_RigidBodyWithControl::InitControlsAndBodyModifiers(USkeletalMeshComponent* const SkeletalMeshComponent)
 {
+	SCOPE_CYCLE_COUNTER(STAT_RigidBodyNodeWithControl_InitControlsAndBodyModifiers);
+
 	check(ControlRecords.IsEmpty()); // Controls should not exist when this function is called.
 
 	// These functions will create the base set of controls and modifiers from SetupData
@@ -551,6 +561,8 @@ void FAnimNode_RigidBodyWithControl::DestroyControlsAndBodyModifiers()
 // TODO This isn't ideal as it only dumps out the "original" values, not including the updates
 void FAnimNode_RigidBodyWithControl::LogControlsModifiersAndSets()
 {
+	SCOPE_CYCLE_COUNTER(STAT_RigidBodyNodeWithControl_LogControlsModifiersAndSets);
+
 #define RBWC_LOG_LEVEL Display
 
 	UE_LOG(LogRigidBodyWithControl, RBWC_LOG_LEVEL, TEXT("Controls:"));
@@ -874,6 +886,8 @@ void FAnimNode_RigidBodyWithControl::ApplyControlAndModifierUpdatesAndParameters
 	const FRigidBodyControlAndModifierUpdates&    Updates,
 	const FRigidBodyControlAndModifierParameters& Parameters)
 {
+	SCOPE_CYCLE_COUNTER(STAT_RigidBodyNodeWithControl_ApplyControlAndModifierUpdatesAndParametersToRecords);
+
 	// Apply control and modifier parameters on a single-use basis
 	ApplyControlAndBodyModifierDatas(
 		Updates.ControlParameters, Updates.ModifierParameters);
@@ -911,6 +925,8 @@ void FAnimNode_RigidBodyWithControl::ApplyControlAndBodyModifierDatas(
 	const TArray<FRigidBodyNamedControlParameters>& InControlParameters,
 	const TArray<FRigidBodyNamedModifierParameters>& InModifierParameters)
 {
+	SCOPE_CYCLE_COUNTER(STAT_RigidBodyNodeWithControl_ApplyControlsAndModifierDatas);
+
 	// This updates the "original" controls and modifiers based on the parameters.
 	for (const FRigidBodyNamedControlParameters& ControlParameters : InControlParameters)
 	{
@@ -952,6 +968,8 @@ void FAnimNode_RigidBodyWithControl::ApplyControlAndBodyModifierDatas(
 //======================================================================================================================
 void FAnimNode_RigidBodyWithControl::ApplyControlsAndModifiers(const FVector& SimSpaceGravity, float DeltaTime)
 {
+	SCOPE_CYCLE_COUNTER(STAT_RigidBodyNodeWithControl_ApplyControlsAndModifiers);
+
 	// If we've skipped a frame then we need to avoid doing any velocity calculations. Simplest
 	// method is to set DeltaTime to zero.
 	{
@@ -995,6 +1013,8 @@ void FAnimNode_RigidBodyWithControl::ApplyControlsAndModifiers(const FVector& Si
 // Note that this will be called AFTER normal kinematic targets have been set
 void FAnimNode_RigidBodyWithControl::ApplyKinematicTargets()
 {
+	SCOPE_CYCLE_COUNTER(STAT_RigidBodyNodeWithControl_ApplyKinematicTargets);
+
 	if (!PoseData.IsEmpty())
 	{
 		// Apply any kinematic targets.
@@ -1063,6 +1083,9 @@ static Chaos::EPlasticityType ConvertPlasticityType(EConstraintPlasticityType In
 void FAnimNode_RigidBodyWithControl::ApplyCurrentConstraintProfile()
 {
 	using namespace Chaos;
+
+	SCOPE_CYCLE_COUNTER(STAT_RigidBodyNodeWithControl_ApplyCurrentConstraintProfile);
+
 	// Go through each joint (in the ragdoll that's been created) in turn...
 	for (int32 JointIndex = 0; JointIndex != Joints.Num(); ++JointIndex)
 	{
