Slightly more info in PhysicsControl warnings.

diff --git a/Engine/Plugins/Experimental/PhysicsControl/Source/PhysicsControl/Private/PhysicsControlComponent.cpp b/Engine/Plugins/Experimental/PhysicsControl/Source/PhysicsControl/Private/PhysicsControlComponent.cpp
--- a/Engine/Plugins/Experimental/PhysicsControl/Source/PhysicsControl/Private/PhysicsControlComponent.cpp
+++ b/Engine/Plugins/Experimental/PhysicsControl/Source/PhysicsControl/Private/PhysicsControlComponent.cpp
@@ -112,11 +112,12 @@ void UPhysicsControlComponent::UpdateControls(float DeltaTime)
 	for (TPair<FName, FPhysicsControlRecord>& RecordPair : Implementation->PhysicsControlRecords)
 	{
 		// New constraint requested when one doesn't exist
+		FName ControlName = RecordPair.Key;
 		FPhysicsControlRecord& Record = RecordPair.Value;
 		FConstraintInstance* ConstraintInstance = Record.PhysicsControlState.ConstraintInstance.Get();
 		if (!ConstraintInstance)
 		{
-			ConstraintInstance = Record.CreateConstraint(this);
+			ConstraintInstance = Record.CreateConstraint(this, ControlName);
 		}
 
 		if (ConstraintInstance)
diff --git a/Engine/Plugins/Experimental/PhysicsControl/Source/PhysicsControl/Private/PhysicsControlRecord.cpp b/Engine/Plugins/Experimental/PhysicsControl/Source/PhysicsControl/Private/PhysicsControlRecord.cpp
--- a/Engine/Plugins/Experimental/PhysicsControl/Source/PhysicsControl/Private/PhysicsControlRecord.cpp
+++ b/Engine/Plugins/Experimental/PhysicsControl/Source/PhysicsControl/Private/PhysicsControlRecord.cpp
@@ -20,7 +20,7 @@ void FPhysicsControlState::Reset()
 //======================================================================================================================
 // A constraint created with identity frames will just have the child frame at the mesh origin,
 // which is not necessarily where the center of gravity is.
-FConstraintInstance* FPhysicsControlRecord::CreateConstraint(UObject* ConstraintDebugOwner)
+FConstraintInstance* FPhysicsControlRecord::CreateConstraint(UObject* ConstraintDebugOwner, FName ControlName)
 {
 	if (!PhysicsControlState.ConstraintInstance)
 	{
@@ -35,12 +35,18 @@ FConstraintInstance* FPhysicsControlRecord::CreateConstraint(UObject* Constraint
 
 	if (PhysicsControl.ParentMeshComponent && !PhysicsControl.ParentBoneName.IsNone() && !ParentBody)
 	{
-		UE_LOG(LogPhysicsControlComponent, Warning, TEXT("Failed to find expected parent body when making constraint"));
+		UE_LOG(LogPhysicsControlComponent, Warning, 
+			TEXT("Failed to find expected parent body %s when making constraint for control %s"), 
+			*PhysicsControl.ParentBoneName.ToString(),
+			*ControlName.ToString());
 		return nullptr;
 	}
 	if (PhysicsControl.ChildMeshComponent && !PhysicsControl.ChildBoneName.IsNone() && !ChildBody)
 	{
-		UE_LOG(LogPhysicsControlComponent, Warning, TEXT("Failed to find expected child body when making constraint"));
+		UE_LOG(LogPhysicsControlComponent, Warning, 
+			TEXT("Failed to find expected child body %s when making constraint for control %s"), 
+			*PhysicsControl.ChildBoneName.ToString(),
+			*ControlName.ToString());
 		return nullptr;
 	}
 
diff --git a/Engine/Plugins/Experimental/PhysicsControl/Source/PhysicsControl/Private/PhysicsControlRecord.h b/Engine/Plugins/Experimental/PhysicsControl/Source/PhysicsControl/Private/PhysicsControlRecord.h
--- a/Engine/Plugins/Experimental/PhysicsControl/Source/PhysicsControl/Private/PhysicsControlRecord.h
+++ b/Engine/Plugins/Experimental/PhysicsControl/Source/PhysicsControl/Private/PhysicsControlRecord.h
@@ -36,7 +36,7 @@ struct FPhysicsControlRecord
 	 * Creates the constraint (if necessary) and stores it in the state. ConstraintDebugOwner is passed 
 	 * to the constraint on creation. 
 	 */
-	FConstraintInstance* CreateConstraint(UObject* ConstraintDebugOwner);
+	FConstraintInstance* CreateConstraint(UObject* ConstraintDebugOwner, FName ControlName);
 
 	/** Ensures the constraint frame matches the control point in the record. */
 	void UpdateConstraintControlPoint();
