Aligning the structures for PhysicsControlComponent and RigidBodyWithControl
Removes support for auto-deactivation of physics control component

diff --git a/Engine/Plugins/Experimental/PhysicsControl/Source/PhysicsControl/Private/PhysicsControlActor.cpp b/Engine/Plugins/Experimental/PhysicsControl/Source/PhysicsControl/Private/PhysicsControlActor.cpp
--- a/Engine/Plugins/Experimental/PhysicsControl/Source/PhysicsControl/Private/PhysicsControlActor.cpp
+++ b/Engine/Plugins/Experimental/PhysicsControl/Source/PhysicsControl/Private/PhysicsControlActor.cpp
@@ -54,10 +54,8 @@ void UPhysicsControlInitializerComponent::CreateInitialCharacterControls(UPhysic
 			InitialCharacterControls.LimbSetupData,
 			InitialCharacterControls.WorldSpaceControlData,
 			InitialCharacterControls.WorldSpaceControlSettings,
-			InitialCharacterControls.bEnableWorldSpaceControls,
 			InitialCharacterControls.ParentSpaceControlData,
 			InitialCharacterControls.ParentSpaceControlSettings,
-			InitialCharacterControls.bEnableParentSpaceControls,
 			InitialCharacterControls.PhysicsMovementType
 		);
 	}
@@ -130,8 +128,7 @@ void UPhysicsControlInitializerComponent::CreateOrUpdateInitialControls(UPhysics
 			InitialPhysicsControl.ControlData,
 			InitialPhysicsControl.ControlTarget,
 			InitialPhysicsControl.ControlSettings,
-			"All",
-			true))
+			"All"))
 		{
 			// If we were replacing a control then add it to the original sets
 			FPhysicsControlNames Names;
diff --git a/Engine/Plugins/Experimental/PhysicsControl/Source/PhysicsControl/Private/PhysicsControlBPLibrary.cpp b/Engine/Plugins/Experimental/PhysicsControl/Source/PhysicsControl/Private/PhysicsControlBPLibrary.cpp
--- a/Engine/Plugins/Experimental/PhysicsControl/Source/PhysicsControl/Private/PhysicsControlBPLibrary.cpp
+++ b/Engine/Plugins/Experimental/PhysicsControl/Source/PhysicsControl/Private/PhysicsControlBPLibrary.cpp
@@ -73,11 +73,11 @@ template<typename TNamedParameters> void InterpolateParametersContainers(
 
 //======================================================================================================================
 template<typename TNamedParameters> void BlendParametersThroughSet(
-	const FRigidBodyControlAndModifierParameters& InParametersContainer, 
-	const TNamedParameters&                       InStartParameters,
-	const TNamedParameters&                       InEndParameters,
-	const TArray<FName>&                          Names,
-	FRigidBodyControlAndModifierParameters&       OutParametersContainer)
+	const FPhysicsControlControlAndModifierParameters& InParametersContainer, 
+	const TNamedParameters&                            InStartParameters,
+	const TNamedParameters&                            InEndParameters,
+	const TArray<FName>&                               Names,
+	FPhysicsControlControlAndModifierParameters&       OutParametersContainer)
 {
 	// TODO - Limbs can include branches. Would need a more sophisticated approach to deal with
 	// those properly, perhaps include a depth index in the control name ?
@@ -97,64 +97,64 @@ template<typename TNamedParameters> void BlendParametersThroughSet(
 
 //======================================================================================================================
 void UPhysicsControlBPLibrary::AddControlParameters(
-	UPARAM(ref) FRigidBodyControlAndModifierParameters& InParametersContainer,
-	FRigidBodyControlAndModifierParameters&             OutParametersContainer,
-	const FName                                         Name, 
-	const FRigidBodyControlSparseData&                  ControlData)
+	UPARAM(ref) FPhysicsControlControlAndModifierParameters& InParametersContainer,
+	FPhysicsControlControlAndModifierParameters&             OutParametersContainer,
+	const FName                                              Name, 
+	const FPhysicsControlSparseData&                         ControlData)
 {
 	OutParametersContainer = InParametersContainer;
-	OutParametersContainer.Add(FRigidBodyNamedControlParameters(Name, ControlData));
+	OutParametersContainer.Add(FPhysicsControlNamedControlParameters(Name, ControlData));
 }
 
 //======================================================================================================================
 void UPhysicsControlBPLibrary::AddMultipleControlParameters(
-	UPARAM(ref) FRigidBodyControlAndModifierParameters& InParametersContainer,
-	FRigidBodyControlAndModifierParameters&             OutParametersContainer,
-	const TArray<FName>&                                Names,
-	const FRigidBodyControlSparseData&                  ControlData)
+	UPARAM(ref) FPhysicsControlControlAndModifierParameters& InParametersContainer,
+	FPhysicsControlControlAndModifierParameters&             OutParametersContainer,
+	const TArray<FName>&                                     Names,
+	const FPhysicsControlSparseData&                         ControlData)
 {
 	OutParametersContainer = InParametersContainer;
 	OutParametersContainer.ControlParameters.Reserve(OutParametersContainer.ControlParameters.Num() + Names.Num());
 
 	for (const FName Name : Names)
 	{
-		OutParametersContainer.Add(FRigidBodyNamedControlParameters(Name, ControlData));
+		OutParametersContainer.Add(FPhysicsControlNamedControlParameters(Name, ControlData));
 	}	
 }
 
 //======================================================================================================================
 void UPhysicsControlBPLibrary::AddModifierParameters(
-	UPARAM(ref) FRigidBodyControlAndModifierParameters& InParametersContainer,
-	FRigidBodyControlAndModifierParameters&             OutParametersContainer,
-	const FName                                         Name,
-	const FRigidBodyModifierSparseData&                 ModifierData)
+	UPARAM(ref) FPhysicsControlControlAndModifierParameters& InParametersContainer,
+	FPhysicsControlControlAndModifierParameters&             OutParametersContainer,
+	const FName                                              Name,
+	const FPhysicsControlModifierSparseData&                 ModifierData)
 {
 	OutParametersContainer = InParametersContainer;
-	OutParametersContainer.Add(FRigidBodyNamedModifierParameters(Name, ModifierData));
+	OutParametersContainer.Add(FPhysicsControlNamedModifierParameters(Name, ModifierData));
 }
 
 //======================================================================================================================
 void UPhysicsControlBPLibrary::AddMultipleModifierParameters(
-	UPARAM(ref) FRigidBodyControlAndModifierParameters& InParametersContainer,
-	FRigidBodyControlAndModifierParameters&             OutParametersContainer,
-	const TArray<FName>&                                Names,
-	const FRigidBodyModifierSparseData&                 ModifierData)
+	UPARAM(ref) FPhysicsControlControlAndModifierParameters& InParametersContainer,
+	FPhysicsControlControlAndModifierParameters&             OutParametersContainer,
+	const TArray<FName>&                                     Names,
+	const FPhysicsControlModifierSparseData&                 ModifierData)
 {
 	OutParametersContainer = InParametersContainer;
 	OutParametersContainer.ModifierParameters.Reserve(OutParametersContainer.ModifierParameters.Num() + Names.Num());
 
 	for (const FName Name : Names)
 	{
-		OutParametersContainer.Add(FRigidBodyNamedModifierParameters(Name, ModifierData));
+		OutParametersContainer.Add(FPhysicsControlNamedModifierParameters(Name, ModifierData));
 	}
 }
 
 //======================================================================================================================
 void UPhysicsControlBPLibrary::BlendParameters(
-	UPARAM(ref) FRigidBodyControlAndModifierParameters& InParametersContainerA, 
-	UPARAM(ref) FRigidBodyControlAndModifierParameters& InParametersContainerB, 
-	const float                                         InInterpolationWeight, 
-	FRigidBodyControlAndModifierParameters&             OutParametersContainers)
+	UPARAM(ref) FPhysicsControlControlAndModifierParameters& InParametersContainerA, 
+	UPARAM(ref) FPhysicsControlControlAndModifierParameters& InParametersContainerB, 
+	const float                                              InInterpolationWeight, 
+	FPhysicsControlControlAndModifierParameters&             OutParametersContainers)
 {
 	OutParametersContainers = InParametersContainerA;
 
@@ -166,22 +166,22 @@ void UPhysicsControlBPLibrary::BlendParameters(
 
 //======================================================================================================================
 void UPhysicsControlBPLibrary::BlendControlParametersThroughSet(
-	UPARAM(ref) FRigidBodyControlAndModifierParameters& InParametersContainer,
-	UPARAM(ref) const FRigidBodyNamedControlParameters& InStartParameters,
-	UPARAM(ref) const FRigidBodyNamedControlParameters& InEndParameters,
-	const TArray<FName>&                                InNames,
-	FRigidBodyControlAndModifierParameters&             OutParametersContainer)
+	UPARAM(ref) FPhysicsControlControlAndModifierParameters& InParametersContainer,
+	UPARAM(ref) const FPhysicsControlNamedControlParameters& InStartParameters,
+	UPARAM(ref) const FPhysicsControlNamedControlParameters& InEndParameters,
+	const TArray<FName>&                                     InNames,
+	FPhysicsControlControlAndModifierParameters&             OutParametersContainer)
 {
 	BlendParametersThroughSet(InParametersContainer, InStartParameters, InEndParameters, InNames, OutParametersContainer);
 }
 
 //======================================================================================================================
 void UPhysicsControlBPLibrary::BlendModifierParametersThroughSet(
-	UPARAM(ref) FRigidBodyControlAndModifierParameters&  InParametersContainer, 
-	UPARAM(ref) const FRigidBodyNamedModifierParameters& InStartParameters,
-	UPARAM(ref) const FRigidBodyNamedModifierParameters& InEndParameters,
-	const TArray<FName>&                                 InNames,
-	FRigidBodyControlAndModifierParameters&              OutParametersContainer)
+	UPARAM(ref) FPhysicsControlControlAndModifierParameters&  InParametersContainer, 
+	UPARAM(ref) const FPhysicsControlNamedModifierParameters& InStartParameters,
+	UPARAM(ref) const FPhysicsControlNamedModifierParameters& InEndParameters,
+	const TArray<FName>&                                      InNames,
+	FPhysicsControlControlAndModifierParameters&              OutParametersContainer)
 {
 	BlendParametersThroughSet(InParametersContainer, InStartParameters, InEndParameters, InNames, OutParametersContainer);
 }
diff --git a/Engine/Plugins/Experimental/PhysicsControl/Source/PhysicsControl/Private/PhysicsControlComponent.cpp b/Engine/Plugins/Experimental/PhysicsControl/Source/PhysicsControl/Private/PhysicsControlComponent.cpp
--- a/Engine/Plugins/Experimental/PhysicsControl/Source/PhysicsControl/Private/PhysicsControlComponent.cpp
+++ b/Engine/Plugins/Experimental/PhysicsControl/Source/PhysicsControl/Private/PhysicsControlComponent.cpp
@@ -127,7 +127,7 @@ void UPhysicsControlComponent::UpdateControls(float DeltaTime)
 			ConstraintInstance->SetDisableCollision(Record.PhysicsControl.ControlSettings.bDisableCollision);
 
 			// Constraint is not enabled
-			if (!Record.PhysicsControlState.bEnabled)
+			if (!Record.Enabled())
 			{
 				ConstraintInstance->SetAngularDriveParams(0.0f, 0.0f, 0.0f);
 				ConstraintInstance->SetLinearDriveParams(0.0f, 0.0f, 0.0f);
@@ -195,16 +195,6 @@ void UPhysicsControlComponent::UpdateControls(float DeltaTime)
 			BodyInstance->AddForce(FVector(0, 0, GravityZToApply), true, true);
 		}
 	}
-
-	// Go through and de-activate any records if they're set to auto disable. 
-	for (TPair<FName, FPhysicsControlRecord>& RecordPair : Implementation->PhysicsControlRecords)
-	{
-		FPhysicsControlRecord& Record = RecordPair.Value;
-		if (Record.PhysicsControl.ControlSettings.bAutoDisable)
-		{
-			Record.PhysicsControlState.bEnabled = false;
-		}
-	}
 }
 
 //======================================================================================================================
@@ -250,6 +240,9 @@ TMap<FName, FPhysicsControlLimbBones> UPhysicsControlComponent::GetLimbBonesFrom
 	{
 		FPhysicsControlLimbBones& LimbBones = Result.Add(LimbSetupData.LimbName);
 		LimbBones.SkeletalMeshComponent = SkeletalMeshComponent;
+		LimbBones.bCreateBodyModifiers = LimbSetupData.bCreateBodyModifiers;
+		LimbBones.bCreateWorldSpaceControls = LimbSetupData.bCreateWorldSpaceControls;
+		LimbBones.bCreateParentSpaceControls = LimbSetupData.bCreateParentSpaceControls;
 
 		if (LimbSetupData.bIncludeParentBone)
 		{
@@ -295,13 +288,12 @@ FName UPhysicsControlComponent::CreateControl(
 	const FPhysicsControlData     ControlData, 
 	const FPhysicsControlTarget   ControlTarget, 
 	const FPhysicsControlSettings ControlSettings,
-	const FName                   Set,
-	const bool                    bEnabled)
+	const FName                   Set)
 {
 	const FName Name = Implementation->GetUniqueControlName(ParentBoneName, ChildBoneName);
 	if (CreateNamedControl(
 		Name, ParentMeshComponent, ParentBoneName, ChildMeshComponent, ChildBoneName, 
-		ControlData, ControlTarget, ControlSettings, Set, bEnabled))
+		ControlData, ControlTarget, ControlSettings, Set))
 	{
 		return Name;
 	}
@@ -319,8 +311,7 @@ bool UPhysicsControlComponent::CreateNamedControl(
 	const FPhysicsControlData     ControlData, 
 	const FPhysicsControlTarget   ControlTarget, 
 	const FPhysicsControlSettings ControlSettings,
-	const FName                   Set,
-	const bool                    bEnabled)
+	const FName                   Set)
 {
 	if (Implementation->FindControlRecord(Name))
 	{
@@ -347,7 +338,6 @@ bool UPhysicsControlComponent::CreateNamedControl(
 		Name, FPhysicsControl(
 			ParentMeshComponent, ParentBoneName, ChildMeshComponent, ChildBoneName,
 			ControlData, ControlTarget, ControlSettings));
-	NewRecord.PhysicsControlState.bEnabled = bEnabled;
 	NewRecord.ResetControlPoint();
 
 	Implementation->NameRecords.AddControl(Name, Set);
@@ -364,8 +354,7 @@ TArray<FName> UPhysicsControlComponent::CreateControlsFromSkeletalMeshBelow(
 	const EPhysicsControlType     ControlType,
 	const FPhysicsControlData     ControlData,
 	const FPhysicsControlSettings ControlSettings,
-	const FName                   Set,
-	const bool                    bEnabled)
+	const FName                   Set)
 {
 	TArray<FName> Result;
 	UPhysicsAsset* PhysicsAsset = SkeletalMeshComponent ? SkeletalMeshComponent->GetPhysicsAsset() : nullptr;
@@ -383,7 +372,7 @@ TArray<FName> UPhysicsControlComponent::CreateControlsFromSkeletalMeshBelow(
 		BoneName, bIncludeSelf, /*bSkipCustomType=*/false, 
 		[
 			this, PhysicsAsset, ParentMeshComponent, SkeletalMeshComponent, 
-			ControlType, &ControlData, &ControlSettings, Set, &Result, bEnabled
+			ControlType, &ControlData, &ControlSettings, Set, &Result
 		](const FBodyInstance* BI)
 		{
 			if (USkeletalBodySetup* BodySetup = Cast<USkeletalBodySetup>(BI->BodySetup.Get()))
@@ -403,8 +392,7 @@ TArray<FName> UPhysicsControlComponent::CreateControlsFromSkeletalMeshBelow(
 				const FName ControlName = CreateControl(
 					ParentMeshComponent, ParentBoneName, SkeletalMeshComponent, ChildBoneName,
 					ControlData, FPhysicsControlTarget(), ControlSettings, 
-					FName(GetControlTypeName(ControlType).ToString().Append("_").Append(Set.ToString())), 
-					bEnabled);
+					FName(GetControlTypeName(ControlType).ToString().Append("_").Append(Set.ToString())));
 				if (!ControlName.IsNone())
 				{
 					Result.Add(ControlName);
@@ -474,12 +462,12 @@ TArray<FName> UPhysicsControlComponent::CreateControlsFromSkeletalMeshAndConstra
 				}
 
 				UE::PhysicsControlComponent::ConvertConstraintProfileToControlData(ControlData, ProfileProperties);
+				ControlData.bEnabled = bEnabled;
 
 				const FName ControlName = CreateControl(
 					SkeletalMeshComponent, ParentBoneName, SkeletalMeshComponent, ChildBoneName,
 					ControlData, FPhysicsControlTarget(), ControlSettings, 
-					FName(GetControlTypeName(EPhysicsControlType::ParentSpace).ToString().Append("_").Append(Set.ToString())),
-					bEnabled);
+					FName(GetControlTypeName(EPhysicsControlType::ParentSpace).ToString().Append("_").Append(Set.ToString())));
 				if (!ControlName.IsNone())
 				{
 					Result.Add(ControlName);
@@ -505,8 +493,7 @@ TArray<FName> UPhysicsControlComponent::CreateControlsFromSkeletalMesh(
 	const EPhysicsControlType     ControlType,
 	const FPhysicsControlData     ControlData,
 	const FPhysicsControlSettings ControlSettings,
-	const FName                   Set,
-	const bool                    bEnabled)
+	const FName                   Set)
 {
 	TArray<FName> Result;
 	UPhysicsAsset* PhysicsAsset = SkeletalMeshComponent ? SkeletalMeshComponent->GetPhysicsAsset() : nullptr;
@@ -536,8 +523,7 @@ TArray<FName> UPhysicsControlComponent::CreateControlsFromSkeletalMesh(
 		const FName ControlName = CreateControl(
 			ParentMeshComponent, ParentBoneName, SkeletalMeshComponent, ChildBoneName,
 			ControlData, FPhysicsControlTarget(), ControlSettings, 
-			FName(GetControlTypeName(ControlType).ToString().Append("_").Append(Set.ToString())),
-			bEnabled);
+			FName(GetControlTypeName(ControlType).ToString().Append("_").Append(Set.ToString())));
 		if (!ControlName.IsNone())
 		{
 			Result.Add(ControlName);
@@ -596,12 +582,12 @@ TArray<FName> UPhysicsControlComponent::CreateControlsFromSkeletalMeshAndConstra
 		}
 
 		UE::PhysicsControlComponent::ConvertConstraintProfileToControlData(ControlData, ProfileProperties);
+		ControlData.bEnabled = bEnabled;
 
 		const FName ControlName = CreateControl(
 			SkeletalMeshComponent, ParentBoneName, SkeletalMeshComponent, ChildBoneName,
 			ControlData, FPhysicsControlTarget(), ControlSettings, 
-			FName(GetControlTypeName(EPhysicsControlType::ParentSpace).ToString().Append("_").Append(Set.ToString())),
-			bEnabled);
+			FName(GetControlTypeName(EPhysicsControlType::ParentSpace).ToString().Append("_").Append(Set.ToString())));
 		if (!ControlName.IsNone())
 		{
 			Result.Add(ControlName);
@@ -623,8 +609,7 @@ TMap<FName, FPhysicsControlNames> UPhysicsControlComponent::CreateControlsFromLi
 	const TMap<FName, FPhysicsControlLimbBones>& LimbBones,
 	const EPhysicsControlType                    ControlType,
 	const FPhysicsControlData                    ControlData,
-	const FPhysicsControlSettings                ControlSettings,
-	const bool                                   bEnabled)
+	const FPhysicsControlSettings                ControlSettings)
 {
 	TMap<FName, FPhysicsControlNames> Result;
 	Result.Reserve(LimbBones.Num());
@@ -640,6 +625,12 @@ TMap<FName, FPhysicsControlNames> UPhysicsControlComponent::CreateControlsFromLi
 			continue;
 		}
 
+		if ((ControlType == EPhysicsControlType::WorldSpace && !BonesInLimb.bCreateWorldSpaceControls) ||
+			(ControlType == EPhysicsControlType::ParentSpace && !BonesInLimb.bCreateParentSpaceControls))
+		{
+			continue;
+		}
+
 		USkeletalMeshComponent* ParentMeshComponent =
 			(ControlType == EPhysicsControlType::ParentSpace) ? BonesInLimb.SkeletalMeshComponent : nullptr;
 
@@ -673,8 +664,7 @@ TMap<FName, FPhysicsControlNames> UPhysicsControlComponent::CreateControlsFromLi
 			const FName ControlName = CreateControl(
 				ParentMeshComponent, ParentBoneName, BonesInLimb.SkeletalMeshComponent, ChildBoneName,
 				ControlData, FPhysicsControlTarget(), ControlSettings, 
-				FName(GetControlTypeName(ControlType).ToString().Append("_").Append(LimbName.ToString())),
-				bEnabled);
+				FName(GetControlTypeName(ControlType).ToString().Append("_").Append(LimbName.ToString())));
 			if (!ControlName.IsNone())
 			{
 				LimbResult.Names.Add(ControlName);
@@ -758,12 +748,12 @@ TMap<FName, FPhysicsControlNames> UPhysicsControlComponent::CreateControlsFromLi
 			}
 
 			UE::PhysicsControlComponent::ConvertConstraintProfileToControlData(ControlData, ProfileProperties);
+			ControlData.bEnabled = bEnabled;
 
 			const FName ControlName = CreateControl(
 				SkeletalMeshComponent, ParentBoneName, SkeletalMeshComponent, ChildBoneName,
 				ControlData, FPhysicsControlTarget(), ControlSettings, 
-				FName(GetControlTypeName(EPhysicsControlType::ParentSpace).ToString().Append("_").Append(LimbName.ToString())),
-				bEnabled);
+				FName(GetControlTypeName(EPhysicsControlType::ParentSpace).ToString().Append("_").Append(LimbName.ToString())));
 			if (!ControlName.IsNone())
 			{
 				LimbResult.Names.Add(ControlName);
@@ -809,7 +799,7 @@ bool UPhysicsControlComponent::SetControlEnabled(const FName Name, const bool bE
 	FPhysicsControlRecord* Record = Implementation->FindControlRecord(Name);
 	if (Record)
 	{
-		Record->PhysicsControlState.bEnabled = bEnable;
+		Record->PhysicsControl.ControlData.bEnabled = bEnable;
 		return true;
 	}
 	return false;
@@ -834,17 +824,12 @@ void UPhysicsControlComponent::SetControlsInSetEnabled(FName SetName, bool bEnab
 // Note - some params passed by value to allow them to be unconnected and enable inline BP editing
 bool UPhysicsControlComponent::SetControlData(
 	const FName               Name, 
-	const FPhysicsControlData ControlData, 
-	const bool                bEnableControl)
+	const FPhysicsControlData ControlData)
 {
 	FPhysicsControlRecord* Record = Implementation->FindControlRecord(Name);
 	if (Record)
 	{
 		Record->PhysicsControl.ControlData = ControlData;
-		if (bEnableControl)
-		{
-			Record->PhysicsControlState.bEnabled = true;
-		}
 		return true;
 	}
 	return false;
@@ -853,22 +838,20 @@ bool UPhysicsControlComponent::SetControlData(
 //======================================================================================================================
 void UPhysicsControlComponent::SetControlDatas(
 	const TArray<FName>&       Names, 
-	const FPhysicsControlData  ControlData, 
-	const bool                 bEnableControl)
+	const FPhysicsControlData  ControlData)
 {
 	for (FName Name : Names)
 	{
-		SetControlData(Name, ControlData, bEnableControl);
+		SetControlData(Name, ControlData);
 	}
 }
 
 //======================================================================================================================
 void UPhysicsControlComponent::SetControlDatasInSet(
 	const FName                SetName,
-	const FPhysicsControlData  ControlData,
-	const bool                 bEnableControl)
+	const FPhysicsControlData  ControlData)
 {
-	SetControlDatas(GetControlNamesInSet(SetName), ControlData, bEnableControl);
+	SetControlDatas(GetControlNamesInSet(SetName), ControlData);
 }
 
 //======================================================================================================================
@@ -884,7 +867,7 @@ bool UPhysicsControlComponent::SetControlMultiplier(
 		Record->PhysicsControl.ControlMultiplier = ControlMultiplier;
 		if (bEnableControl)
 		{
-			Record->PhysicsControlState.bEnabled = true;
+			Record->PhysicsControl.ControlData.bEnabled = true;
 		}
 		return true;
 	}
@@ -930,7 +913,7 @@ bool UPhysicsControlComponent::SetControlLinearData(
 		Record->PhysicsControl.ControlData.MaxForce = MaxForce;
 		if (bEnableControl)
 		{
-			Record->PhysicsControlState.bEnabled = true;
+			Record->PhysicsControl.ControlData.bEnabled = true;
 		}
 		return true;
 	}
@@ -955,7 +938,7 @@ bool UPhysicsControlComponent::SetControlAngularData(
 		Record->PhysicsControl.ControlData.MaxTorque = MaxTorque;
 		if (bEnableControl)
 		{
-			Record->PhysicsControlState.bEnabled = true;
+			Record->PhysicsControl.ControlData.bEnabled = true;
 		}
 		return true;
 	}
@@ -999,7 +982,7 @@ bool UPhysicsControlComponent::SetControlTarget(
 		Record->PhysicsControl.ControlTarget = ControlTarget;
 		if (bEnableControl)
 		{
-			Record->PhysicsControlState.bEnabled = true;
+			Record->PhysicsControl.ControlData.bEnabled = true;
 		}
 		return true;
 	}
@@ -1100,7 +1083,7 @@ bool UPhysicsControlComponent::SetControlTargetPosition(
 		Record->PhysicsControl.ControlTarget.bApplyControlPointToTarget = bApplyControlPointToTarget;
 		if (bEnableControl)
 		{
-			Record->PhysicsControlState.bEnabled = true;
+			Record->PhysicsControl.ControlData.bEnabled = true;
 		}
 		return true;
 	}
@@ -1163,7 +1146,7 @@ bool UPhysicsControlComponent::SetControlTargetOrientation(
 		Record->PhysicsControl.ControlTarget.bApplyControlPointToTarget = bApplyControlPointToTarget;
 		if (bEnableControl)
 		{
-			Record->PhysicsControlState.bEnabled = true;
+			Record->PhysicsControl.ControlData.bEnabled = true;
 		}
 		return true;
 	}
@@ -1307,7 +1290,7 @@ bool UPhysicsControlComponent::SetControlTargetPoses(
 		Record->PhysicsControl.ControlTarget.bApplyControlPointToTarget = true;
 		if (bEnableControl)
 		{
-			Record->PhysicsControlState.bEnabled = true;
+			Record->PhysicsControl.ControlData.bEnabled = true;
 		}
 		return true;
 	}
@@ -1352,33 +1335,6 @@ void UPhysicsControlComponent::SetControlsInSetUseSkeletalAnimation(
 		GetControlNamesInSet(SetName), bUseSkeletalAnimation, SkeletalAnimationVelocityMultiplier);
 }
 
-//======================================================================================================================
-bool UPhysicsControlComponent::SetControlAutoDisable(const FName Name, const bool bAutoDisable)
-{
-	FPhysicsControl* PhysicsControl = Implementation->FindControl(Name);
-	if (PhysicsControl)
-	{
-		PhysicsControl->ControlSettings.bAutoDisable = bAutoDisable;
-		return true;
-	}
-	return false;
-}
-
-//======================================================================================================================
-void UPhysicsControlComponent::SetControlsAutoDisable(const TArray<FName>& Names, const bool bAutoDisable)
-{
-	for (FName Name : Names)
-	{
-		SetControlAutoDisable(Name, bAutoDisable);
-	}
-}
-
-//======================================================================================================================
-void UPhysicsControlComponent::SetControlsInSetAutoDisable(const FName SetName, const bool bAutoDisable)
-{
-	SetControlsAutoDisable(GetControlNamesInSet(SetName), bAutoDisable);
-}
-
 //======================================================================================================================
 bool UPhysicsControlComponent::SetControlDisableCollision(const FName Name, const bool bDisableCollision)
 {
@@ -1442,24 +1398,13 @@ bool UPhysicsControlComponent::GetControlTarget(const FName Name, FPhysicsContro
 	return false;
 }
 
-//======================================================================================================================
-bool UPhysicsControlComponent::GetControlAutoDisable(const FName Name) const
-{
-	FPhysicsControl* PhysicsControl = Implementation->FindControl(Name);
-	if (PhysicsControl)
-	{
-		return PhysicsControl->ControlSettings.bAutoDisable;
-	}
-	return false;
-}
-
 //======================================================================================================================
 bool UPhysicsControlComponent::GetControlEnabled(const FName Name) const
 {
 	FPhysicsControlRecord* Record = Implementation->FindControlRecord(Name);
 	if (Record)
 	{
-		return Record->PhysicsControlState.bEnabled;
+		return Record->Enabled();
 	}
 	return false;
 }
@@ -1878,10 +1823,8 @@ void UPhysicsControlComponent::CreateControlsAndBodyModifiersFromLimbBones(
 	const TArray<FPhysicsControlLimbSetupData>& LimbSetupData,
 	const FPhysicsControlData                   WorldSpaceControlData,
 	const FPhysicsControlSettings               WorldSpaceControlSettings,
-	const bool                                  bEnableWorldSpaceControls,
 	const FPhysicsControlData                   ParentSpaceControlData,
 	const FPhysicsControlSettings               ParentSpaceControlSettings,
-	const bool                                  bEnableParentSpaceControls,
 	const EPhysicsMovementType                  PhysicsMovementType,
 	const float                                 GravityMultiplier,
 	const float                                 PhysicsBlendWeight)
@@ -1891,11 +1834,11 @@ void UPhysicsControlComponent::CreateControlsAndBodyModifiersFromLimbBones(
 
 	LimbWorldSpaceControls = CreateControlsFromLimbBones(
 		AllWorldSpaceControls, LimbBones, EPhysicsControlType::WorldSpace, 
-		WorldSpaceControlData, WorldSpaceControlSettings, bEnableWorldSpaceControls);
+		WorldSpaceControlData, WorldSpaceControlSettings);
 
 	LimbParentSpaceControls = CreateControlsFromLimbBones(
 		AllParentSpaceControls, LimbBones, EPhysicsControlType::ParentSpace,
-		ParentSpaceControlData, ParentSpaceControlSettings, bEnableParentSpaceControls);
+		ParentSpaceControlData, ParentSpaceControlSettings);
 
 	LimbBodyModifiers = CreateBodyModifiersFromLimbBones(
 		AllBodyModifiers, LimbBones, PhysicsMovementType, ECollisionEnabled::QueryAndPhysics, 
@@ -2314,7 +2257,7 @@ void UPhysicsControlComponent::DebugDraw(FPrimitiveDrawInterface* PDI) const
 
 				GEngine->AddOnScreenDebugMessage(
 					-1, 0.0f,
-					Record.PhysicsControlState.bEnabled ? FColor::Green : FColor::Red, Text);
+					Record.Enabled() ? FColor::Green : FColor::Red, Text);
 			}
 		}
 	}
@@ -2406,7 +2349,7 @@ void UPhysicsControlComponent::DebugDrawControl(
 	const bool bHaveLinear = Record.PhysicsControl.ControlData.LinearStrength > 0;
 	const bool bHaveAngular = Record.PhysicsControl.ControlData.AngularStrength > 0;
 
-	if (Record.PhysicsControlState.bEnabled && ConstraintInstance)
+	if (Record.Enabled() && ConstraintInstance)
 	{
 		FBodyInstance* ChildBodyInstance = UE::PhysicsControlComponent::GetBodyInstance(
 			Record.PhysicsControl.ChildMeshComponent, Record.PhysicsControl.ChildBoneName);
diff --git a/Engine/Plugins/Experimental/PhysicsControl/Source/PhysicsControl/Private/PhysicsControlComponentHelpers.cpp b/Engine/Plugins/Experimental/PhysicsControl/Source/PhysicsControl/Private/PhysicsControlComponentHelpers.cpp
--- a/Engine/Plugins/Experimental/PhysicsControl/Source/PhysicsControl/Private/PhysicsControlComponentHelpers.cpp
+++ b/Engine/Plugins/Experimental/PhysicsControl/Source/PhysicsControl/Private/PhysicsControlComponentHelpers.cpp
@@ -104,6 +104,8 @@ void ConvertConstraintProfileToControlData(
 		GetAngularDriveDamping(InProfileProperties.AngularDrive));
 
 	OutControlData.MaxTorque = float(GetAngularDriveMaxTorque(InProfileProperties.AngularDrive));
+
+	OutControlData.bEnabled = true;
 }
 
 //======================================================================================================================
diff --git a/Engine/Plugins/Experimental/PhysicsControl/Source/PhysicsControl/Private/PhysicsControlComponentImpl.cpp b/Engine/Plugins/Experimental/PhysicsControl/Source/PhysicsControl/Private/PhysicsControlComponentImpl.cpp
--- a/Engine/Plugins/Experimental/PhysicsControl/Source/PhysicsControl/Private/PhysicsControlComponentImpl.cpp
+++ b/Engine/Plugins/Experimental/PhysicsControl/Source/PhysicsControl/Private/PhysicsControlComponentImpl.cpp
@@ -707,7 +707,7 @@ void FPhysicsControlComponentImpl::ApplyControl(FPhysicsControlRecord& Record)
 {
 	FConstraintInstance* ConstraintInstance = Record.PhysicsControlState.ConstraintInstance.Get();
 
-	if (!ConstraintInstance || !Record.PhysicsControlState.bEnabled)
+	if (!ConstraintInstance || !Record.Enabled())
 	{
 		return;
 	}
diff --git a/Engine/Plugins/Experimental/PhysicsControl/Source/PhysicsControl/Private/PhysicsControlData.cpp b/Engine/Plugins/Experimental/PhysicsControl/Source/PhysicsControl/Private/PhysicsControlData.cpp
new file mode 100644
--- /dev/null
+++ b/Engine/Plugins/Experimental/PhysicsControl/Source/PhysicsControl/Private/PhysicsControlData.cpp
@@ -0,0 +1,30 @@
+// Copyright Epic Games, Inc. All Rights Reserved.
+
+#include "PhysicsControlData.h"
+
+#define SET_SPARSE_DATA(NAME) NAME = SparseData.bEnable##NAME ? SparseData.NAME : NAME
+
+//======================================================================================================================
+void FPhysicsControlData::UpdateFromSparseData(const FPhysicsControlSparseData& SparseData)
+{
+	SET_SPARSE_DATA(LinearStrength);
+	SET_SPARSE_DATA(LinearDampingRatio);
+	SET_SPARSE_DATA(LinearExtraDamping);
+	SET_SPARSE_DATA(MaxForce);
+	SET_SPARSE_DATA(AngularStrength);
+	SET_SPARSE_DATA(AngularDampingRatio);
+	SET_SPARSE_DATA(AngularExtraDamping);
+	SET_SPARSE_DATA(MaxTorque);
+	SET_SPARSE_DATA(LinearTargetVelocityMultiplier);
+	SET_SPARSE_DATA(AngularTargetVelocityMultiplier);
+	SET_SPARSE_DATA(bEnabled);
+}
+
+//======================================================================================================================
+void FPhysicsControlModifierData::UpdateFromSparseData(const FPhysicsControlModifierSparseData& SparseData)
+{
+	SET_SPARSE_DATA(MovementType);
+	SET_SPARSE_DATA(GravityMultiplier);
+}
+
+#undef SET_SPARSE_DATA
diff --git a/Engine/Plugins/Experimental/PhysicsControl/Source/PhysicsControl/Private/PhysicsControlRecord.h b/Engine/Plugins/Experimental/PhysicsControl/Source/PhysicsControl/Private/PhysicsControlRecord.h
--- a/Engine/Plugins/Experimental/PhysicsControl/Source/PhysicsControl/Private/PhysicsControlRecord.h
+++ b/Engine/Plugins/Experimental/PhysicsControl/Source/PhysicsControl/Private/PhysicsControlRecord.h
@@ -19,8 +19,6 @@ struct FPhysicsControlState
 	void Reset();
 
 	TSharedPtr<FConstraintInstance> ConstraintInstance;
-
-	bool bEnabled = false;
 };
 
 /**
@@ -44,6 +42,8 @@ struct FPhysicsControlRecord
 	/** Sets the control point to the center of mass of the child mesh (or to zero if that fails). */
 	void ResetControlPoint();
 
+	bool Enabled() const { return PhysicsControl.ControlData.bEnabled; }
+
 	/** The configuration data */
 	FPhysicsControl PhysicsControl;
 
diff --git a/Engine/Plugins/Experimental/PhysicsControl/Source/PhysicsControl/Private/RigidBodyControlData.cpp b/Engine/Plugins/Experimental/PhysicsControl/Source/PhysicsControl/Private/RigidBodyControlData.cpp
--- a/Engine/Plugins/Experimental/PhysicsControl/Source/PhysicsControl/Private/RigidBodyControlData.cpp
+++ b/Engine/Plugins/Experimental/PhysicsControl/Source/PhysicsControl/Private/RigidBodyControlData.cpp
@@ -2,40 +2,14 @@
 
 #include "RigidBodyControlData.h"
 
-//======================================================================================================================
-FName GetControlTypeName(const ERigidBodyMovementType MovementType)
-{
-	switch (MovementType)
-	{
-	case ERigidBodyMovementType::Kinematic:
-		return "Kinematic";
-	case ERigidBodyMovementType::Simulated:
-		return "Simulated";
-	}
-	return "None";
-}
-
-//======================================================================================================================
-FName GetControlTypeName(const ERigidBodyControlType ControlType)
-{
-	switch (ControlType)
-	{
-	case ERigidBodyControlType::ParentSpace:
-		return "ParentSpace";
-	case ERigidBodyControlType::WorldSpace:
-		return "WorldSpace";
-	}
-	return "None";
-}
-
 #define INTERPOLATE_PARAM(NAME) Output.NAME = FMath::Lerp(A.NAME, B.NAME, Weight)
 #define SET_ENABLED_PARAM(NAME) Output.bEnable##NAME = A.bEnable##NAME && B.bEnable##NAME
 
 //======================================================================================================================
-FRigidBodyControlData Interpolate(
-	const FRigidBodyControlData& A, const FRigidBodyControlData& B, const float Weight)
+FPhysicsControlData Interpolate(
+	const FPhysicsControlData& A, const FPhysicsControlData& B, const float Weight)
 {
-	FRigidBodyControlData Output;
+	FPhysicsControlData Output;
 	INTERPOLATE_PARAM(LinearStrength);
 	INTERPOLATE_PARAM(LinearDampingRatio);
 	INTERPOLATE_PARAM(LinearExtraDamping);
@@ -49,10 +23,10 @@ FRigidBodyControlData Interpolate(
 }
 
 //======================================================================================================================
-FRigidBodyControlSparseData Interpolate(
-	const FRigidBodyControlSparseData& A, const FRigidBodyControlSparseData& B, const float Weight)
+FPhysicsControlSparseData Interpolate(
+	const FPhysicsControlSparseData& A, const FPhysicsControlSparseData& B, const float Weight)
 {
-	FRigidBodyControlSparseData Output;
+	FPhysicsControlSparseData Output;
 	INTERPOLATE_PARAM(LinearStrength);
 	INTERPOLATE_PARAM(LinearDampingRatio);
 	INTERPOLATE_PARAM(LinearExtraDamping);
@@ -77,20 +51,20 @@ FRigidBodyControlSparseData Interpolate(
 }
 
 //======================================================================================================================
-FRigidBodyModifierData Interpolate(
-	const FRigidBodyModifierData& A, const FRigidBodyModifierData& B, const float Weight)
+FPhysicsControlModifierData Interpolate(
+	const FPhysicsControlModifierData& A, const FPhysicsControlModifierData& B, const float Weight)
 {
-	FRigidBodyModifierData Output;
+	FPhysicsControlModifierData Output;
 	Output.MovementType = (Weight < 0.5f) ? A.MovementType : B.MovementType;
 	INTERPOLATE_PARAM(GravityMultiplier);
 	return Output;
 }
 
 //======================================================================================================================
-FRigidBodyModifierSparseData Interpolate(
-	const FRigidBodyModifierSparseData& A, const FRigidBodyModifierSparseData& B, const float Weight)
+FPhysicsControlModifierSparseData Interpolate(
+	const FPhysicsControlModifierSparseData& A, const FPhysicsControlModifierSparseData& B, const float Weight)
 {
-	FRigidBodyModifierSparseData Output;
+	FPhysicsControlModifierSparseData Output;
 	Output.MovementType = (Weight < 0.5f) ? A.MovementType : B.MovementType;
 	INTERPOLATE_PARAM(GravityMultiplier);
 
@@ -103,7 +77,7 @@ FRigidBodyModifierSparseData Interpolate(
 #undef SET_ENABLED_PARAM
 
 //======================================================================================================================
-FControlRecord::FControlRecord(const FRigidBodyControl& InControl, ImmediatePhysics::FJointHandle* InJointHandle)
+FRigidBodyControlRecord::FRigidBodyControlRecord(const FRigidBodyControl& InControl, ImmediatePhysics::FJointHandle* InJointHandle)
 	: Control(InControl)
 	, JointHandle(InJointHandle)
 	, CurrentData(Control.ControlData)
@@ -111,7 +85,7 @@ FControlRecord::FControlRecord(const FRigidBodyControl& InControl, ImmediatePhys
 }
 
 //======================================================================================================================
-void FControlRecord::ResetCurrent(bool bResetTarget)
+void FRigidBodyControlRecord::ResetCurrent(bool bResetTarget)
 {
 	CurrentData = Control.ControlData;
 
@@ -122,7 +96,7 @@ void FControlRecord::ResetCurrent(bool bResetTarget)
 }
 
 //======================================================================================================================
-FBodyModifierRecord::FBodyModifierRecord(const FRigidBodyModifier& InModifier, ImmediatePhysics::FActorHandle* InActorHandle)
+FRigidBodyModifierRecord::FRigidBodyModifierRecord(const FRigidBodyModifier& InModifier, ImmediatePhysics::FActorHandle* InActorHandle)
 	: Modifier(InModifier)
 	, ActorHandle(InActorHandle)
 	, CurrentData(Modifier.ModifierData)
@@ -130,32 +104,8 @@ FBodyModifierRecord::FBodyModifierRecord(const FRigidBodyModifier& InModifier, I
 }
 
 //======================================================================================================================
-void FBodyModifierRecord::ResetCurrent()
+void FRigidBodyModifierRecord::ResetCurrent()
 {
 	CurrentData = Modifier.ModifierData;
 }
 
-#define SET_SPARSE_DATA(NAME) NAME = SparseData.bEnable##NAME ? SparseData.NAME : NAME
-
-//======================================================================================================================
-void FRigidBodyControlData::UpdateFromSparseData(const FRigidBodyControlSparseData& SparseData)
-{
-	SET_SPARSE_DATA(LinearStrength);
-	SET_SPARSE_DATA(LinearDampingRatio);
-	SET_SPARSE_DATA(LinearExtraDamping);
-	SET_SPARSE_DATA(AngularStrength);
-	SET_SPARSE_DATA(AngularDampingRatio);
-	SET_SPARSE_DATA(AngularExtraDamping);
-	SET_SPARSE_DATA(LinearTargetVelocityMultiplier);
-	SET_SPARSE_DATA(AngularTargetVelocityMultiplier);
-	SET_SPARSE_DATA(bEnabled);
-}
-
-//======================================================================================================================
-void FRigidBodyModifierData::UpdateFromSparseData(const FRigidBodyModifierSparseData& SparseData)
-{
-	SET_SPARSE_DATA(MovementType);
-	SET_SPARSE_DATA(GravityMultiplier);
-}
-
-#undef SET_SPARSE_DATA
diff --git a/Engine/Plugins/Experimental/PhysicsControl/Source/PhysicsControl/Private/RigidBodyWithControl.cpp b/Engine/Plugins/Experimental/PhysicsControl/Source/PhysicsControl/Private/RigidBodyWithControl.cpp
--- a/Engine/Plugins/Experimental/PhysicsControl/Source/PhysicsControl/Private/RigidBodyWithControl.cpp
+++ b/Engine/Plugins/Experimental/PhysicsControl/Source/PhysicsControl/Private/RigidBodyWithControl.cpp
@@ -216,7 +216,7 @@ FName FAnimNode_RigidBodyWithControl::FindParentBoneNameFromBoneName(const FName
 
 //======================================================================================================================
 FName FAnimNode_RigidBodyWithControl::CreateControl(
-	const FName ParentBoneName, const FName ChildBoneName, const FRigidBodyControlData& ControlData)
+	const FName ParentBoneName, const FName ChildBoneName, const FPhysicsControlData& ControlData)
 {
 	FRigidBodyControl Control;
 	Control.ParentBoneName = ParentBoneName;
@@ -246,23 +246,23 @@ FName FAnimNode_RigidBodyWithControl::CreateControl(
 	}
 	
 	FName ControlName = GetUniqueControlName(Control.ParentBoneName, Control.ChildBoneName);
-	ControlRecords.Add(ControlName, FControlRecord(Control, JointHandle));
+	ControlRecords.Add(ControlName, FRigidBodyControlRecord(Control, JointHandle));
 
 	return ControlName;
 }
 
 //======================================================================================================================
 void FAnimNode_RigidBodyWithControl::CreateControlsFromLimbBones(
-	const FName                   LimbName,
-	const FRigidBodyLimbBones&    LimbBones,
-	const ERigidBodyControlType   ControlType,
-	const FRigidBodyControlData&  ControlData)
+	const FName                     LimbName,
+	const FPhysicsControlLimbBones& LimbBones,
+	const EPhysicsControlType       ControlType,
+	const FPhysicsControlData&      ControlData)
 {
 	for (const FName ChildBoneName : LimbBones.BoneNames)
 	{
 		FName ParentBoneName;
 
-		if (ControlType == ERigidBodyControlType::ParentSpace)
+		if (ControlType == EPhysicsControlType::ParentSpace)
 		{
 			ParentBoneName = FindParentBoneNameFromBoneName(ChildBoneName);
 			if (ParentBoneName.IsNone())
@@ -282,9 +282,9 @@ void FAnimNode_RigidBodyWithControl::CreateControlsFromLimbBones(
 		else
 		{
 			NameRecords.AddControl(ControlName, LimbName);
-			NameRecords.AddControl(ControlName, GetControlTypeName(ControlType));
+			NameRecords.AddControl(ControlName, GetPhysicsControlTypeName(ControlType));
 			NameRecords.AddControl(ControlName, FName(
-				GetControlTypeName(ControlType).ToString().Append("_").Append(LimbName.ToString())));
+				GetPhysicsControlTypeName(ControlType).ToString().Append("_").Append(LimbName.ToString())));
 		}
 	}
 }
@@ -383,7 +383,7 @@ FName FAnimNode_RigidBodyWithControl::GetUniqueControlName(const FName ParentBon
 }
 
 //======================================================================================================================
-FName FAnimNode_RigidBodyWithControl::CreateBodyModifier(FName BoneName, const FRigidBodyModifierData& ModifierData)
+FName FAnimNode_RigidBodyWithControl::CreateBodyModifier(FName BoneName, const FPhysicsControlModifierData& ModifierData)
 {
 	FName Name;
 
@@ -395,7 +395,7 @@ FName FAnimNode_RigidBodyWithControl::CreateBodyModifier(FName BoneName, const F
 		FRigidBodyModifier BodyModifier;
 		BodyModifier.BoneName = BoneName;
 		BodyModifier.ModifierData = ModifierData;
-		FBodyModifierRecord& Modifier = ModifierRecords.Add(Name, FBodyModifierRecord(BodyModifier, ActorHandle));
+		FRigidBodyModifierRecord& Modifier = ModifierRecords.Add(Name, FRigidBodyModifierRecord(BodyModifier, ActorHandle));
 	}
 
 	return Name;
@@ -403,9 +403,9 @@ FName FAnimNode_RigidBodyWithControl::CreateBodyModifier(FName BoneName, const F
 
 //======================================================================================================================
 void FAnimNode_RigidBodyWithControl::CreateBodyModifiersFromLimbBones(
-	const FName                   LimbName,
-	const FRigidBodyLimbBones&    LimbBones,
-	const FRigidBodyModifierData& DefaultModifierData)
+	const FName                        LimbName,
+	const FPhysicsControlLimbBones&    LimbBones,
+	const FPhysicsControlModifierData& DefaultModifierData)
 {
 	for (const FName BoneName : LimbBones.BoneNames)
 	{
@@ -423,15 +423,15 @@ void FAnimNode_RigidBodyWithControl::CreateBodyModifiersFromLimbBones(
 }
 
 //======================================================================================================================
-TMap<FName, FRigidBodyLimbBones> FAnimNode_RigidBodyWithControl::GetLimbBonesFromSkeletalMesh(
-	const TArray<FRigidBodyLimbSetupData>& LimbSetupData,
-	USkeletalMeshComponent* const SkeletalMeshComponent) const
+TMap<FName, FPhysicsControlLimbBones> FAnimNode_RigidBodyWithControl::GetLimbBonesFromSkeletalMesh(
+	const TArray<FPhysicsControlLimbSetupData>& LimbSetupData,
+	USkeletalMeshComponent* const               SkeletalMeshComponent) const
 {
 	// Parse the skeleton tree to figure out which bones are associated with which limbs. 
 	
 	// TODO - Output limb bones are not in the order specified in the skeleton - would be better if they were
 	
-	TMap<FName, FRigidBodyLimbBones> Result;
+	TMap<FName, FPhysicsControlLimbBones> Result;
 
 	UPhysicsAsset* const PhysicsAsset = GetPhysicsAsset();
 
@@ -451,9 +451,9 @@ TMap<FName, FRigidBodyLimbBones> FAnimNode_RigidBodyWithControl::GetLimbBonesFro
 
 	TSet<FName> AllBones;
 
-	for (const FRigidBodyLimbSetupData& LimbSetup : LimbSetupData)
+	for (const FPhysicsControlLimbSetupData& LimbSetup : LimbSetupData)
 	{
-		FRigidBodyLimbBones& LimbBones = Result.Add(LimbSetup.LimbName);
+		FPhysicsControlLimbBones& LimbBones = Result.Add(LimbSetup.LimbName);
 
 		LimbBones.bFirstBoneIsAdditional = false;
 		LimbBones.bCreateWorldSpaceControls = LimbSetup.bCreateWorldSpaceControls;
@@ -507,22 +507,22 @@ void FAnimNode_RigidBodyWithControl::InitControlsAndBodyModifiers(USkeletalMeshC
 	check(ControlRecords.IsEmpty()); // Controls should not exist when this function is called.
 
 	// These functions will create the base set of controls and modifiers from SetupData
-	TMap<FName, FRigidBodyLimbBones> AllLimbBones = 
+	TMap<FName, FPhysicsControlLimbBones> AllLimbBones =
 		GetLimbBonesFromSkeletalMesh(SetupData.LimbSetupData, SkeletalMeshComponent);
 
-	for (const TMap<FName, FRigidBodyLimbBones>::ElementType& LimbBoneEntry : AllLimbBones)
+	for (const TMap<FName, FPhysicsControlLimbBones>::ElementType& LimbBoneEntry : AllLimbBones)
 	{
 		const FName LimbName = LimbBoneEntry.Key;
-		const FRigidBodyLimbBones& LimbBones = LimbBoneEntry.Value;
+		const FPhysicsControlLimbBones& LimbBones = LimbBoneEntry.Value;
 		if (LimbBones.bCreateWorldSpaceControls)
 		{
 			CreateControlsFromLimbBones(
-				LimbName, LimbBones, ERigidBodyControlType::WorldSpace, SetupData.DefaultWorldSpaceControlData);
+				LimbName, LimbBones, EPhysicsControlType::WorldSpace, SetupData.DefaultWorldSpaceControlData);
 		}
 		if (LimbBones.bCreateParentSpaceControls)
 		{
 			CreateControlsFromLimbBones(
-				LimbName, LimbBones, ERigidBodyControlType::ParentSpace, SetupData.DefaultParentSpaceControlData);
+				LimbName, LimbBones, EPhysicsControlType::ParentSpace, SetupData.DefaultParentSpaceControlData);
 		}
 		if (LimbBones.bCreateBodyModifiers)
 		{
@@ -566,10 +566,10 @@ void FAnimNode_RigidBodyWithControl::LogControlsModifiersAndSets()
 #define RBWC_LOG_LEVEL Display
 
 	UE_LOG(LogRigidBodyWithControl, RBWC_LOG_LEVEL, TEXT("Controls:"));
-	for (TMap<FName, FControlRecord>::ElementType& NameRecordPair : ControlRecords)
+	for (TMap<FName, FRigidBodyControlRecord>::ElementType& NameRecordPair : ControlRecords)
 	{
 		UE_LOG(LogRigidBodyWithControl, RBWC_LOG_LEVEL, TEXT("  %s:"), *NameRecordPair.Key.ToString());
-		const FControlRecord& Record = NameRecordPair.Value;
+		const FRigidBodyControlRecord& Record = NameRecordPair.Value;
 		UE_LOG(LogRigidBodyWithControl, RBWC_LOG_LEVEL, TEXT("    Parent bone: %s Child bone: %s"), 
 			*Record.Control.ParentBoneName.ToString(), *Record.Control.ChildBoneName.ToString());
 		UE_LOG(LogRigidBodyWithControl, RBWC_LOG_LEVEL, TEXT("    Enabled %d"),
@@ -585,14 +585,14 @@ void FAnimNode_RigidBodyWithControl::LogControlsModifiersAndSets()
 	}
 
 	UE_LOG(LogRigidBodyWithControl, RBWC_LOG_LEVEL, TEXT("Body Modifiers:"));
-	for (TMap<FName, FBodyModifierRecord>::ElementType& NameRecordPair : ModifierRecords)
+	for (TMap<FName, FRigidBodyModifierRecord>::ElementType& NameRecordPair : ModifierRecords)
 	{
 		UE_LOG(LogRigidBodyWithControl, RBWC_LOG_LEVEL, TEXT("  %s:"), *NameRecordPair.Key.ToString());
-		const FBodyModifierRecord& Record = NameRecordPair.Value;
+		const FRigidBodyModifierRecord& Record = NameRecordPair.Value;
 		UE_LOG(LogRigidBodyWithControl, RBWC_LOG_LEVEL, TEXT("    Bone: %s Body: %s"),
 			*Record.Modifier.BoneName.ToString(), *GetBodyFromBoneName(Record.Modifier.BoneName).ToString());
 		UE_LOG(LogRigidBodyWithControl, RBWC_LOG_LEVEL, TEXT("    Movement: %s GravityMultiplier: %f"),
-			*GetControlTypeName(Record.Modifier.ModifierData.MovementType).ToString(),
+			*GetPhysicsMovementTypeName(Record.Modifier.ModifierData.MovementType).ToString(),
 			Record.Modifier.ModifierData.GravityMultiplier);
 	}
 
@@ -670,7 +670,7 @@ void FAnimNode_RigidBodyWithControl::CreateAdditionalControls()
 // Slightly annoying to have to add the names individually, but we want to check they exist
 void FAnimNode_RigidBodyWithControl::CreateAdditionalSets()
 {
-	for (const FSetUpdate& Set : AdditionalSets.ControlSetUpdates)
+	for (const FPhysicsControlSetUpdate& Set : AdditionalSets.ControlSetUpdates)
 	{
 		TArray<FName> Names = ExpandSetNames(Set.Names, NameRecords.ControlSets);
 		for (FName Name : Names)
@@ -688,7 +688,7 @@ void FAnimNode_RigidBodyWithControl::CreateAdditionalSets()
 		}
 	}
 
-	for (const FSetUpdate& Set : AdditionalSets.ModifierSetUpdates)
+	for (const FPhysicsControlSetUpdate& Set : AdditionalSets.ModifierSetUpdates)
 	{
 		TArray<FName> Names = ExpandSetNames(Set.Names, NameRecords.BodyModifierSets);
 		for (FName Name : Names)
@@ -723,8 +723,8 @@ static void ConvertStrengthToSpringParams(
 
 //======================================================================================================================
 void FAnimNode_RigidBodyWithControl::UpdateDriveSpringDamperSettings(
-	Chaos::FPBDJointSettings&             Settings, 
-	const FRigidBodyControlData&          ControlData)
+	Chaos::FPBDJointSettings&  Settings, 
+	const FPhysicsControlData& ControlData)
 {
 	float LinearSpring;
 	float LinearDamping;
@@ -764,7 +764,7 @@ static FTransform CalculateTargetTM(
 }
 
 //======================================================================================================================
-void FAnimNode_RigidBodyWithControl::ApplyControl(FControlRecord& ControlRecord, float DeltaTime)
+void FAnimNode_RigidBodyWithControl::ApplyControl(FRigidBodyControlRecord& ControlRecord, float DeltaTime)
 {
 	using namespace ImmediatePhysics;
 	FJointHandle* const JointHandle = ControlRecord.JointHandle;
@@ -858,14 +858,14 @@ void FAnimNode_RigidBodyWithControl::ApplyControl(FControlRecord& ControlRecord,
 
 //======================================================================================================================
 void FAnimNode_RigidBodyWithControl::ApplyModifier(
-	const FBodyModifierRecord& BodyModifierRecord, 
+	const FRigidBodyModifierRecord& BodyModifierRecord, 
 	const FVector&             SimSpaceGravity)
 {
 	if (BodyModifierRecord.ActorHandle)
 	{
 		// Note that there's an early out if there's no change needed, so this should be OK.
 		BodyModifierRecord.ActorHandle->SetIsKinematic(
-			BodyModifierRecord.CurrentData.MovementType == ERigidBodyMovementType::Kinematic);
+			BodyModifierRecord.CurrentData.MovementType != EPhysicsMovementType::Simulated);
 
 		// Note that the actual kinematic targets will be set separately, since they need to be set
 		// for all kinematics whether or not they were under a modifier.
@@ -883,8 +883,8 @@ void FAnimNode_RigidBodyWithControl::ApplyModifier(
 
 //======================================================================================================================
 void FAnimNode_RigidBodyWithControl::ApplyControlAndModifierUpdatesAndParametersToRecords(
-	const FRigidBodyControlAndModifierUpdates&    Updates,
-	const FRigidBodyControlAndModifierParameters& Parameters)
+	const FPhysicsControlControlAndModifierUpdates&    Updates,
+	const FPhysicsControlControlAndModifierParameters& Parameters)
 {
 	SCOPE_CYCLE_COUNTER(STAT_RigidBodyNodeWithControl_ApplyControlAndModifierUpdatesAndParametersToRecords);
 
@@ -895,7 +895,7 @@ void FAnimNode_RigidBodyWithControl::ApplyControlAndModifierUpdatesAndParameters
 	// This goes through the records, resetting the update parts.
 	// Then the update structures get adjusted based on the parameters.
 	// The results don't get applied to the actual constraints yet - that happens in ApplyControlsAndModifiers
-	for (TMap<FName, FControlRecord>::ElementType& NameRecordPair : ControlRecords)
+	for (TMap<FName, FRigidBodyControlRecord>::ElementType& NameRecordPair : ControlRecords)
 	{
 		FName ControlName = NameRecordPair.Key;
 		if (const FRigidBodyControlTarget* ControlTarget = ControlTargets.Targets.Find(ControlName))
@@ -909,7 +909,7 @@ void FAnimNode_RigidBodyWithControl::ApplyControlAndModifierUpdatesAndParameters
 		}
 	}
 
-	for (TMap<FName, FBodyModifierRecord>::ElementType& NameRecordPair : ModifierRecords)
+	for (TMap<FName, FRigidBodyModifierRecord>::ElementType& NameRecordPair : ModifierRecords)
 	{
 		NameRecordPair.Value.ResetCurrent();
 	}
@@ -922,19 +922,19 @@ void FAnimNode_RigidBodyWithControl::ApplyControlAndModifierUpdatesAndParameters
 
 //======================================================================================================================
 void FAnimNode_RigidBodyWithControl::ApplyControlAndBodyModifierDatas(
-	const TArray<FRigidBodyNamedControlParameters>& InControlParameters,
-	const TArray<FRigidBodyNamedModifierParameters>& InModifierParameters)
+	const TArray<FPhysicsControlNamedControlParameters>& InControlParameters,
+	const TArray<FPhysicsControlNamedModifierParameters>& InModifierParameters)
 {
 	SCOPE_CYCLE_COUNTER(STAT_RigidBodyNodeWithControl_ApplyControlsAndModifierDatas);
 
 	// This updates the "original" controls and modifiers based on the parameters.
-	for (const FRigidBodyNamedControlParameters& ControlParameters : InControlParameters)
+	for (const FPhysicsControlNamedControlParameters& ControlParameters : InControlParameters)
 	{
 		TArray<FName> Names = ExpandSetName(ControlParameters.Name, NameRecords.ControlSets);
 		for (FName Name : Names)
 		{
-			const FRigidBodyControlSparseData& ControlData = ControlParameters.Data;
-			if (FControlRecord* ControlRecord = ControlRecords.Find(Name))
+			const FPhysicsControlSparseData& ControlData = ControlParameters.Data;
+			if (FRigidBodyControlRecord* ControlRecord = ControlRecords.Find(Name))
 			{
 				ControlRecord->Control.ControlData.UpdateFromSparseData(ControlData);
 			}
@@ -946,13 +946,13 @@ void FAnimNode_RigidBodyWithControl::ApplyControlAndBodyModifierDatas(
 		}
 	}
 
-	for (const FRigidBodyNamedModifierParameters& ModifierParameters : InModifierParameters)
+	for (const FPhysicsControlNamedModifierParameters& ModifierParameters : InModifierParameters)
 	{
 		TArray<FName> Names = ExpandSetName(ModifierParameters.Name, NameRecords.BodyModifierSets);
 		for (FName Name : Names)
 		{
-			const FRigidBodyModifierSparseData& ModifierData = ModifierParameters.Data;
-			if (FBodyModifierRecord* ModifierRecord = ModifierRecords.Find(Name))
+			const FPhysicsControlModifierSparseData& ModifierData = ModifierParameters.Data;
+			if (FRigidBodyModifierRecord* ModifierRecord = ModifierRecords.Find(Name))
 			{
 				ModifierRecord->Modifier.ModifierData.UpdateFromSparseData(ModifierData);
 			}
@@ -990,9 +990,9 @@ void FAnimNode_RigidBodyWithControl::ApplyControlsAndModifiers(const FVector& Si
 	if (!PoseData.IsEmpty())
 	{
 		// Apply Controls.
-		for (TMap<FName, FControlRecord>::ElementType& NameRecordPair : ControlRecords)
+		for (TMap<FName, FRigidBodyControlRecord>::ElementType& NameRecordPair : ControlRecords)
 		{
-			FControlRecord& ControlRecord = NameRecordPair.Value;
+			FRigidBodyControlRecord& ControlRecord = NameRecordPair.Value;
 			if (ControlRecord.IsEnabled())
 			{
 				ApplyControl(ControlRecord, DeltaTime);
@@ -1002,7 +1002,7 @@ void FAnimNode_RigidBodyWithControl::ApplyControlsAndModifiers(const FVector& Si
 		}
 
 		// Apply Body Modifiers.
-		for (const TMap<FName, FBodyModifierRecord>::ElementType& NameRecordPair : ModifierRecords)
+		for (const TMap<FName, FRigidBodyModifierRecord>::ElementType& NameRecordPair : ModifierRecords)
 		{
 			ApplyModifier(NameRecordPair.Value, SimSpaceGravity);
 		}
@@ -1022,7 +1022,7 @@ void FAnimNode_RigidBodyWithControl::ApplyKinematicTargets()
 		{
 			FName BodyModifierName = KinematicTargetPair.Key;
 			FRigidBodyKinematicTarget& Target = KinematicTargetPair.Value;
-			FBodyModifierRecord* ModifierRecord = ModifierRecords.Find(BodyModifierName);
+			FRigidBodyModifierRecord* ModifierRecord = ModifierRecords.Find(BodyModifierName);
 			if (ModifierRecord && ModifierRecord->ActorHandle)
 			{
 				ImmediatePhysics::FActorHandle* ActorHandle = ModifierRecord->ActorHandle;
diff --git a/Engine/Plugins/Experimental/PhysicsControl/Source/PhysicsControl/Public/AnimNode_RigidBodyWithControl.h b/Engine/Plugins/Experimental/PhysicsControl/Source/PhysicsControl/Public/AnimNode_RigidBodyWithControl.h
--- a/Engine/Plugins/Experimental/PhysicsControl/Source/PhysicsControl/Public/AnimNode_RigidBodyWithControl.h
+++ b/Engine/Plugins/Experimental/PhysicsControl/Source/PhysicsControl/Public/AnimNode_RigidBodyWithControl.h
@@ -9,6 +9,7 @@
 #include "CoreMinimal.h"
 #include "RigidBodyPoseData.h"
 #include "RigidBodyControlData.h"
+#include "PhysicsControlLimbData.h"
 #include "RigidBodyNameRecords.h"
 #include "BoneControllers/AnimNode_SkeletalControlBase.h"
 #include "Components/SkeletalMeshComponent.h"
@@ -222,7 +223,7 @@ public:
 	 * Allows additional sets of controls or modifiers to be created, and existing sets to be modified
 	 */
 	UPROPERTY(EditAnywhere, Category = ControlSetup, meta = (PinHiddenByDefault))
-	FRigidBodySetUpdates AdditionalSets;
+	FPhysicsControlSetUpdates AdditionalSets;
 
 	/**
 	 * An initial set of controls that should be applied immediately after setup. This allows individual or 
@@ -230,20 +231,20 @@ public:
 	 * controls.
 	 */
 	UPROPERTY(EditAnywhere, Category = ControlSetup, meta = (PinHiddenByDefault))
-	FRigidBodyControlAndModifierParameters InitialControlAndBodyModifierUpdates;
+	FPhysicsControlControlAndModifierParameters InitialControlAndBodyModifierUpdates;
 
 	/**
 	 * Controls that should be applied each frame, and can be expected to change. Note that if these
 	 * stop being passed in then the controls and modifiers will return to their normal/original state.
 	 */
 	UPROPERTY(EditAnywhere, Category = Controls, meta = (PinShownByDefault))
-	FRigidBodyControlAndModifierParameters ControlAndModifierParameters;
+	FPhysicsControlControlAndModifierParameters ControlAndModifierParameters;
 
 	/**
 	 * Updates to controls that can be applied. Note that these update the normal/original state.
 	 */
 	UPROPERTY(EditAnywhere, Category = Controls, meta = (PinShownByDefault))
-	FRigidBodyControlAndModifierUpdates ControlAndModifierUpdates;
+	FPhysicsControlControlAndModifierUpdates ControlAndModifierUpdates;
 
 	/**
 	 * Targets that should be applied to the controls. 
@@ -302,8 +303,8 @@ private:
 		ImmediatePhysics::FActorHandle* const ChildBodyHandle, ImmediatePhysics::FActorHandle* const ParentBodyHandle);
 	void CreateWorldSpaceControlRootBody(UPhysicsAsset* const PhysicsAsset);
 
-	void ApplyControl(FControlRecord& ControlRecord, float DeltaTime);
-	void ApplyModifier(const FBodyModifierRecord& BodyModifierRecord, const FVector& SimSpaceGravity);
+	void ApplyControl(FRigidBodyControlRecord& ControlRecord, float DeltaTime);
+	void ApplyModifier(const FRigidBodyModifierRecord& BodyModifierRecord, const FVector& SimSpaceGravity);
 
 	void ApplyControlsAndModifiers(const FVector& SimSpaceGravity, float DeltaTime);
 
@@ -314,7 +315,7 @@ private:
 	void ApplyCurrentConstraintProfile();
 
 	// Adjusts the spring drive settings to reflect the control data
-	void UpdateDriveSpringDamperSettings(Chaos::FPBDJointSettings& Settings, const FRigidBodyControlData& ControlData);
+	void UpdateDriveSpringDamperSettings(Chaos::FPBDJointSettings& Settings, const FPhysicsControlData& ControlData);
 
 	TMap<FName, int32> BodyNameToIndexMap;
 	ImmediatePhysics::FActorHandle* WorldSpaceControlActorHandle;
@@ -325,24 +326,24 @@ private:
 
 	FName GetUniqueControlName(const FName ParentBoneName, const FName ChildBoneName);
 
-	FName CreateControl(const FName ParentBoneName, const FName ChildBoneName, const FRigidBodyControlData& ControlData);
+	FName CreateControl(const FName ParentBoneName, const FName ChildBoneName, const FPhysicsControlData& ControlData);
 
-	FName CreateBodyModifier(const FName BoneName, const FRigidBodyModifierData& ModifierData);
+	FName CreateBodyModifier(const FName BoneName, const FPhysicsControlModifierData& ModifierData);
 
 	void CreateControlsFromLimbBones(
-		const FName                  LimbName,
-		const FRigidBodyLimbBones&   LimbBones,
-		const ERigidBodyControlType  ControlType,
-		const FRigidBodyControlData& ControlData);
+		const FName                     LimbName,
+		const FPhysicsControlLimbBones& LimbBones,
+		const EPhysicsControlType       ControlType,
+		const FPhysicsControlData&      ControlData);
 
 	void CreateBodyModifiersFromLimbBones(
-		const FName                   LimbName,
-		const FRigidBodyLimbBones&    LimbBones,
-		const FRigidBodyModifierData& DefaultModifierData);
+		const FName                        LimbName,
+		const FPhysicsControlLimbBones&    LimbBones,
+		const FPhysicsControlModifierData& DefaultModifierData);
 
-	TMap<FName, FRigidBodyLimbBones> GetLimbBonesFromSkeletalMesh(
-		const TArray<FRigidBodyLimbSetupData>& LimbSetupData,
-		USkeletalMeshComponent* const SkeletalMeshComponent) const;
+	TMap<FName, FPhysicsControlLimbBones> GetLimbBonesFromSkeletalMesh(
+		const TArray<FPhysicsControlLimbSetupData>& LimbSetupData,
+		USkeletalMeshComponent* const               SkeletalMeshComponent) const;
 
 	// This will walk through the skeleton, create controls and body modifiers, and create the sets.
 	void InitControlsAndBodyModifiers(USkeletalMeshComponent* const SkeletalMeshComponent);
@@ -360,13 +361,13 @@ private:
 
 	// Applies the overrides to the underlying controls and modifiers
 	void ApplyControlAndBodyModifierDatas(
-		const TArray<FRigidBodyNamedControlParameters>& ControlParameters,
-		const TArray<FRigidBodyNamedModifierParameters>& ModifierParameters);
+		const TArray<FPhysicsControlNamedControlParameters>& ControlParameters,
+		const TArray<FPhysicsControlNamedModifierParameters>& ModifierParameters);
 
 	// Applies the parameters to control and modifier records
 	void ApplyControlAndModifierUpdatesAndParametersToRecords(
-		const FRigidBodyControlAndModifierUpdates&    Updates,
-		const FRigidBodyControlAndModifierParameters& Parameters);
+		const FPhysicsControlControlAndModifierUpdates&    Updates,
+		const FPhysicsControlControlAndModifierParameters& Parameters);
 
 	void InitializeNewBodyTransformsDuringSimulation(
 		FComponentSpacePoseContext& Output, const FTransform& ComponentTransform, const FTransform& BaseBoneTM);
@@ -496,10 +497,10 @@ private:
 	RigidBodyWithControl::FRigidBodyPoseData PoseData;
 
 	// Map of control records - they will be referenced by name
-	TMap<FName, FControlRecord> ControlRecords;
+	TMap<FName, FRigidBodyControlRecord> ControlRecords;
 
 	// Map of body modifier records - they will be referenced by name
-	TMap<FName, FBodyModifierRecord> ModifierRecords;
+	TMap<FName, FRigidBodyModifierRecord> ModifierRecords;
 
 	// Details about sets etc
 	FRigidBodyNameRecords NameRecords;
diff --git a/Engine/Plugins/Experimental/PhysicsControl/Source/PhysicsControl/Public/PhysicsControlActor.h b/Engine/Plugins/Experimental/PhysicsControl/Source/PhysicsControl/Public/PhysicsControlActor.h
--- a/Engine/Plugins/Experimental/PhysicsControl/Source/PhysicsControl/Public/PhysicsControlActor.h
+++ b/Engine/Plugins/Experimental/PhysicsControl/Source/PhysicsControl/Public/PhysicsControlActor.h
@@ -176,18 +176,12 @@ struct PHYSICSCONTROL_API FInitialCharacterControls
 	UPROPERTY(EditAnywhere, BlueprintReadWrite, Category = PhysicsControl)
 	FPhysicsControlSettings WorldSpaceControlSettings;
 
-	UPROPERTY(EditAnywhere, BlueprintReadWrite, Category = PhysicsControl)
-	bool bEnableWorldSpaceControls = false;
-
 	UPROPERTY(EditAnywhere, BlueprintReadWrite, Category = PhysicsControl)
 	FPhysicsControlData ParentSpaceControlData;
 
 	UPROPERTY(EditAnywhere, BlueprintReadWrite, Category = PhysicsControl)
 	FPhysicsControlSettings ParentSpaceControlSettings;
 
-	UPROPERTY(EditAnywhere, BlueprintReadWrite, Category = PhysicsControl)
-	bool bEnableParentSpaceControls = false;
-
 	UPROPERTY(EditAnywhere, BlueprintReadWrite, Category = PhysicsControl)
 	EPhysicsMovementType PhysicsMovementType = EPhysicsMovementType::Kinematic;
 };
diff --git a/Engine/Plugins/Experimental/PhysicsControl/Source/PhysicsControl/Public/PhysicsControlBPLibrary.h b/Engine/Plugins/Experimental/PhysicsControl/Source/PhysicsControl/Public/PhysicsControlBPLibrary.h
--- a/Engine/Plugins/Experimental/PhysicsControl/Source/PhysicsControl/Public/PhysicsControlBPLibrary.h
+++ b/Engine/Plugins/Experimental/PhysicsControl/Source/PhysicsControl/Public/PhysicsControlBPLibrary.h
@@ -30,73 +30,73 @@ class PHYSICSCONTROL_API UPhysicsControlBPLibrary : public UBlueprintFunctionLib
 	/** Add a single control parameter Parameters to a set of parameter Parameters. */
 	UFUNCTION(BlueprintPure, Category = "Animation|PhysicsControl", meta = (BlueprintThreadSafe, AutoCreateRefTerm = "LinearStrengthMultiplier, LinearExtraDampingMultiplier, AngularStrengthMultiplier, AngularExtraDampingMultiplier, bEnabled"))
 	static void AddControlParameters(
-		UPARAM(ref, DisplayName = "Parameters") FRigidBodyControlAndModifierParameters& InParameters, 
-		UPARAM(DisplayName = "Parameters") FRigidBodyControlAndModifierParameters& OutParameters,
-		const FName                        Name,
-		const FRigidBodyControlSparseData& ControlData);
+		UPARAM(ref, DisplayName = "Parameters") FPhysicsControlControlAndModifierParameters& InParameters, 
+		UPARAM(DisplayName = "Parameters") FPhysicsControlControlAndModifierParameters&      OutParameters,
+		const FName                                                                          Name,
+		const FPhysicsControlSparseData&                                                     ControlData);
 
 	/** Add a single control parameter Parameters to a set of parameter Parameters. */
 	UFUNCTION(BlueprintPure, Category = "Animation|PhysicsControl", meta = (BlueprintThreadSafe, AutoCreateRefTerm = "LinearStrengthMultiplier, LinearExtraDampingMultiplier, AngularStrengthMultiplier, AngularExtraDampingMultiplier, bEnabled"))
 	static void AddMultipleControlParameters(
-		UPARAM(ref, DisplayName = "Parameters") FRigidBodyControlAndModifierParameters& InParameters, 
-		UPARAM(DisplayName = "Parameters") FRigidBodyControlAndModifierParameters& OutParameters,
-		const TArray<FName>&               Names,
-		const FRigidBodyControlSparseData& ControlData);
+		UPARAM(ref, DisplayName = "Parameters") FPhysicsControlControlAndModifierParameters& InParameters, 
+		UPARAM(DisplayName = "Parameters") FPhysicsControlControlAndModifierParameters&      OutParameters,
+		const TArray<FName>&                                                                 Names,
+		const FPhysicsControlSparseData&                                                     ControlData);
 
 	/** Add a single body modifier parameter Parameters to a set of parameter Parameters. */
 	UFUNCTION(BlueprintPure, Category = "Animation|PhysicsControl", meta = (BlueprintThreadSafe))
 	static void AddModifierParameters(
-		UPARAM(ref, DisplayName = "Parameters") FRigidBodyControlAndModifierParameters& InParameters,
-		UPARAM(DisplayName = "Parameters") FRigidBodyControlAndModifierParameters& OutParameters,
-		const FName                         Name, 
-		const FRigidBodyModifierSparseData& ModifierData);
+		UPARAM(ref, DisplayName = "Parameters") FPhysicsControlControlAndModifierParameters& InParameters,
+		UPARAM(DisplayName = "Parameters") FPhysicsControlControlAndModifierParameters&      OutParameters,
+		const FName                                                                          Name, 
+		const FPhysicsControlModifierSparseData&                                             ModifierData);
 
 	/** Add a single body modifier parameter Parameters to a set of parameter Parameters. */
 	UFUNCTION(BlueprintPure, Category = "Animation|PhysicsControl", meta = (BlueprintThreadSafe))
 	static void AddMultipleModifierParameters(
-		UPARAM(ref, DisplayName = "Parameters") FRigidBodyControlAndModifierParameters& InParameters, 
-		UPARAM(DisplayName = "Parameters") FRigidBodyControlAndModifierParameters& OutParameters,
-		const TArray<FName>&                Names,
-		const FRigidBodyModifierSparseData& ModifierData);
+		UPARAM(ref, DisplayName = "Parameters") FPhysicsControlControlAndModifierParameters& InParameters, 
+		UPARAM(DisplayName = "Parameters") FPhysicsControlControlAndModifierParameters&      OutParameters,
+		const TArray<FName>&                                                                 Names,
+		const FPhysicsControlModifierSparseData&                                             ModifierData);
 
 	/** Returns the linear interpolation of two sets of parameter Parameters. Any Parameters that exists in one of the input sets but not the other will be added to the output with a weight of 1. */
 	UFUNCTION(BlueprintPure, Category = "Animation|PhysicsControl", meta = (BlueprintThreadSafe))
 	static void BlendParameters(
-		UPARAM(ref, DisplayName = "From Parameters") FRigidBodyControlAndModifierParameters& InParametersA,
-		UPARAM(ref, DisplayName = "To Parameters") FRigidBodyControlAndModifierParameters& InParametersB,
-		UPARAM(DisplayName = "Weight") const float InInterpolationWeight,
-		UPARAM(DisplayName = "Parameters") FRigidBodyControlAndModifierParameters& OutParameters);
+		UPARAM(ref, DisplayName = "From Parameters") FPhysicsControlControlAndModifierParameters& InParametersA,
+		UPARAM(ref, DisplayName = "To Parameters") FPhysicsControlControlAndModifierParameters&   InParametersB,
+		UPARAM(DisplayName = "Weight") const float                                                InInterpolationWeight,
+		UPARAM(DisplayName = "Parameters") FPhysicsControlControlAndModifierParameters&           OutParameters);
 	
 	/** Adds an Parameters to the output parameters for each control name in the supplied set. The values in each Parameters will be a linear interpolation of the two supplied Parameters, blending from the start Parameters to the end Parameters across the elements in the set. */
 	UFUNCTION(BlueprintPure, Category = "Animation|PhysicsControl", meta = (BlueprintThreadSafe))
 	static void BlendControlParametersThroughSet(
-		UPARAM(ref, DisplayName = "Parameters") FRigidBodyControlAndModifierParameters& InParameters,
-		UPARAM(ref, DisplayName = "Start Parameters") const FRigidBodyNamedControlParameters& InStartControlParameters,
-		UPARAM(ref, DisplayName = "End Parameters") const FRigidBodyNamedControlParameters& InEndControlParameters,
-		const TArray<FName>& ControlNames,
-		FRigidBodyControlAndModifierParameters& OutParameters);
+		UPARAM(ref, DisplayName = "Parameters") FPhysicsControlControlAndModifierParameters&       InParameters,
+		UPARAM(ref, DisplayName = "Start Parameters") const FPhysicsControlNamedControlParameters& InStartControlParameters,
+		UPARAM(ref, DisplayName = "End Parameters") const FPhysicsControlNamedControlParameters&   InEndControlParameters,
+		const TArray<FName>&                                                                       ControlNames,
+		FPhysicsControlControlAndModifierParameters&                                               OutParameters);
 
 	/** Adds an Parameters to the output parameters for each modifier name in the supplied set. The values in each Parameters will be a linear interpolation of the two supplied Parameters, blending from the start Parameters to the end Parameters across the elements in the set. */
 	UFUNCTION(BlueprintPure, Category = "Animation|PhysicsControl", meta = (BlueprintThreadSafe))
 	static void BlendModifierParametersThroughSet(
-		UPARAM(ref) FRigidBodyControlAndModifierParameters& InParameters,
-		UPARAM(ref) const FRigidBodyNamedModifierParameters& InStartModifierParameters,
-		UPARAM(ref) const FRigidBodyNamedModifierParameters& InEndModifierParameters,
-		const TArray<FName>& ModifierNames,
-		UPARAM(DisplayName = "Parameters") FRigidBodyControlAndModifierParameters& OutParameters);
+		UPARAM(ref) FPhysicsControlControlAndModifierParameters&                        InParameters,
+		UPARAM(ref) const FPhysicsControlNamedModifierParameters&                       InStartModifierParameters,
+		UPARAM(ref) const FPhysicsControlNamedModifierParameters&                       InEndModifierParameters,
+		const TArray<FName>&                                                            ModifierNames,
+		UPARAM(DisplayName = "Parameters") FPhysicsControlControlAndModifierParameters& OutParameters);
 
 	/** Get a Rigid Body With Control node reference from an anim node reference */
 	UFUNCTION(BlueprintCallable, Category = "Animation|PhysicsControl", meta = (BlueprintThreadSafe, ExpandEnumAsExecs = "Result"))
 	static FRigidBodyWithControlReference ConvertToRigidBodyWithControl(
-		const FAnimNodeReference& Node,
+		const FAnimNodeReference&           Node,
 		EAnimNodeReferenceConversionResult& Result);
 
 	/** Get a Rigid Body With Control node from an anim node (pure) */
 	UFUNCTION(BlueprintPure, Category = "Animation|PhysicsControl", meta = (BlueprintThreadSafe, DisplayName = "Convert to Rigid Body With Control"))
 	static void ConvertToRigidBodyWithControlPure(
-		const FAnimNodeReference& Node, 
+		const FAnimNodeReference&       Node, 
 		FRigidBodyWithControlReference& RigidBodyWithControl, 
-		bool& Result)
+		bool&                           Result)
 	{
 		EAnimNodeReferenceConversionResult ConversionResult;
 		RigidBodyWithControl = ConvertToRigidBodyWithControl(Node, ConversionResult);
@@ -107,13 +107,13 @@ class PHYSICSCONTROL_API UPhysicsControlBPLibrary : public UBlueprintFunctionLib
 	UFUNCTION(BlueprintPure, Category = "Animation|PhysicsControl", meta = (BlueprintThreadSafe))
 	static TArray<FName> GetControlNamesInSet(
 		UPARAM(DisplayName = "Node") const FRigidBodyWithControlReference& RigidBodyWithControl, 
-		const FName SetName);
+		const FName                                                        SetName);
 
 	/** Get the names of all the body modifiers in a specified set managed by this Rigid Body With Control node. */
 	UFUNCTION(BlueprintPure, Category = "Animation|PhysicsControl", meta = (BlueprintThreadSafe))
 	static TArray<FName> GetBodyModifierNamesInSet(
 		UPARAM(DisplayName = "Node") const FRigidBodyWithControlReference& RigidBodyWithControl, 
-		const FName SetName);
+		const FName                                                        SetName);
 };
 
 #endif
\ No newline at end of file
diff --git a/Engine/Plugins/Experimental/PhysicsControl/Source/PhysicsControl/Public/PhysicsControlComponent.h b/Engine/Plugins/Experimental/PhysicsControl/Source/PhysicsControl/Public/PhysicsControlComponent.h
--- a/Engine/Plugins/Experimental/PhysicsControl/Source/PhysicsControl/Public/PhysicsControlComponent.h
+++ b/Engine/Plugins/Experimental/PhysicsControl/Source/PhysicsControl/Public/PhysicsControlComponent.h
@@ -18,20 +18,6 @@ struct FPhysicsControlRecord;
 struct FPhysicsBodyModifier;
 struct FConstraintInstance;
 
-/**
- * Specifies the type of control that is created when making controls from a skeleton or a set of limbs. 
- * Note that if controls are made individually then other options are available - i.e. in a character, 
- * any body part can be controlled relative to any other part, or indeed any other object.
- */
-UENUM(BlueprintType)
-enum class EPhysicsControlType : uint8
-{
-	/** Control is done in world space, so each object/part is driven independently */
-	WorldSpace,
-	/** Control is done in the space of the parent of each object */
-	ParentSpace,
-};
-
 /**
  * Specifies how any reset to cached target should work. 
  */
@@ -118,8 +104,7 @@ public:
 		const FPhysicsControlData     ControlData,
 		const FPhysicsControlTarget   ControlTarget, 
 		const FPhysicsControlSettings ControlSettings,
-		FName                         Set,
-		const bool                    bEnabled = true
+		FName                         Set
 	);
 
 	/**
@@ -145,8 +130,7 @@ public:
 		const FPhysicsControlData     ControlData, 
 		const FPhysicsControlTarget   ControlTarget, 
 		const FPhysicsControlSettings ControlSettings, 
-		const FName                   Set,
-		const bool                    bEnabled = true);
+		const FName                   Set);
 
 	/**
 	 * Creates a collection of controls controlling a skeletal mesh
@@ -171,8 +155,7 @@ public:
 		const EPhysicsControlType     ControlType,
 		const FPhysicsControlData     ControlData,
 		const FPhysicsControlSettings ControlSettings,
-		const FName                   Set,
-		const bool                    bEnabled = true);
+		const FName                   Set);
 
 	/**
 	 * Creates a collection of ParentSpace controls controlling a skeletal mesh, initializing
@@ -227,8 +210,7 @@ public:
 		const EPhysicsControlType     ControlType,
 		const FPhysicsControlData     ControlData,
 		const FPhysicsControlSettings ControlSettings,
-		const FName                   Set,
-		const bool                    bEnabled = true);
+		const FName                   Set);
 
 	/**
 	 * Creates a collection of ParentSpace controls controlling a skeletal mesh, initializing them 
@@ -295,8 +277,7 @@ public:
 		const TMap<FName, FPhysicsControlLimbBones>& LimbBones,
 		const EPhysicsControlType                    ControlType,
 		const FPhysicsControlData                    ControlData,
-		const FPhysicsControlSettings                ControlSettings,
-		const bool                                   bEnabled = true);
+		const FPhysicsControlSettings                ControlSettings);
 
 	/**
 	 * Creates a collection of ParentSpace controls controlling a skeletal mesh, grouped together in limbs, initializing
@@ -362,7 +343,7 @@ public:
 	 * @return true if the control was found and modified, false if not
 	 */
 	UFUNCTION(BlueprintCallable, Category = PhysicsControl)
-	bool SetControlData(const FName Name, const FPhysicsControlData ControlData, const bool bEnableControl = true);
+	bool SetControlData(const FName Name, const FPhysicsControlData ControlData);
 
 	/**
 	 * Modifies existing control data - i.e. the strengths etc of the controls driving towards the targets
@@ -373,7 +354,7 @@ public:
 	 * @param bEnableControl Enables the control if it is currently disabled
 	 */
 	UFUNCTION(BlueprintCallable, Category = PhysicsControl)
-	void SetControlDatas(const TArray<FName>& Names, const FPhysicsControlData ControlData, const bool bEnableControl = true);
+	void SetControlDatas(const TArray<FName>& Names, const FPhysicsControlData ControlData);
 
 	/**
 	 * Modifies existing control data - i.e. the strengths etc of the controls driving towards the targets
@@ -384,7 +365,7 @@ public:
 	 * @param bEnableControl Enables the controls if currently disabled
 	 */
 	UFUNCTION(BlueprintCallable, Category = PhysicsControl)
-	void SetControlDatasInSet(const FName Set, const FPhysicsControlData ControlData, const bool bEnableControl = true);
+	void SetControlDatasInSet(const FName Set, const FPhysicsControlData ControlData);
 
 	/**
 	 * Modifies an existing control data using the multipliers
@@ -809,29 +790,6 @@ public:
 	UFUNCTION(BlueprintCallable, Category = PhysicsControl)
 	void SetControlsInSetEnabled(const FName Set, const bool bEnable = true);
 
-	/**
-	 * @param Name The name of the control to modify. 
-	 * @param bAutoDisable If set then the control will automatically deactivate after each tick.
-	 * @return true if the control was found and modified, false if not
-	 */
-	UFUNCTION(BlueprintCallable, Category = PhysicsControl)
-	bool SetControlAutoDisable(const FName Name, const bool bAutoDisable);
-
-	/**
-	 * @param Names The names of the controls to modify. 
-	 * @param bAutoDisable If set then the control will automatically deactivate after each tick.
-	 */
-	UFUNCTION(BlueprintCallable, Category = PhysicsControl)
-	void SetControlsAutoDisable(const TArray<FName>& Names, const bool bAutoDisable);
-
-	/**
-	 * @param Set The set of controls to modify. Standard sets will include "All", "WorldSpace",
-	 *        "ParentSpace" and things like "WorldSpace-ArmLeft", depending on how controls have been created.
-	 * @param bAutoDisable If set then the control will automatically deactivate after each tick.
-	 */
-	UFUNCTION(BlueprintCallable, Category = PhysicsControl)
-	void SetControlsInSetAutoDisable(const FName Set, const bool bAutoDisable);
-
 	/**
 	 * @param Name The name of the control to modify. 
 	 * @param bDisableCollision If set then the control will disable collision between the bodies it connects.
@@ -879,13 +837,6 @@ public:
 	UFUNCTION(BlueprintCallable, Category = PhysicsControl)
 	bool GetControlTarget(const FName Name, FPhysicsControlTarget& ControlTarget) const;
 
-	/**
-	 * @param Name        The name of the control to access. 
-	 * @return            Returns true if the control is marked to automatically disable after each tick
-	 */
-	UFUNCTION(BlueprintCallable, Category = PhysicsControl)
-	bool GetControlAutoDisable(const FName Name) const;
-
 	/**
 	 * @param Name        The name of the control to access. 
 	 * @return            Returns true if the control is enabled
@@ -1329,10 +1280,8 @@ public:
 		const TArray<FPhysicsControlLimbSetupData>& LimbSetupData,
 		const FPhysicsControlData                   WorldSpaceControlData,
 		const FPhysicsControlSettings               WorldSpaceControlSettings,
-		const bool                                  bEnableWorldSpaceControls,
 		const FPhysicsControlData                   ParentSpaceControlData,
 		const FPhysicsControlSettings               ParentSpaceControlSettings,
-		const bool                                  bEnableParentSpaceControls,
 		const EPhysicsMovementType                  PhysicsMovementType = EPhysicsMovementType::Static,
 		const float                                 GravityMultiplier = 1.0f,
 		const float                                 PhysicsBlendWeight = 1.0f
diff --git a/Engine/Plugins/Experimental/PhysicsControl/Source/PhysicsControl/Public/PhysicsControlData.h b/Engine/Plugins/Experimental/PhysicsControl/Source/PhysicsControl/Public/PhysicsControlData.h
--- a/Engine/Plugins/Experimental/PhysicsControl/Source/PhysicsControl/Public/PhysicsControlData.h
+++ b/Engine/Plugins/Experimental/PhysicsControl/Source/PhysicsControl/Public/PhysicsControlData.h
@@ -7,6 +7,11 @@
 #include "EngineDefines.h"
 #include "PhysicsControlData.generated.h"
 
+/**
+ * Note that this file defines structures that mostly are, or could be, shared between the PhysicsControlComponent
+ * and the RigidBodyWithControl node.
+ */
+
 class UMeshComponent;
 
 /**
@@ -25,6 +30,284 @@ enum class EPhysicsMovementType : uint8
 	Simulated
 };
 
+inline FName GetPhysicsMovementTypeName(const EPhysicsMovementType MovementType)
+{
+	switch (MovementType)
+	{
+	case EPhysicsMovementType::Static:
+		return "Static";
+	case EPhysicsMovementType::Kinematic:
+		return "Kinematic";
+	case EPhysicsMovementType::Simulated:
+		return "Simulated";
+	}
+	return "None";
+}
+
+/**
+ * Specifies the type of control that is created when making controls from a skeleton or a set of limbs. 
+ * Note that if controls are made individually then other options are available - i.e. in a character, 
+ * any body part can be controlled relative to any other part, or indeed any other object.
+ */
+UENUM(BlueprintType)
+enum class EPhysicsControlType : uint8
+{
+	/** Control is done in world space, so each object/part is driven independently */
+	WorldSpace,
+	/** Control is done in the space of the parent of each object */
+	ParentSpace,
+};
+
+inline FName GetPhysicsControlTypeName(const EPhysicsControlType ControlType)
+{
+	switch (ControlType)
+	{
+	case EPhysicsControlType::WorldSpace:
+		return "WorldSpace";
+	case EPhysicsControlType::ParentSpace:
+		return "ParentSpace";
+	}
+	return "None";
+}
+
+/**
+ * Update an existing set, or add to it
+ */
+USTRUCT(BlueprintType)
+struct PHYSICSCONTROL_API FPhysicsControlSetUpdate
+{
+	GENERATED_BODY();
+
+	UPROPERTY(EditAnywhere, BlueprintReadWrite, Category = PhysicsControl)
+	FName SetName;
+
+	/** The names of either controls or body modifiers (depending on context), or sets of controls/body modifiers */
+	UPROPERTY(EditAnywhere, BlueprintReadWrite, Category = PhysicsControl)
+	TArray<FName> Names;
+};
+
+/**
+ * Combines updates for control and modifier sets
+ */
+USTRUCT(BlueprintType)
+struct PHYSICSCONTROL_API FPhysicsControlSetUpdates
+{
+	GENERATED_BODY();
+
+	UPROPERTY(EditAnywhere, BlueprintReadWrite, Category = PhysicsControl)
+	TArray<FPhysicsControlSetUpdate> ControlSetUpdates;
+
+	UPROPERTY(EditAnywhere, BlueprintReadWrite, Category = PhysicsControl)
+	TArray<FPhysicsControlSetUpdate> ModifierSetUpdates;
+};
+
+/**
+ * Analogous to the ControlData, this indicates how an individual controlled body should move, with flags indicating
+ * whether each element should get used.
+ */
+USTRUCT(BlueprintType)
+struct PHYSICSCONTROL_API FPhysicsControlModifierSparseData
+{
+	GENERATED_BODY();
+
+	FPhysicsControlModifierSparseData(
+		const EPhysicsMovementType InMovementType = EPhysicsMovementType::Simulated,
+		const float                InGravityMultiplier = 1.0f)
+		: MovementType(InMovementType)
+		, GravityMultiplier(InGravityMultiplier)
+		, bEnableMovementType(1)
+		, bEnableGravityMultiplier(1)
+	{
+	}
+
+	UPROPERTY(EditAnywhere, BlueprintReadWrite, Category = PhysicsControl, meta = (editcondition = "bEnableMovementType"))
+	EPhysicsMovementType MovementType;
+
+	UPROPERTY(EditAnywhere, BlueprintReadWrite, Category = PhysicsControl, meta = (editcondition = "bEnableGravityMultiplier"))
+	float GravityMultiplier;
+
+	UPROPERTY(EditAnywhere, Category = PhysicsControl, meta = (InlineEditConditionToggle))
+	uint8 bEnableMovementType : 1;
+
+	UPROPERTY(EditAnywhere, Category = PhysicsControl, meta = (InlineEditConditionToggle))
+	uint8 bEnableGravityMultiplier : 1;
+};
+
+/**
+ * Analogous to the ControlData, this indicates how an individual controlled body should move
+ */
+USTRUCT(BlueprintType)
+struct PHYSICSCONTROL_API FPhysicsControlModifierData
+{
+	GENERATED_BODY();
+
+	FPhysicsControlModifierData(
+		const EPhysicsMovementType InMovementType = EPhysicsMovementType::Simulated,
+		const float                InGravityMultiplier = 1.0f)
+		: MovementType(InMovementType)
+		, GravityMultiplier(InGravityMultiplier)
+	{
+	}
+
+	void UpdateFromSparseData(const FPhysicsControlModifierSparseData& SparseData);
+
+	UPROPERTY(EditAnywhere, BlueprintReadWrite, Category = PhysicsControl)
+	EPhysicsMovementType MovementType;
+
+	UPROPERTY(EditAnywhere, BlueprintReadWrite, Category = PhysicsControl)
+	float GravityMultiplier;
+};
+
+
+/**
+ * Strength and damping etc parameters that will affect a control, with flags indicating
+ * whether each element should get used.
+ */
+USTRUCT(BlueprintType)
+struct PHYSICSCONTROL_API FPhysicsControlSparseData
+{
+	GENERATED_BODY();
+
+	FPhysicsControlSparseData(
+		float InLinearStrength = 0.0f, float InLinearDampingRatio = 1.0f, 
+		float InLinearExtraDamping = 0.0f, float InMaxForce = 0.0f,
+		float InAngularStrength = 0.0f, float InAngularDampingRatio = 1.0f, 
+		float InAngularExtraDamping = 0.0f, float InMaxTorque = 0.0f,
+		float InLinearTargetVelocityMultiplier = 1.0f, float InAngularTargetVelocityMultiplier = 1.0f,
+		bool bInEnabled = true)
+		: LinearStrength(InLinearStrength)
+		, LinearDampingRatio(InLinearDampingRatio)
+		, LinearExtraDamping(InLinearExtraDamping)
+		, MaxForce(InMaxForce)
+		, AngularStrength(InAngularStrength)
+		, AngularDampingRatio(InAngularDampingRatio)
+		, AngularExtraDamping(InAngularExtraDamping)
+		, MaxTorque(InMaxTorque)
+		, LinearTargetVelocityMultiplier(InLinearTargetVelocityMultiplier)
+		, AngularTargetVelocityMultiplier(InAngularTargetVelocityMultiplier)
+		, bEnabled(bInEnabled)
+		, bEnableLinearStrength(true)
+		, bEnableLinearDampingRatio(true)
+		, bEnableLinearExtraDamping(true)
+		, bEnableMaxForce(true)
+		, bEnableAngularStrength(true)
+		, bEnableAngularDampingRatio(true)
+		, bEnableAngularExtraDamping(true)
+		, bEnableMaxTorque(true)
+		, bEnableLinearTargetVelocityMultiplier(true)
+		, bEnableAngularTargetVelocityMultiplier(true)
+		, bEnablebEnabled(true)
+	{
+	}
+
+	/** The strength used to drive linear motion */
+	UPROPERTY(EditAnywhere, BlueprintReadWrite, Category = PhysicsControl, meta = (ClampMin = "0.0", editcondition = "bEnableLinearStrength"))
+	float LinearStrength;
+
+	/** 
+	 * The amount of damping associated with the linear strength. A value of 1 Results in critically 
+	 * damped motion where the control drives as quickly as possible to the target without overshooting. 
+	 * Values > 1 result in more damped motion, and values below 1 result in faster, but more "wobbly" motion.
+	 */
+	UPROPERTY(EditAnywhere, BlueprintReadWrite, Category = PhysicsControl, meta = (ClampMin = "0.0", editcondition = "bEnableLinearDampingRatio"))
+	float LinearDampingRatio;
+
+	/** 
+	 * The amount of additional linear damping. This is added to the damping that comes from LinearDampingRatio
+	 * and can be useful when you want damping even when LinearStrength is zero.
+	 */
+	UPROPERTY(EditAnywhere, BlueprintReadWrite, Category = PhysicsControl, meta = (ClampMin = "0.0", editcondition = "bEnableLinearExtraDamping"))
+	float LinearExtraDamping;
+
+	/** 
+	 * The maximum force used to drive the linear motion. Zero indicates no limit. 
+	 * Note - not yet implemented for RigidBodyWithControl 
+	 */
+	UPROPERTY(EditAnywhere, BlueprintReadWrite, Category = PhysicsControl, meta = (ClampMin = "0.0", editcondition = "bEnableMaxForce"))
+	float MaxForce;
+
+	/** The strength used to drive angular motion */
+	UPROPERTY(EditAnywhere, BlueprintReadWrite, Category = PhysicsControl, meta = (ClampMin = "0.0", editcondition = "bEnableAngularStrength"))
+	float AngularStrength;
+
+	/** 
+	 * The amount of damping associated with the angular strength. A value of 1 Results in critically 
+	 * damped motion where the control drives as quickly as possible to the target without overshooting. 
+	 * Values > 1 result in more damped motion, and values below 1 result in faster, but more "wobbly" motion.
+	 */
+	UPROPERTY(EditAnywhere, BlueprintReadWrite, Category = PhysicsControl, meta = (ClampMin = "0.0", editcondition = "bEnableAngularDampingRatio"))
+	float AngularDampingRatio;
+
+	/** 
+	 * The amount of additional angular damping. This is added to the damping that comes from AngularDampingRatio
+	 * and can be useful when you want damping even when AngularStrength is zero.
+	 */
+	UPROPERTY(EditAnywhere, BlueprintReadWrite, Category = PhysicsControl, meta = (ClampMin = "0.0", editcondition = "bEnableAngularExtraDamping"))
+	float AngularExtraDamping;
+
+	/** 
+	 * The maximum torque used to drive the angular motion. Zero indicates no limit. 
+ 	 * Note - not yet implemented for RigidBodyWithControl 
+	 */
+	UPROPERTY(EditAnywhere, BlueprintReadWrite, Category = PhysicsControl, meta = (ClampMin = "0.0", editcondition = "bEnableMaxTorque"))
+	float MaxTorque;
+
+	/**
+	 * Multiplier on the velocity, which gets applied to the damping. A value of 1 means the animation target
+	 * velocity is used, which helps it track the animation. A value of 0 means damping happens in "world space" 
+	 * - so damping acts like drag on the movement.
+	 */
+	UPROPERTY(EditAnywhere, BlueprintReadWrite, Category = PhysicsControl, meta = (ClampMin = "0.0", editcondition = "bEnableLinearTargetVelocityMultiplier"))
+	float LinearTargetVelocityMultiplier;
+
+	/**
+	 * Multiplier on the angular velocity, which gets applied to the damping. A value of 1 means the animation target
+	 * velocity is used, which helps it track the animation. A value of 0 means damping happens in "world space" 
+	 * - so damping acts like drag on the movement.
+	 */
+	UPROPERTY(EditAnywhere, BlueprintReadWrite, Category = PhysicsControl, meta = (ClampMin = "0.0", editcondition = "bEnableAngularTargetVelocityMultiplier"))
+	float AngularTargetVelocityMultiplier;
+
+	/**
+	 * Whether this control should be enabled
+	 */
+	UPROPERTY(EditAnywhere, BlueprintReadWrite, Category = PhysicsControl, meta = (editcondition = "bEnablebEnabled"))
+	uint8 bEnabled : 1;
+
+	UPROPERTY(EditAnywhere, Category = PhysicsControl, meta = (InlineEditConditionToggle))
+	uint8 bEnableLinearStrength : 1;
+
+	UPROPERTY(EditAnywhere, Category = PhysicsControl, meta = (InlineEditConditionToggle))
+	uint8 bEnableLinearDampingRatio : 1;
+
+	UPROPERTY(EditAnywhere, Category = PhysicsControl, meta = (InlineEditConditionToggle))
+	uint8 bEnableLinearExtraDamping : 1;
+
+	UPROPERTY(EditAnywhere, Category = PhysicsControl, meta = (InlineEditConditionToggle))
+	uint8 bEnableMaxForce : 1;
+
+	UPROPERTY(EditAnywhere, Category = PhysicsControl, meta = (InlineEditConditionToggle))
+	uint8 bEnableAngularStrength : 1;
+
+	UPROPERTY(EditAnywhere, Category = PhysicsControl, meta = (InlineEditConditionToggle))
+	uint8 bEnableAngularDampingRatio : 1;
+
+	UPROPERTY(EditAnywhere, Category = PhysicsControl, meta = (InlineEditConditionToggle))
+	uint8 bEnableAngularExtraDamping : 1;
+
+	UPROPERTY(EditAnywhere, Category = PhysicsControl, meta = (InlineEditConditionToggle))
+	uint8 bEnableMaxTorque : 1;
+
+	UPROPERTY(EditAnywhere, Category = PhysicsControl, meta = (InlineEditConditionToggle))
+	uint8 bEnableLinearTargetVelocityMultiplier : 1;
+
+	UPROPERTY(EditAnywhere, Category = PhysicsControl, meta = (InlineEditConditionToggle))
+	uint8 bEnableAngularTargetVelocityMultiplier : 1;
+
+	UPROPERTY(EditAnywhere, Category = PhysicsControl, meta = (InlineEditConditionToggle))
+	uint8 bEnablebEnabled : 1;
+};
+
 /**
  * Contains data associated with how physical bodies should be controlled/directed towards their targets. 
  * The underlying control is done through damped springs, so the parameters here relate to that.
@@ -43,9 +326,15 @@ struct PHYSICSCONTROL_API FPhysicsControlData
 		, AngularDampingRatio(1.0f)
 		, AngularExtraDamping(0.0f)
 		, MaxTorque(0.0f)
+		, LinearTargetVelocityMultiplier(1.0f)
+		, AngularTargetVelocityMultiplier(1.0f)
+		, bEnabled(true)
 	{
 	}
 
+	/** Applies the values that have been flagged as enabled from the sparse data */
+	void UpdateFromSparseData(const FPhysicsControlSparseData& SparseData);
+
 	/** The strength used to drive linear motion */
 	UPROPERTY(EditAnywhere, BlueprintReadWrite, Category = PhysicsControl, meta = (ClampMin = "0.0"))
 	float LinearStrength;
@@ -65,7 +354,10 @@ struct PHYSICSCONTROL_API FPhysicsControlData
 	UPROPERTY(EditAnywhere, BlueprintReadWrite, Category = PhysicsControl, meta = (ClampMin = "0.0"))
 	float LinearExtraDamping;
 
-	/** The maximum force used to drive the linear motion. Zero indicates no limit. */
+	/** 
+	 * The maximum force used to drive the linear motion. Zero indicates no limit. 
+	 * Note - not yet implemented for RigidBodyWithControl 
+	 */
 	UPROPERTY(EditAnywhere, BlueprintReadWrite, Category = PhysicsControl, meta = (ClampMin = "0.0"))
 	float MaxForce;
 
@@ -88,9 +380,34 @@ struct PHYSICSCONTROL_API FPhysicsControlData
 	UPROPERTY(EditAnywhere, BlueprintReadWrite, Category = PhysicsControl, meta = (ClampMin = "0.0"))
 	float AngularExtraDamping;
 
-	/** The maximum torque used to drive the angular motion. Zero indicates no limit. */
+	/** 
+	 * The maximum torque used to drive the angular motion. Zero indicates no limit. 
+ 	 * Note - not yet implemented for RigidBodyWithControl 
+	 */
 	UPROPERTY(EditAnywhere, BlueprintReadWrite, Category = PhysicsControl, meta = (ClampMin = "0.0"))
 	float MaxTorque;
+
+	/**
+	 * Multiplier on the velocity, which gets applied to the damping. A value of 1 means the animation target
+	 * velocity is used, which helps it track the animation. A value of 0 means damping happens in "world space" 
+	 * - so damping acts like drag on the movement.
+	 */
+	UPROPERTY(EditAnywhere, BlueprintReadWrite, Category = PhysicsControl, meta = (ClampMin = "0.0"))
+	float LinearTargetVelocityMultiplier;
+
+	/**
+	 * Multiplier on the angular velocity, which gets applied to the damping. A value of 1 means the animation target
+	 * velocity is used, which helps it track the animation. A value of 0 means damping happens in "world space" 
+	 * - so damping acts like drag on the movement.
+	 */
+	UPROPERTY(EditAnywhere, BlueprintReadWrite, Category = PhysicsControl, meta = (ClampMin = "0.0"))
+	float AngularTargetVelocityMultiplier;
+
+	/**
+	 * Whether this control should be enabled
+	 */
+	UPROPERTY(EditAnywhere, BlueprintReadWrite, Category = PhysicsControl)
+	uint8 bEnabled : 1;
 };
 
 /**
@@ -212,7 +529,6 @@ struct PHYSICSCONTROL_API FPhysicsControlSettings
 		, bUseSkeletalAnimation(true)
 		, SkeletalAnimationVelocityMultiplier(1.0f)
 		, bDisableCollision(false)
-		, bAutoDisable(false)
 	{
 	}
 
@@ -237,15 +553,6 @@ struct PHYSICSCONTROL_API FPhysicsControlSettings
 	 */
 	UPROPERTY(EditAnywhere, BlueprintReadWrite, Category = PhysicsControl)
 	bool bDisableCollision;
-
-	/**
-	 * Whether or not this control should automatically disable itself at the end of each tick. This
-	 * can be useful when it is more convenient to have a branch that handles some condition (e.g. character 
-	 * is flailing) by updating/enabling the control, and to then have the control automatically get 
-	 * disabled when the branch is no longer taken.
-	 */
-	UPROPERTY(EditAnywhere, BlueprintReadWrite, Category = PhysicsControl)
-	bool bAutoDisable;
 };
 
 /**
diff --git a/Engine/Plugins/Experimental/PhysicsControl/Source/PhysicsControl/Public/PhysicsControlLimbData.h b/Engine/Plugins/Experimental/PhysicsControl/Source/PhysicsControl/Public/PhysicsControlLimbData.h
--- a/Engine/Plugins/Experimental/PhysicsControl/Source/PhysicsControl/Public/PhysicsControlLimbData.h
+++ b/Engine/Plugins/Experimental/PhysicsControl/Source/PhysicsControl/Public/PhysicsControlLimbData.h
@@ -8,20 +8,26 @@
 #include "PhysicsControlLimbData.generated.h"
 
 /**
- * Setup data that is used to create the representation of a single limb. A limb is an array of 
+ * Setup data that are used to create the representation of a single limb. A limb is an array of 
  * contiguous bones (e.g. left arm, or the spine etc). We can define it has the set of bones that
  * are children of a start bone (plus the start bone itself), plus optionally the parent of that start 
  * bone (this is useful when defining the spine, since you will want to include the pelvis, but you 
  * don't to include all the children of the pelvis since that would include the legs), but excluding 
  * any bones that are already part of another limb. This implies limbs should be constructed in order
- * from leaf to root.
+ * from leaf to root (in practice, that normally means defining the spine after other limbs such as 
+ * the arms and legs).
  */
 USTRUCT(BlueprintType)
 struct PHYSICSCONTROL_API FPhysicsControlLimbSetupData
 {
 	GENERATED_BODY()
 
-	/* The name of the limb that this will be used to create */
+	FPhysicsControlLimbSetupData()
+		: bIncludeParentBone(false), bCreateWorldSpaceControls(true)
+		, bCreateParentSpaceControls(true), bCreateBodyModifiers(true)
+	{}
+
+	/** The name of the limb that this will be used to create */
 	UPROPERTY(EditAnywhere, BlueprintReadWrite, Category = PhysicsControl)
 	FName LimbName;
 
@@ -38,7 +44,19 @@ struct PHYSICSCONTROL_API FPhysicsControlLimbSetupData
 	 * of spine_01) in the spine limb.
 	 */
 	UPROPERTY(EditAnywhere, BlueprintReadWrite, Category = PhysicsControl)
-	bool bIncludeParentBone = false;
+	uint8 bIncludeParentBone : 1;
+
+	/** Whether to create-world space controls for this limb */
+	UPROPERTY(EditAnywhere, BlueprintReadWrite, Category = PhysicsControl)
+	uint8 bCreateWorldSpaceControls : 1;
+
+	/** Whether to create parent-space controls for this limb */
+	UPROPERTY(EditAnywhere, BlueprintReadWrite, Category = PhysicsControl)
+	uint8 bCreateParentSpaceControls : 1;
+
+	/** Whether to create body modifiers for this limb */
+	UPROPERTY(EditAnywhere, BlueprintReadWrite, Category = PhysicsControl)
+	uint8 bCreateBodyModifiers : 1;
 };
 
 /**
@@ -49,6 +67,11 @@ struct PHYSICSCONTROL_API FPhysicsControlLimbBones
 {
 	GENERATED_BODY()
 
+	FPhysicsControlLimbBones()
+		: bFirstBoneIsAdditional(false), bCreateWorldSpaceControls(true)
+		, bCreateParentSpaceControls(true), bCreateBodyModifiers(true)
+	{}
+
 	/** The Skeletal mesh that this limb is associated with */
 	TObjectPtr<USkeletalMeshComponent> SkeletalMeshComponent;
 
@@ -56,7 +79,11 @@ struct PHYSICSCONTROL_API FPhysicsControlLimbBones
 	TArray<FName> BoneNames;
 
 	/** Indicates if the first bone in this limb was included due to "IncludeParentBone". */
-	bool bFirstBoneIsAdditional = false;
+	uint8 bFirstBoneIsAdditional : 1;
+
+	uint8 bCreateWorldSpaceControls : 1;
+	uint8 bCreateParentSpaceControls : 1;
+	uint8 bCreateBodyModifiers : 1;
 };
 
 /**
diff --git a/Engine/Plugins/Experimental/PhysicsControl/Source/PhysicsControl/Public/RigidBodyControlData.h b/Engine/Plugins/Experimental/PhysicsControl/Source/PhysicsControl/Public/RigidBodyControlData.h
--- a/Engine/Plugins/Experimental/PhysicsControl/Source/PhysicsControl/Public/RigidBodyControlData.h
+++ b/Engine/Plugins/Experimental/PhysicsControl/Source/PhysicsControl/Public/RigidBodyControlData.h
@@ -2,380 +2,16 @@
 #pragma once
 
 #include "CoreMinimal.h"
+#include "PhysicsControlLimbData.h"
 #include "Physics/ImmediatePhysics/ImmediatePhysicsDeclares.h"
 #include "Animation/AnimTypes.h"
 #include "RigidBodyControlData.generated.h"
 
 /**
- * Setup data that is used to create the representation of a single limb. A limb is an array of 
- * contiguous bones (e.g. left arm, or the spine etc). We can define it has the set of bones that
- * are children of a start bone (plus the start bone itself), plus optionally the parent of that start 
- * bone (this is useful when defining the spine, since you will want to include the pelvis, but you 
- * don't to include all the children of the pelvis since that would include the legs), but excluding 
- * any bones that are already part of another limb. This implies limbs should be constructed in order
- * from leaf to root.
+ * Note that this file defines structures that are only used by the RigidBodyWithControl node. The RBWC node will
+ * also use structures defined in PhysicsControlData.h Some structures defined here may end up being shared with 
+ * PhysicsControlComponent, so have a PhysicsControl prefix in anticipation of that.
  */
-USTRUCT(BlueprintType)
-struct PHYSICSCONTROL_API FRigidBodyLimbSetupData
-{
-	GENERATED_BODY();
-
-	FRigidBodyLimbSetupData()
-		: bIncludeParentBone(false), bCreateWorldSpaceControls(true)
-		, bCreateParentSpaceControls(true), bCreateBodyModifiers(true)
-	{}
-
-	/** The name of the limb that this will be used to create */
-	UPROPERTY(EditAnywhere, BlueprintReadWrite, Category = RigidBodyControl)
-	FName LimbName;
-
-	/* 
-	 * Normally the root-most bone of the limb (e.g. left clavicle when defining the left arm) - so the
-	 * limb will contain children of this bone (plus this bone itself).
-	 */
-	UPROPERTY(EditAnywhere, BlueprintReadWrite, Category = RigidBodyControl)
-	FName StartBone;
-
-	/* 
-	 * Whether or not to include the parent of the start bone. This is intended to be used for limbs like 
-	 * the spine, where you would set StartBone = spine_01 but also expect to include the pelvis (parent 
-	 * of spine_01) in the spine limb.
-	 */
-	UPROPERTY(EditAnywhere, BlueprintReadWrite, Category = RigidBodyControl)
-	uint8 bIncludeParentBone : 1;
-
-	/** Whether to create-world space controls for this limb */
-	UPROPERTY(EditAnywhere, BlueprintReadWrite, Category = RigidBodyControl)
-	uint8 bCreateWorldSpaceControls : 1;
-
-	/** Whether to create parent-space controls for this limb */
-	UPROPERTY(EditAnywhere, BlueprintReadWrite, Category = RigidBodyControl)
-	uint8 bCreateParentSpaceControls : 1;
-
-	/** Whether to create body modifiers for this limb */
-	UPROPERTY(EditAnywhere, BlueprintReadWrite, Category = RigidBodyControl)
-	uint8 bCreateBodyModifiers : 1;
-};
-
-/**
- * Helper structure when creating the limb setup
- */
-struct FRigidBodyLimbBones
-{
-	FRigidBodyLimbBones() 
-		: bFirstBoneIsAdditional(false), bCreateWorldSpaceControls(true)
-		, bCreateParentSpaceControls(true), bCreateBodyModifiers(true)
-	{}
-
-	/** The names of the bones in the limb */
-	TArray<FName> BoneNames;
-
-	/** Indicates if the first bone in this limb was included due to "IncludeParentBone". */
-	uint8 bFirstBoneIsAdditional : 1;
-
-	uint8 bCreateWorldSpaceControls : 1;
-	uint8 bCreateParentSpaceControls : 1;
-	uint8 bCreateBodyModifiers : 1;
-};
-
-/**
- * Used by Body Modifiers to specify how the physical bodies should move.
- */
-UENUM(BlueprintType)
-enum class ERigidBodyMovementType : uint8
-{
-	// Kinematic means that the object won't be simulated, but will be moved according to the
-	// kinematic target set in the Body Modifier.
-	Kinematic,
-	// Simulated means that the object will be controlled by the physics solver
-	Simulated
-};
-FName GetControlTypeName(const ERigidBodyMovementType MovementType);
-
-/**
- * Analogous to the ControlData, this indicates how an individual controlled body should move, with flags indicating
- * whether each element should get used.
- */
-USTRUCT(BlueprintType)
-struct PHYSICSCONTROL_API FRigidBodyModifierSparseData
-{
-	GENERATED_BODY();
-
-	FRigidBodyModifierSparseData(
-		const ERigidBodyMovementType InMovementType = ERigidBodyMovementType::Simulated,
-		const float                  InGravityMultiplier = 1.0f)
-		: MovementType(InMovementType)
-		, GravityMultiplier(InGravityMultiplier)
-		, bEnableMovementType(1)
-		, bEnableGravityMultiplier(1)
-	{
-	}
-
-	UPROPERTY(EditAnywhere, BlueprintReadWrite, Category = PhysicsControl, meta = (editcondition = "bEnableMovementType"))
-	ERigidBodyMovementType MovementType;
-
-	UPROPERTY(EditAnywhere, BlueprintReadWrite, Category = PhysicsControl, meta = (editcondition = "bEnableGravityMultiplier"))
-	float GravityMultiplier;
-
-	UPROPERTY(EditAnywhere, Category = PhysicsControl, meta = (InlineEditConditionToggle))
-	uint8 bEnableMovementType : 1;
-
-	UPROPERTY(EditAnywhere, Category = PhysicsControl, meta = (InlineEditConditionToggle))
-	uint8 bEnableGravityMultiplier : 1;
-};
-
-/**
- * Analogous to the ControlData, this indicates how an individual controlled body should move
- */
-USTRUCT(BlueprintType)
-struct PHYSICSCONTROL_API FRigidBodyModifierData
-{
-	GENERATED_BODY();
-
-	FRigidBodyModifierData(
-		const ERigidBodyMovementType InMovementType = ERigidBodyMovementType::Simulated,
-		const float                  InGravityMultiplier = 1.0f)
-		: MovementType(InMovementType)
-		, GravityMultiplier(InGravityMultiplier)
-	{
-	}
-
-	void UpdateFromSparseData(const FRigidBodyModifierSparseData& SparseData);
-
-	UPROPERTY(EditAnywhere, BlueprintReadWrite, Category = PhysicsControl)
-	ERigidBodyMovementType MovementType;
-
-	UPROPERTY(EditAnywhere, BlueprintReadWrite, Category = PhysicsControl)
-	float GravityMultiplier;
-};
-
-/**
- * Specifies the type of control that is created when making controls from a skeleton or a set of limbs.
- * Note that if controls are made individually then other options are available - i.e. in a character,
- * any body part can be controlled relative to any other part, or indeed any other object.
- */
-UENUM(BlueprintType)
-enum class ERigidBodyControlType : uint8
-{
-	/** Control is done in world space, so each object/part is driven independently */
-	WorldSpace,
-	/** Control is done in the space of the parent of each object */
-	ParentSpace,
-};
-FName GetControlTypeName(const ERigidBodyControlType ControlType);
-
-/**
- * Strength and damping etc parameters that will affect a control, with flags indicating
- * whether each element should get used.
- */
-USTRUCT(BlueprintType)
-struct PHYSICSCONTROL_API FRigidBodyControlSparseData
-{
-	GENERATED_BODY();
-
-	FRigidBodyControlSparseData(
-		float InLinearStrength = 0.0f, float InLinearDampingRatio = 1.0f, float InLinearExtraDamping = 0.0f,
-		float InAngularStrength = 0.0f, float InAngularDampingRatio = 1.0f, float InAngularExtraDamping = 0.0f,
-		float InLinearTargetVelocityMultiplier = 1.0f, float InAngularTargetVelocityMultiplier = 1.0f,
-		bool bInEnabled = true)
-		: LinearStrength(InLinearStrength)
-		, LinearDampingRatio(InLinearDampingRatio)
-		, LinearExtraDamping(InLinearExtraDamping)
-		, AngularStrength(InAngularStrength)
-		, AngularDampingRatio(InAngularDampingRatio)
-		, AngularExtraDamping(InAngularExtraDamping)
-		, LinearTargetVelocityMultiplier(InLinearTargetVelocityMultiplier)
-		, AngularTargetVelocityMultiplier(InAngularTargetVelocityMultiplier)
-		, bEnabled(bInEnabled)
-		, bEnableLinearStrength(true)
-		, bEnableLinearDampingRatio(true)
-		, bEnableLinearExtraDamping(true)
-		, bEnableAngularStrength(true)
-		, bEnableAngularDampingRatio(true)
-		, bEnableAngularExtraDamping(true)
-		, bEnableLinearTargetVelocityMultiplier(true)
-		, bEnableAngularTargetVelocityMultiplier(true)
-		, bEnablebEnabled(true)
-	{
-	}
-
-	/** The strength used to drive linear motion */
-	UPROPERTY(EditAnywhere, BlueprintReadWrite, Category = PhysicsControl, meta = (ClampMin = "0.0", editcondition = "bEnableLinearStrength"))
-	float LinearStrength;
-
-	/** 
-	 * The amount of damping associated with the linear strength. A value of 1 Results in critically 
-	 * damped motion where the control drives as quickly as possible to the target without overshooting. 
-	 * Values > 1 result in more damped motion, and values below 1 result in faster, but more "wobbly" motion.
-	 */
-	UPROPERTY(EditAnywhere, BlueprintReadWrite, Category = PhysicsControl, meta = (ClampMin = "0.0", editcondition = "bEnableLinearDampingRatio"))
-	float LinearDampingRatio;
-
-	/** 
-	 * The amount of additional linear damping. This is added to the damping that comes from LinearDampingRatio
-	 * and can be useful when you want damping even when LinearStrength is zero.
-	 */
-	UPROPERTY(EditAnywhere, BlueprintReadWrite, Category = PhysicsControl, meta = (ClampMin = "0.0", editcondition = "bEnableLinearExtraDamping"))
-	float LinearExtraDamping;
-
-	/** The strength used to drive angular motion */
-	UPROPERTY(EditAnywhere, BlueprintReadWrite, Category = PhysicsControl, meta = (ClampMin = "0.0", editcondition = "bEnableAngularStrength"))
-	float AngularStrength;
-
-	/** 
-	 * The amount of damping associated with the angular strength. A value of 1 Results in critically 
-	 * damped motion where the control drives as quickly as possible to the target without overshooting. 
-	 * Values > 1 result in more damped motion, and values below 1 result in faster, but more "wobbly" motion.
-	 */
-	UPROPERTY(EditAnywhere, BlueprintReadWrite, Category = PhysicsControl, meta = (ClampMin = "0.0", editcondition = "bEnableAngularDampingRatio"))
-	float AngularDampingRatio;
-
-	/** 
-	 * The amount of additional angular damping. This is added to the damping that comes from AngularDampingRatio
-	 * and can be useful when you want damping even when AngularStrength is zero.
-	 */
-	UPROPERTY(EditAnywhere, BlueprintReadWrite, Category = PhysicsControl, meta = (ClampMin = "0.0", editcondition = "bEnableAngularExtraDamping"))
-	float AngularExtraDamping;
-
-	/**
-	 * Multiplier on the velocity, which gets applied to the damping. A value of 1 means the animation target
-	 * velocity is used, which helps it track the animation. A value of 0 means damping happens in "world space" 
-	 * - so damping acts like drag on the movement.
-	 */
-	UPROPERTY(EditAnywhere, BlueprintReadWrite, Category = PhysicsControl, meta = (ClampMin = "0.0", editcondition = "bEnableLinearTargetVelocityMultiplier"))
-	float LinearTargetVelocityMultiplier;
-
-	/**
-	 * Multiplier on the angular velocity, which gets applied to the damping. A value of 1 means the animation target
-	 * velocity is used, which helps it track the animation. A value of 0 means damping happens in "world space" 
-	 * - so damping acts like drag on the movement.
-	 */
-	UPROPERTY(EditAnywhere, BlueprintReadWrite, Category = PhysicsControl, meta = (ClampMin = "0.0", editcondition = "bEnableAngularTargetVelocityMultiplier"))
-	float AngularTargetVelocityMultiplier;
-
-	/**
-	 * Whether this control should be enabled
-	 */
-	UPROPERTY(EditAnywhere, BlueprintReadWrite, Category = PhysicsControl, meta = (editcondition = "bEnablebEnabled"))
-	uint8 bEnabled : 1;
-
-	UPROPERTY(EditAnywhere, Category = PhysicsControl, meta = (InlineEditConditionToggle))
-	uint8 bEnableLinearStrength : 1;
-
-	UPROPERTY(EditAnywhere, Category = PhysicsControl, meta = (InlineEditConditionToggle))
-	uint8 bEnableLinearDampingRatio : 1;
-
-	UPROPERTY(EditAnywhere, Category = PhysicsControl, meta = (InlineEditConditionToggle))
-	uint8 bEnableLinearExtraDamping : 1;
-
-	UPROPERTY(EditAnywhere, Category = PhysicsControl, meta = (InlineEditConditionToggle))
-	uint8 bEnableAngularStrength : 1;
-
-	UPROPERTY(EditAnywhere, Category = PhysicsControl, meta = (InlineEditConditionToggle))
-	uint8 bEnableAngularDampingRatio : 1;
-
-	UPROPERTY(EditAnywhere, Category = PhysicsControl, meta = (InlineEditConditionToggle))
-	uint8 bEnableAngularExtraDamping : 1;
-
-	UPROPERTY(EditAnywhere, Category = PhysicsControl, meta = (InlineEditConditionToggle))
-	uint8 bEnableLinearTargetVelocityMultiplier : 1;
-
-	UPROPERTY(EditAnywhere, Category = PhysicsControl, meta = (InlineEditConditionToggle))
-	uint8 bEnableAngularTargetVelocityMultiplier : 1;
-
-	UPROPERTY(EditAnywhere, Category = PhysicsControl, meta = (InlineEditConditionToggle))
-	uint8 bEnablebEnabled : 1;
-};
-
-
-
-/**
- * Strength and damping etc parameters that will affect a control.
- */
-USTRUCT(BlueprintType)
-struct PHYSICSCONTROL_API FRigidBodyControlData
-{
-	GENERATED_BODY();
-
-	FRigidBodyControlData(
-		float InLinearStrength = 0.0f, float InLinearDampingRatio = 1.0f, float InLinearExtraDamping = 0.0f,
-		float InAngularStrength = 0.0f, float InAngularDampingRatio = 1.0f, float InAngularExtraDamping = 0.0f,
-		float InLinearTargetVelocityMultiplier = 1.0f, float InAngularTargetVelocityMultiplier = 1.0f,
-		bool bInEnabled = true)
-		: LinearStrength(InLinearStrength)
-		, LinearDampingRatio(InLinearDampingRatio)
-		, LinearExtraDamping(InLinearExtraDamping)
-		, AngularStrength(InAngularStrength)
-		, AngularDampingRatio(InAngularDampingRatio)
-		, AngularExtraDamping(InAngularExtraDamping)
-		, LinearTargetVelocityMultiplier(InLinearTargetVelocityMultiplier)
-		, AngularTargetVelocityMultiplier(InAngularTargetVelocityMultiplier)
-		, bEnabled(bInEnabled)
-	{
-	}
-
-	void UpdateFromSparseData(const FRigidBodyControlSparseData& SparseData);
-
-	/** The strength used to drive linear motion */
-	UPROPERTY(EditAnywhere, BlueprintReadWrite, Category = PhysicsControl, meta = (ClampMin = "0.0"))
-	float LinearStrength;
-
-	/** 
-	 * The amount of damping associated with the linear strength. A value of 1 Results in critically 
-	 * damped motion where the control drives as quickly as possible to the target without overshooting. 
-	 * Values > 1 result in more damped motion, and values below 1 result in faster, but more "wobbly" motion.
-	 */
-	UPROPERTY(EditAnywhere, BlueprintReadWrite, Category = PhysicsControl, meta = (ClampMin = "0.0"))
-	float LinearDampingRatio;
-
-	/** 
-	 * The amount of additional linear damping. This is added to the damping that comes from LinearDampingRatio
-	 * and can be useful when you want damping even when LinearStrength is zero.
-	 */
-	UPROPERTY(EditAnywhere, BlueprintReadWrite, Category = PhysicsControl, meta = (ClampMin = "0.0"))
-	float LinearExtraDamping;
-
-	/** The strength used to drive angular motion */
-	UPROPERTY(EditAnywhere, BlueprintReadWrite, Category = PhysicsControl, meta = (ClampMin = "0.0"))
-	float AngularStrength;
-
-	/** 
-	 * The amount of damping associated with the angular strength. A value of 1 Results in critically 
-	 * damped motion where the control drives as quickly as possible to the target without overshooting. 
-	 * Values > 1 result in more damped motion, and values below 1 result in faster, but more "wobbly" motion.
-	 */
-	UPROPERTY(EditAnywhere, BlueprintReadWrite, Category = PhysicsControl, meta = (ClampMin = "0.0"))
-	float AngularDampingRatio;
-
-	/** 
-	 * The amount of additional angular damping. This is added to the damping that comes from AngularDampingRatio
-	 * and can be useful when you want damping even when AngularStrength is zero.
-	 */
-	UPROPERTY(EditAnywhere, BlueprintReadWrite, Category = PhysicsControl, meta = (ClampMin = "0.0"))
-	float AngularExtraDamping;
-
-	/**
-	 * Multiplier on the velocity, which gets applied to the damping. A value of 1 means the animation target
-	 * velocity is used, which helps it track the animation. A value of 0 means damping happens in "world space" 
-	 * - so damping acts like drag on the movement.
-	 */
-	UPROPERTY(EditAnywhere, BlueprintReadWrite, Category = PhysicsControl, meta = (ClampMin = "0.0"))
-	float LinearTargetVelocityMultiplier;
-
-	/**
-	 * Multiplier on the angular velocity, which gets applied to the damping. A value of 1 means the animation target
-	 * velocity is used, which helps it track the animation. A value of 0 means damping happens in "world space" 
-	 * - so damping acts like drag on the movement.
-	 */
-	UPROPERTY(EditAnywhere, BlueprintReadWrite, Category = PhysicsControl, meta = (ClampMin = "0.0"))
-	float AngularTargetVelocityMultiplier;
-
-	/**
-	 * Whether this control should be enabled
-	 */
-	UPROPERTY(EditAnywhere, BlueprintReadWrite, Category = PhysicsControl)
-	uint8 bEnabled : 1;
-};
 
 /**
  * Setup data for all the bodies controlled by the node. Contains info to split the skeleton up into limbs, 
@@ -387,16 +23,16 @@ struct FRigidBodySetupData
 	GENERATED_BODY();
 
 	UPROPERTY(EditAnywhere, Category = ControlSetup)
-	TArray<FRigidBodyLimbSetupData> LimbSetupData;
+	TArray<FPhysicsControlLimbSetupData> LimbSetupData;
 
 	UPROPERTY(EditAnywhere, Category = ControlSetup)
-	FRigidBodyControlData DefaultWorldSpaceControlData;
+	FPhysicsControlData DefaultWorldSpaceControlData;
 
 	UPROPERTY(EditAnywhere, Category = ControlSetup)
-	FRigidBodyControlData DefaultParentSpaceControlData;
+	FPhysicsControlData DefaultParentSpaceControlData;
 
 	UPROPERTY(EditAnywhere, Category = ControlSetup)
-	FRigidBodyModifierData DefaultBodyModifierData;
+	FPhysicsControlModifierData DefaultBodyModifierData;
 };
 
 
@@ -415,7 +51,7 @@ struct PHYSICSCONTROL_API FRigidBodyControl
 	FName ChildBoneName;
 
 	UPROPERTY(EditAnywhere, BlueprintReadWrite, Category = PhysicsControl)
-	FRigidBodyControlData ControlData;
+	FPhysicsControlData ControlData;
 
 	bool IsEnabled() const { return ControlData.bEnabled; }
 };
@@ -448,7 +84,7 @@ struct PHYSICSCONTROL_API FRigidBodyModifier
 	FName BoneName;
 
 	UPROPERTY(EditAnywhere, BlueprintReadWrite, Category = PhysicsControl)
-	FRigidBodyModifierData ModifierData;
+	FPhysicsControlModifierData ModifierData;
 
 	UPROPERTY(EditAnywhere, BlueprintReadWrite, Category = PhysicsControl)
 	TArray<FName> Sets;
@@ -484,53 +120,22 @@ struct PHYSICSCONTROL_API FRigidBodyControlAndBodyModifierCreations
 	TArray<FRigidBodyModifierCreation> Modifiers;
 };
 
-/**
- * Update an existing set, or add to it
- */
-USTRUCT(BlueprintType)
-struct PHYSICSCONTROL_API FSetUpdate
-{
-	GENERATED_BODY();
-
-	UPROPERTY(EditAnywhere, BlueprintReadWrite, Category = PhysicsControl)
-	FName SetName;
-
-	/** The names of either controls or body modifiers (depending on context), or sets of controls/body modifiers */
-	UPROPERTY(EditAnywhere, BlueprintReadWrite, Category = PhysicsControl)
-	TArray<FName> Names;
-};
-
-/**
- * Combines updates for control and modifier sets
- */
-USTRUCT(BlueprintType)
-struct PHYSICSCONTROL_API FRigidBodySetUpdates
-{
-	GENERATED_BODY();
-
-	UPROPERTY(EditAnywhere, BlueprintReadWrite, Category = PhysicsControl)
-	TArray<FSetUpdate> ControlSetUpdates;
-
-	UPROPERTY(EditAnywhere, BlueprintReadWrite, Category = PhysicsControl)
-	TArray<FSetUpdate> ModifierSetUpdates;
-};
-
-FRigidBodyControlData Interpolate(
-	const FRigidBodyControlData& A, const FRigidBodyControlData& B, const float Weight);
-FRigidBodyControlSparseData Interpolate(
-	const FRigidBodyControlSparseData& A, const FRigidBodyControlSparseData& B, const float Weight);
+FPhysicsControlData Interpolate(
+	const FPhysicsControlData& A, const FPhysicsControlData& B, const float Weight);
+FPhysicsControlSparseData Interpolate(
+	const FPhysicsControlSparseData& A, const FPhysicsControlSparseData& B, const float Weight);
 
 /**
  * Data that can be used to parameterize (modify/update) a control 
  */
 USTRUCT(BlueprintType)
-struct PHYSICSCONTROL_API FRigidBodyNamedControlParameters
+struct PHYSICSCONTROL_API FPhysicsControlNamedControlParameters
 {
 	GENERATED_BODY();
 
-	FRigidBodyNamedControlParameters() {}
+	FPhysicsControlNamedControlParameters() {}
 
-	FRigidBodyNamedControlParameters(FName InName, const FRigidBodyControlSparseData& InData) 
+	FPhysicsControlNamedControlParameters(FName InName, const FPhysicsControlSparseData& InData)
 		: Name(InName), Data(InData) {}
 
 	// The name of the control (or set of controls) to update
@@ -538,25 +143,25 @@ struct PHYSICSCONTROL_API FRigidBodyNamedControlParameters
 	FName Name;
 
 	UPROPERTY(EditAnywhere, BlueprintReadWrite, Category = PhysicsControl)
-	FRigidBodyControlSparseData Data;
+	FPhysicsControlSparseData Data;
 };
 
-FRigidBodyModifierData Interpolate(
-	const FRigidBodyModifierData& A, const FRigidBodyModifierData& B, const float Weight);
-FRigidBodyModifierSparseData Interpolate(
-	const FRigidBodyModifierSparseData& A, const FRigidBodyModifierSparseData& B, const float Weight);
+FPhysicsControlModifierData Interpolate(
+	const FPhysicsControlModifierData& A, const FPhysicsControlModifierData& B, const float Weight);
+FPhysicsControlModifierSparseData Interpolate(
+	const FPhysicsControlModifierSparseData& A, const FPhysicsControlModifierSparseData& B, const float Weight);
 
 /**
  * Data that can be used to parameterize(modify / update) a control
  */
 USTRUCT(BlueprintType)
-struct PHYSICSCONTROL_API FRigidBodyNamedModifierParameters
+struct PHYSICSCONTROL_API FPhysicsControlNamedModifierParameters
 {
 	GENERATED_BODY();
 
-	FRigidBodyNamedModifierParameters() {}
+	FPhysicsControlNamedModifierParameters() {}
 
-	FRigidBodyNamedModifierParameters(FName InName, const FRigidBodyModifierSparseData& InData)
+	FPhysicsControlNamedModifierParameters(FName InName, const FPhysicsControlModifierSparseData& InData)
 		: Name(InName), Data(InData) {}
 
 	// The name of the modifier (or set of modifiers) to update
@@ -564,14 +169,14 @@ struct PHYSICSCONTROL_API FRigidBodyNamedModifierParameters
 	FName Name;
 
 	UPROPERTY(EditAnywhere, BlueprintReadWrite, Category = PhysicsControl)
-	FRigidBodyModifierSparseData Data;
+	FPhysicsControlModifierSparseData Data;
 };
 
 /**
  * These apply temporary/ephemeral changes to the controls that only persist for one tick.
  */
 USTRUCT(BlueprintType)
-struct PHYSICSCONTROL_API FRigidBodyControlAndModifierParameters
+struct PHYSICSCONTROL_API FPhysicsControlControlAndModifierParameters
 {
 	GENERATED_BODY();
 
@@ -581,17 +186,17 @@ struct PHYSICSCONTROL_API FRigidBodyControlAndModifierParameters
 	 * subsequent entries will override earlier ones if they apply to the same control).
 	 */ 
 	UPROPERTY(EditAnywhere, BlueprintReadWrite, Category = Settings)
-	TArray<FRigidBodyNamedControlParameters> ControlParameters;
+	TArray<FPhysicsControlNamedControlParameters> ControlParameters;
 	
 	/**
 	 *  Parameters for existing modifiers. Each name can be the name of a modifier, or the name of a 
 	 * set of modifiers. They will only apply for one tick/update.
 	 */
 	UPROPERTY(EditAnywhere, BlueprintReadWrite, Category = Settings)
-	TArray<FRigidBodyNamedModifierParameters> ModifierParameters;
+	TArray<FPhysicsControlNamedModifierParameters> ModifierParameters;
 
-	void Add(const FRigidBodyNamedControlParameters& InParameters) { ControlParameters.Add(InParameters); }
-	void Add(const FRigidBodyNamedModifierParameters& InParameters) { ModifierParameters.Add(InParameters); }
+	void Add(const FPhysicsControlNamedControlParameters& InParameters) { ControlParameters.Add(InParameters); }
+	void Add(const FPhysicsControlNamedModifierParameters& InParameters) { ModifierParameters.Add(InParameters); }
 };
 
 /**
@@ -599,17 +204,17 @@ struct PHYSICSCONTROL_API FRigidBodyControlAndModifierParameters
  * (apart from the actual bodies that are being controlled/affected)
  */
 USTRUCT(BlueprintType)
-struct PHYSICSCONTROL_API FRigidBodyControlAndModifierUpdates
+struct PHYSICSCONTROL_API FPhysicsControlControlAndModifierUpdates
 {
 	GENERATED_BODY();
 
 	/** Modifications to the underlying controls - these will persist */
 	UPROPERTY(EditAnywhere, BlueprintReadWrite, Category = PhysicsControl)
-	TArray<FRigidBodyNamedControlParameters> ControlParameters;
+	TArray<FPhysicsControlNamedControlParameters> ControlParameters;
 
 	/** Modifications to the underlying modifiers - these will persist */
 	UPROPERTY(EditAnywhere, BlueprintReadWrite, Category = PhysicsControl)
-	TArray<FRigidBodyNamedModifierParameters> ModifierParameters;
+	TArray<FPhysicsControlNamedModifierParameters> ModifierParameters;
 };
 
 /**
@@ -712,9 +317,9 @@ struct PHYSICSCONTROL_API FRigidBodyKinematicTargets
  * In addition, it will also store some updates to these original data which will determine how the actual
  * controls get used.
  */
-struct FControlRecord
+struct FRigidBodyControlRecord
 {
-	FControlRecord(const FRigidBodyControl& InControl, ImmediatePhysics::FJointHandle* InJointHandle);
+	FRigidBodyControlRecord(const FRigidBodyControl& InControl, ImmediatePhysics::FJointHandle* InJointHandle);
 
 	void ResetCurrent(bool bResetTarget);
 
@@ -727,7 +332,7 @@ struct FControlRecord
 
 	// This contains the currently active control data. It will be updated just prior to
 	// applying the controls, by setting it to the default, and then updating it with any parameters.
-	FRigidBodyControlData CurrentData;
+	FPhysicsControlData CurrentData;
 
 	// Contains any control target that has been set
 	FRigidBodyControlTarget ControlTarget;
@@ -746,9 +351,9 @@ struct FControlRecord
  * In addition, it will also store some updates to these original data which will determine how the actual
  * body properties get set.
  */
-struct FBodyModifierRecord
+struct FRigidBodyModifierRecord
 {
-	FBodyModifierRecord(const FRigidBodyModifier& InModifier, ImmediatePhysics::FActorHandle* InActorHandle);
+	FRigidBodyModifierRecord(const FRigidBodyModifier& InModifier, ImmediatePhysics::FActorHandle* InActorHandle);
 
 	void ResetCurrent();
 
@@ -758,7 +363,7 @@ struct FBodyModifierRecord
 
 	// This contains the currently active modifier data. It will be updated just prior to
 	// applying the controls, by setting it to the default, and then updating it with any parameters.
-	FRigidBodyModifierData CurrentData;
+	FPhysicsControlModifierData CurrentData;
 };
 
 
diff --git a/Engine/Plugins/Experimental/PhysicsControl/Source/PhysicsControlUncookedOnly/Private/AnimGraphNode_RigidBodyWithControl.cpp b/Engine/Plugins/Experimental/PhysicsControl/Source/PhysicsControlUncookedOnly/Private/AnimGraphNode_RigidBodyWithControl.cpp
--- a/Engine/Plugins/Experimental/PhysicsControl/Source/PhysicsControlUncookedOnly/Private/AnimGraphNode_RigidBodyWithControl.cpp
+++ b/Engine/Plugins/Experimental/PhysicsControl/Source/PhysicsControlUncookedOnly/Private/AnimGraphNode_RigidBodyWithControl.cpp
@@ -103,9 +103,9 @@ void UAnimGraphNode_RigidBodyWithControl::Draw(FPrimitiveDrawInterface* PDI, USk
 		// Draw Space Controls
 		if (bIsSelected || (bIsPoseWatchEnabled && PoseWatchElementWorldSpaceControls.IsValid() && PoseWatchElementWorldSpaceControls->GetIsVisible()))
 		{
-			for (const TMap<FName, FControlRecord>::ElementType& NameRecordPair : RuntimeRigidBodyNode->ControlRecords)
+			for (const TMap<FName, FRigidBodyControlRecord>::ElementType& NameRecordPair : RuntimeRigidBodyNode->ControlRecords)
 			{
-				const FControlRecord& ControlRecord = NameRecordPair.Value;
+				const FRigidBodyControlRecord& ControlRecord = NameRecordPair.Value;
 				if (!ControlRecord.JointHandle || !ControlRecord.IsEnabled())
 				{
 					continue;
