[UE] Editor tools for working with generated control names in the PhysicsControl plugin
- On request, creates a new window in the AnimGraph editor that shows the generated control and body modifier names for every RigidBodyWithControl node in every open blueprint
- Names can be copied from the window onto the clipboard
- The names in the wndow are generated using the same code path as the node uses at runtime
- Includes changes to the way names are generated at runtime

diff --git a/Engine/Plugins/Experimental/PhysicsControl/Source/PhysicsControl/Private/AnimNode_RigidBodyWithControl.cpp b/Engine/Plugins/Experimental/PhysicsControl/Source/PhysicsControl/Private/AnimNode_RigidBodyWithControl.cpp
--- a/Engine/Plugins/Experimental/PhysicsControl/Source/PhysicsControl/Private/AnimNode_RigidBodyWithControl.cpp
+++ b/Engine/Plugins/Experimental/PhysicsControl/Source/PhysicsControl/Private/AnimNode_RigidBodyWithControl.cpp
@@ -143,7 +143,7 @@ FAnimNode_RigidBodyWithControl::FAnimNode_RigidBodyWithControl()
 	, bEnabled(false)
 	, bSimulationStarted(false)
 	, bCheckForBodyTransformInit(false)
-	, bHaveRunControlSeup(false)
+	, bHaveSetupControls(false)
 	, SimulationTiming(ESimulationTiming::Default)
 	, WorldTimeSeconds(0.0f)
 	, LastEvalTimeSeconds(0.0f)
@@ -489,6 +489,29 @@ static FTransform CalculateJointTargetTransform(
 	return ParentConnectorTM.Inverse() * ChildConnectorTM;
 }
 
+void FAnimNode_RigidBodyWithControl::SetupControls(USkeletalMeshComponent* const SkeletalMeshComponent)
+{
+	bool bSuccess = false;
+
+	if (SkeletalMeshComponent)
+	{
+		if (USkeletalMesh* const SkeletalMesh = SkeletalMeshComponent->GetSkeletalMeshAsset())
+		{
+			CreateWorldSpaceControlRootBody(PhysicsAssetToUse);
+			InitControlsAndBodyModifiers(SkeletalMesh->GetRefSkeleton());
+			bHaveSetupControls = true;
+		}
+		else
+		{
+			UE_LOG(LogRigidBodyWithControl, Warning, TEXT("Invalid Skeletal Mesh"));
+		}
+	}
+	else
+	{
+		UE_LOG(LogRigidBodyWithControl, Warning, TEXT("Invalid Skeletal Mesh Component"));
+	}
+}
+
 
 void FAnimNode_RigidBodyWithControl::EvaluateSkeletalControl_AnyThread(FComponentSpacePoseContext& Output, TArray<FBoneTransform>& OutBoneTransforms)
 {
@@ -511,11 +534,9 @@ void FAnimNode_RigidBodyWithControl::EvaluateSkeletalControl_AnyThread(FComponen
 		FlushDeferredSimulationTask();
 
 		// Handle deferred control creation
-		if (!bHaveRunControlSeup && bEnableControls)
+		if (!bHaveSetupControls && bEnableControls)
 		{
-			CreateWorldSpaceControlRootBody(PhysicsAssetToUse);
-			bHaveRunControlSeup = true;
-			InitControlsAndBodyModifiers(Output.AnimInstanceProxy->GetSkelMeshComponent());
+			SetupControls(Output.AnimInstanceProxy->GetSkelMeshComponent());
 		}
 
 		const FBoneContainer& BoneContainer = Output.Pose.GetPose().GetBoneContainer();
@@ -1234,11 +1255,9 @@ void FAnimNode_RigidBodyWithControl::InitPhysics(const UAnimInstance* InAnimInst
 
 		//== WithControl ==================================
 
-		if (!bHaveRunControlSeup && bEnableControls)
+		if (!bHaveSetupControls && bEnableControls)
 		{
-			CreateWorldSpaceControlRootBody(PhysicsAssetToUse);
-			bHaveRunControlSeup = true;
-			InitControlsAndBodyModifiers(InAnimInstance->GetSkelMeshComponent());
+			SetupControls(InAnimInstance->GetSkelMeshComponent());
 		}
 
 		//== WithControl ==================================
diff --git a/Engine/Plugins/Experimental/PhysicsControl/Source/PhysicsControl/Private/PhysicsControlLog.cpp b/Engine/Plugins/Experimental/PhysicsControl/Source/PhysicsControl/Private/PhysicsControlLog.cpp
--- a/Engine/Plugins/Experimental/PhysicsControl/Source/PhysicsControl/Private/PhysicsControlLog.cpp
+++ b/Engine/Plugins/Experimental/PhysicsControl/Source/PhysicsControl/Private/PhysicsControlLog.cpp
@@ -2,5 +2,5 @@
 
 #include "PhysicsControlLog.h"
 
-DEFINE_LOG_CATEGORY(LogPhysicsControlComponent);
-DEFINE_LOG_CATEGORY(LogRigidBodyWithControl);
+PHYSICSCONTROL_API DEFINE_LOG_CATEGORY(LogPhysicsControlComponent);
+PHYSICSCONTROL_API DEFINE_LOG_CATEGORY(LogRigidBodyWithControl);
diff --git a/Engine/Plugins/Experimental/PhysicsControl/Source/PhysicsControl/Private/PhysicsControlLog.h b/Engine/Plugins/Experimental/PhysicsControl/Source/PhysicsControl/Private/PhysicsControlLog.h
deleted file mode 100644
diff --git a/Engine/Plugins/Experimental/PhysicsControl/Source/PhysicsControl/Private/PhysicsControlOperatorNameGeneration.cpp b/Engine/Plugins/Experimental/PhysicsControl/Source/PhysicsControl/Private/PhysicsControlOperatorNameGeneration.cpp
new file mode 100644
--- /dev/null
+++ b/Engine/Plugins/Experimental/PhysicsControl/Source/PhysicsControl/Private/PhysicsControlOperatorNameGeneration.cpp
@@ -0,0 +1,378 @@
+// Copyright Epic Games, Inc. All Rights Reserved.
+
+#include "PhysicsControlOperatorNameGeneration.h"
+
+#include "ReferenceSkeleton.h"
+#include "PhysicsEngine/PhysicsAsset.h"
+#include "AnimNode_RigidBodyWithControl.h"
+#include "PhysicsControlLimbData.h"
+
+//======================================================================================================================
+template<typename TOperatorCreationType> FName FindUniqueOperatorName(const TOperatorCreationType& OperatorCreationSpecifier, const TSet<FName>& ExistingNames)
+{
+	return NAME_None;
+}
+
+//======================================================================================================================
+template<> FName FindUniqueOperatorName<FRigidBodyModifierCreation>(const FRigidBodyModifierCreation& OperatorCreationSpecifier, const TSet<FName>& ExistingNames)
+{
+	return FindUniqueBodyModifierName(OperatorCreationSpecifier.Modifier.BoneName, ExistingNames);;
+}
+
+//======================================================================================================================
+template<> FName FindUniqueOperatorName<FRigidBodyControlCreation>(const FRigidBodyControlCreation& OperatorCreationSpecifier, const TSet<FName>& ExistingNames)
+{
+	return FindUniqueControlName(OperatorCreationSpecifier.Control.ParentBoneName, OperatorCreationSpecifier.Control.ChildBoneName, ExistingNames);
+}
+
+//======================================================================================================================
+FName FindParentBodyBoneName(const FName BoneName, const FReferenceSkeleton& RefSkeleton, UPhysicsAsset* const PhysicsAsset)
+{
+	FName ParentBoneName;
+
+	if (PhysicsAsset)
+	{
+		const int32 StartBoneIndex = RefSkeleton.FindBoneIndex(BoneName);
+
+		if (StartBoneIndex != INDEX_NONE)
+		{
+			const int32 ParentBodyIndex = PhysicsAsset->FindParentBodyIndex(RefSkeleton, StartBoneIndex);
+
+			if (PhysicsAsset->SkeletalBodySetups.IsValidIndex(ParentBodyIndex) && PhysicsAsset->SkeletalBodySetups[ParentBodyIndex])
+			{
+				ParentBoneName = PhysicsAsset->SkeletalBodySetups[ParentBodyIndex]->BoneName;
+			}
+		}
+	}
+
+	return ParentBoneName;
+}
+
+//======================================================================================================================
+template<typename TOperatorFunctorType> void CreateAdditionalBodyModifiers(const TArray<FRigidBodyModifierCreation> CreationSpecifiers, FRigidBodyNameRecords& NameRecords, TOperatorFunctorType& OperatorFunctor)
+{
+	for (const FRigidBodyModifierCreation& Specifier : CreationSpecifiers)
+	{
+		const FName Name = OperatorFunctor(Specifier.Modifier.BoneName, Specifier.Modifier.ModifierData);
+		if (Name.IsNone())
+		{
+			UE_LOG(LogRigidBodyWithControl, Warning,
+				TEXT("CreateAdditionalControls: Failed to make body modifier for %s"), *Specifier.Modifier.BoneName.ToString());
+		}
+		else
+		{
+			UE_LOG(LogRigidBodyWithControl, Verbose,
+				TEXT("Made modifier %s for %s"),
+				*Name.ToString(),
+				*Specifier.Modifier.BoneName.ToString());
+
+			NameRecords.AddBodyModifier(Name, Specifier.Sets);
+		}
+	}
+}
+
+//======================================================================================================================
+template<typename TOperatorFunctorType> void CreateAdditionalControls(const TArray<FRigidBodyControlCreation> CreationSpecifiers, FRigidBodyNameRecords& NameRecords, TOperatorFunctorType& OperatorFunctor)
+{
+	for (const FRigidBodyControlCreation& Specifier : CreationSpecifiers)
+	{
+		const FName Name = OperatorFunctor(Specifier.Control.ParentBoneName, Specifier.Control.ChildBoneName, Specifier.Control.ControlData);
+		if (Name.IsNone())
+		{
+			UE_LOG(LogRigidBodyWithControl, Warning,
+				TEXT("CreateAdditionalControls: Failed to make control between %s and %s"),
+				*Specifier.Control.ParentBoneName.ToString(), *Specifier.Control.ChildBoneName.ToString());
+		}
+		else
+		{
+			UE_LOG(LogRigidBodyWithControl, Verbose,
+				TEXT("Made control %s between %s and %s"),
+				*Name.ToString(),
+				*Specifier.Control.ParentBoneName.ToString(),
+				*Specifier.Control.ChildBoneName.ToString());
+			NameRecords.AddControl(Name, Specifier.Sets);
+		}
+	}
+}
+
+
+//======================================================================================================================
+template<typename TFunctor> void CreateControlsFromLimbBones(
+	const FName                   LimbName,
+	const FPhysicsControlLimbBones& LimbBones,
+	const EPhysicsControlType   ControlType,
+	const FReferenceSkeleton& RefSkeleton,
+	UPhysicsAsset* const PhysicsAsset,
+	const FPhysicsControlData& ControlData,
+	FRigidBodyNameRecords& NameRecords,
+	TFunctor& CreateOperationLambda)
+{
+	for (const FName ChildBoneName : LimbBones.BoneNames)
+	{
+		FName ParentBoneName;
+
+		if (ControlType == EPhysicsControlType::ParentSpace)
+		{
+			ParentBoneName = FindParentBodyBoneName(ChildBoneName, RefSkeleton, PhysicsAsset);
+
+			if (ParentBoneName.IsNone())
+			{
+				// This happens for the pelvis, for example - we only create parent space
+				// controls if there's a parent!
+				continue;
+			}
+		}
+
+		const FName ControlName = CreateOperationLambda(ParentBoneName, ChildBoneName, ControlData);
+
+		if (!ControlName.IsNone())
+		{
+			NameRecords.AddControl(ControlName, LimbName);
+			NameRecords.AddControl(ControlName, GetPhysicsControlTypeName(ControlType));
+			NameRecords.AddControl(ControlName, FName(
+				GetPhysicsControlTypeName(ControlType).ToString().Append("_").Append(LimbName.ToString())));
+		}
+		else if (!ParentBoneName.IsNone())
+		{
+			UE_LOG(LogRigidBodyWithControl, Warning,
+				TEXT("Failed to find unique control name for %s and %s"), *ParentBoneName.ToString(), *ChildBoneName.ToString());
+		}
+		else
+		{
+			UE_LOG(LogRigidBodyWithControl, Warning,
+				TEXT("Failed to find unique control name for %s"), *ChildBoneName.ToString());
+		}
+	}
+}
+
+//======================================================================================================================
+template<typename TFunctor> void CreateBodyModifiersFromLimbBones(
+	const FName                   LimbName,
+	const FPhysicsControlLimbBones& LimbBones,
+	const FPhysicsControlModifierData& ModifierData,
+	FRigidBodyNameRecords& NameRecords,
+	TFunctor& CreateOperationLambda)
+{
+	for (const FName BoneName : LimbBones.BoneNames)
+	{
+		const FName BodyModifierName = CreateOperationLambda(BoneName, ModifierData);
+
+		if (!BodyModifierName.IsNone())
+		{
+			NameRecords.AddBodyModifier(BodyModifierName, LimbName);
+		}
+		else
+		{
+			UE_LOG(LogRigidBodyWithControl, Warning,
+				TEXT("Failed to find unique body modifier name for %s"), *BoneName.ToString());
+		}
+	}
+}
+
+//======================================================================================================================
+template<typename TControlFunctorType, typename TBodyModifierFunctorType> void ForEachPotentialOperator(const FAnimNode_RigidBodyWithControl* const Node, TMap<FName, FPhysicsControlLimbBones> AllLimbBones, const FReferenceSkeleton& RefSkeleton, UPhysicsAsset* const PhysicsAsset, FRigidBodyNameRecords& NameRecords, TControlFunctorType& ControlFunctor, TBodyModifierFunctorType& BodyModifierFunctor)
+{
+	for (const TMap<FName, FPhysicsControlLimbBones>::ElementType& LimbBoneEntry : AllLimbBones)
+	{
+		const FName LimbName = LimbBoneEntry.Key;
+		const FPhysicsControlLimbBones& LimbBones = LimbBoneEntry.Value;
+		if (LimbBones.bCreateWorldSpaceControls)
+		{
+			CreateControlsFromLimbBones(LimbName, LimbBones, EPhysicsControlType::WorldSpace, RefSkeleton, PhysicsAsset, Node->SetupData.DefaultWorldSpaceControlData, NameRecords, ControlFunctor);
+		}
+		if (LimbBones.bCreateParentSpaceControls)
+		{
+			CreateControlsFromLimbBones(LimbName, LimbBones, EPhysicsControlType::ParentSpace, RefSkeleton, PhysicsAsset, Node->SetupData.DefaultParentSpaceControlData, NameRecords, ControlFunctor);
+		}
+		if (LimbBones.bCreateBodyModifiers)
+		{
+			CreateBodyModifiersFromLimbBones(LimbName, LimbBones, Node->SetupData.DefaultBodyModifierData, NameRecords, BodyModifierFunctor);
+		}
+	}
+
+	if (Node != nullptr)
+	{
+		// Find names for any additional controls that have been requested	
+		CreateAdditionalBodyModifiers(Node->AdditionalControlsAndBodyModifiers.Modifiers, NameRecords, BodyModifierFunctor);
+		CreateAdditionalControls(Node->AdditionalControlsAndBodyModifiers.Controls, NameRecords, ControlFunctor);
+	}
+}
+
+//======================================================================================================================
+// Slightly annoying to have to add the names individually, but we want to check they exist
+template<typename TBodyModifierNameContainerType, typename TControlNameContainerType> void CreateAdditionalSets_Implimentation(const FPhysicsControlSetUpdates& AdditionalSets, const TBodyModifierNameContainerType& BodyModifierNames, const TControlNameContainerType& ControlNames, FRigidBodyNameRecords& NameRecords)
+{
+	for (const FPhysicsControlSetUpdate& Set : AdditionalSets.ControlSetUpdates)
+	{
+		TArray<FName> Names = ExpandName(Set.Names, NameRecords.ControlSets);
+
+		for (FName Name : Names)
+		{
+			if (ControlNames.Contains(Name))
+			{
+				NameRecords.AddControl(Name, Set.SetName);
+			}
+			else
+			{
+				UE_LOG(LogRigidBodyWithControl, Warning,
+					TEXT("CreateAdditionalSets: Failed to find control with name %s to add to set %s"),
+					*Name.ToString(), *Set.SetName.ToString());
+			}
+		}
+	}
+
+	for (const FPhysicsControlSetUpdate& Set : AdditionalSets.ModifierSetUpdates)
+	{
+		TArray<FName> Names = ExpandName(Set.Names, NameRecords.BodyModifierSets);
+
+		for (FName Name : Names)
+		{
+			if (BodyModifierNames.Contains(Name))
+			{
+				NameRecords.AddBodyModifier(Name, Set.SetName);
+			}
+			else
+			{
+				UE_LOG(LogRigidBodyWithControl, Warning,
+					TEXT("CreateAdditionalSets: Failed to find body modifier with name %s to add to set %s"),
+					*Name.ToString(), *Set.SetName.ToString());
+			}
+		}
+	}
+}
+
+//======================================================================================================================
+FName FindUniqueName(const FString& NameBase, const TSet<FName>& Keys, const int32 MaxNmeaIndex)
+{
+	FString NameStr = NameBase;
+
+	// If the number gets too large, almost certainly we're in some nasty situation where this is
+	// getting called in a loop. Better to quit and fail, rather than allow the constraint set to
+	// increase without bound. 
+	for (int32 Index = 0; Index < MaxNmeaIndex; ++Index)
+	{
+		const FName Name(NameStr);
+		if (!Keys.Find(Name))
+		{
+			return Name;
+		}
+		NameStr = FString::Format(TEXT("{0}_{1}"), { NameBase, Index });
+	}
+
+	return NAME_None;
+}
+
+//======================================================================================================================
+FName FindUniqueBodyModifierName(const FName BodyName, const TSet<FName>& ExistingNames)
+{
+	FString NameBase = TEXT("");
+	if (!BodyName.IsNone())
+	{
+		NameBase += BodyName.ToString();
+	}
+	else
+	{
+		NameBase = TEXT("Body");
+	}
+
+	return FindUniqueName(NameBase, ExistingNames, MaxNumControlsOrModifiersPerName);
+}
+
+//======================================================================================================================
+FName FindUniqueControlName(const FName ParentBodyName, const FName ChildBodyName, const TSet<FName>& ExistingNames)
+{
+	FString NameBase = TEXT("");
+	if (!ParentBodyName.IsNone())
+	{
+		NameBase += ParentBodyName.ToString() + TEXT("_");
+	}
+	if (!ChildBodyName.IsNone())
+	{
+		NameBase += ChildBodyName.ToString();
+	}
+
+	return FindUniqueName(NameBase, ExistingNames, MaxNumControlsOrModifiersPerName);
+}
+
+//======================================================================================================================
+TMap<FName, FPhysicsControlLimbBones> GetLimbBones(
+	const TArray<FPhysicsControlLimbSetupData>& LimbSetupData,
+	const FReferenceSkeleton& RefSkeleton,
+	UPhysicsAsset* const PhysicsAsset)
+{
+	// TODO - Output limb bones are not in the order specified in the skeleton - would be better if they were
+
+	TMap<FName, FPhysicsControlLimbBones> Result;
+
+	if (!PhysicsAsset)
+	{
+		UE_LOG(LogRigidBodyWithControl, Warning, TEXT("Physics asset missing"));
+		return Result;
+	}
+
+	TSet<FName> AllBones;
+
+	for (const FPhysicsControlLimbSetupData& LimbSetup : LimbSetupData)
+	{
+		FPhysicsControlLimbBones& LimbBones = Result.Add(LimbSetup.LimbName);
+
+		LimbBones.bFirstBoneIsAdditional = false;
+		LimbBones.bCreateWorldSpaceControls = LimbSetup.bCreateWorldSpaceControls;
+		LimbBones.bCreateParentSpaceControls = LimbSetup.bCreateParentSpaceControls;
+		LimbBones.bCreateBodyModifiers = LimbSetup.bCreateBodyModifiers;
+
+		if (LimbSetup.bIncludeParentBone)
+		{
+			const FName ParentBoneName = FindParentBodyBoneName(LimbSetup.StartBone, RefSkeleton, PhysicsAsset);
+
+			if (!ParentBoneName.IsNone())
+			{
+				LimbBones.BoneNames.Add(ParentBoneName);
+				AllBones.Add(ParentBoneName);
+				LimbBones.bFirstBoneIsAdditional = true;
+			}
+		}
+
+		TArray<int32> ChildBodyIndices;
+		PhysicsAsset->GetBodyIndicesBelow(ChildBodyIndices, LimbSetup.StartBone, RefSkeleton);
+
+		for (int32 ChildBodyIndex : ChildBodyIndices)
+		{
+			const FName BoneName = PhysicsAsset->SkeletalBodySetups[ChildBodyIndex]->BoneName;
+
+			if (!AllBones.Find(BoneName))
+			{
+				LimbBones.BoneNames.Add(BoneName);
+				AllBones.Add(BoneName);
+			}
+		}
+	}
+
+	return Result;
+}
+
+void CollectOperatorNames(const FAnimNode_RigidBodyWithControl* const Node, TMap<FName, FPhysicsControlLimbBones> AllLimbBones, const FReferenceSkeleton& RefSkeleton, UPhysicsAsset* const PhysicsAsset, TSet<FName>& BodyModifierNames, TSet<FName>& ControlNames, FRigidBodyNameRecords& NameRecords)
+{
+	auto CollectControlName = [&ControlNames](const FName ParentBoneName, const FName ChildBoneName, const FPhysicsControlData& Data) -> FName { const FName Name = FindUniqueControlName(ParentBoneName, ChildBoneName, ControlNames); ControlNames.Add(Name); return Name; };
+	auto CollectBodyModifierName = [&BodyModifierNames](const FName BoneName, const FPhysicsControlModifierData& Data) -> FName { const FName Name = FindUniqueBodyModifierName(BoneName, BodyModifierNames); BodyModifierNames.Add(Name); return Name; };
+
+	ForEachPotentialOperator(Node, AllLimbBones, RefSkeleton, PhysicsAsset, NameRecords, CollectControlName, CollectBodyModifierName);
+}
+
+void CreateOperatorsForNode(FAnimNode_RigidBodyWithControl* const Node, TMap<FName, FPhysicsControlLimbBones> AllLimbBones, const FReferenceSkeleton& RefSkeleton, UPhysicsAsset* const PhysicsAsset, FRigidBodyNameRecords& NameRecords)
+{
+	auto CreateControl = [Node](const FName ParentBoneName, const FName ChildBoneName, const FPhysicsControlData& Data) -> FName { return Node->CreateControl(ParentBoneName, ChildBoneName, Data); };
+	auto CreateBodyModifier = [Node](const FName BoneName, const FPhysicsControlModifierData& Data) -> FName { return Node->CreateBodyModifier(BoneName, Data); };
+
+	ForEachPotentialOperator(Node, AllLimbBones, RefSkeleton, PhysicsAsset, NameRecords, CreateControl, CreateBodyModifier);
+}
+
+void CreateAdditionalSets(const FPhysicsControlSetUpdates& AdditionalSets, const TSet<FName>& BodyModifierNames, const TSet<FName>& ControlNames, FRigidBodyNameRecords& NameRecords)
+{
+	CreateAdditionalSets_Implimentation(AdditionalSets, BodyModifierNames, ControlNames, NameRecords);
+}
+
+void CreateAdditionalSets(const FPhysicsControlSetUpdates& AdditionalSets, const TMap<FName, FRigidBodyModifierRecord>& BodyModifierNames, const TMap<FName, FRigidBodyControlRecord>& ControlNames, FRigidBodyNameRecords& NameRecords)
+{
+	CreateAdditionalSets_Implimentation(AdditionalSets, BodyModifierNames, ControlNames, NameRecords);
+}
diff --git a/Engine/Plugins/Experimental/PhysicsControl/Source/PhysicsControl/Private/RigidBodyNameRecords.cpp b/Engine/Plugins/Experimental/PhysicsControl/Source/PhysicsControl/Private/RigidBodyNameRecords.cpp
--- a/Engine/Plugins/Experimental/PhysicsControl/Source/PhysicsControl/Private/RigidBodyNameRecords.cpp
+++ b/Engine/Plugins/Experimental/PhysicsControl/Source/PhysicsControl/Private/RigidBodyNameRecords.cpp
@@ -111,7 +111,7 @@ void FRigidBodyNameRecords::Reset()
 }
 
 //======================================================================================================================
-TArray<FName> ExpandSetName(const FName InName, const TMap<FName, TArray<FName>>& SetNames)
+TArray<FName> ExpandName(const FName InName, const TMap<FName, TArray<FName>>& SetNames)
 {
 	TArray<FName> OutputNames;
 	if (const TArray<FName>* const FoundSet = SetNames.Find(InName))
@@ -120,19 +120,19 @@ TArray<FName> ExpandSetName(const FName InName, const TMap<FName, TArray<FName>>
 	}
 	else
 	{
-		OutputNames.AddUnique(InName);
+		OutputNames.Add(InName);
 	}
 	return OutputNames;
 }
 
 //======================================================================================================================
-TArray<FName> ExpandSetNames(const TArray<FName>& InNames, const TMap<FName, TArray<FName>>& SetNames)
+TArray<FName> ExpandName(const TArray<FName>& InNames, const TMap<FName, TArray<FName>>& SetNames)
 {
 	TArray<FName> OutputNames;
 
 	for (const FName Name : InNames)
 	{
-		OutputNames.Append(ExpandSetName(Name, SetNames));
+		OutputNames.Append(ExpandName(Name, SetNames));
 	}
 
 	return OutputNames;
diff --git a/Engine/Plugins/Experimental/PhysicsControl/Source/PhysicsControl/Private/RigidBodyWithControl.cpp b/Engine/Plugins/Experimental/PhysicsControl/Source/PhysicsControl/Private/RigidBodyWithControl.cpp
--- a/Engine/Plugins/Experimental/PhysicsControl/Source/PhysicsControl/Private/RigidBodyWithControl.cpp
+++ b/Engine/Plugins/Experimental/PhysicsControl/Source/PhysicsControl/Private/RigidBodyWithControl.cpp
@@ -14,6 +14,7 @@
 #include "PhysicsEngine/PhysicsSettings.h"
 #include "Chaos/ChaosConstraintSettings.h"
 #include "PhysicsControlLog.h"
+#include "PhysicsControlOperatorNameGeneration.h"
 
 //UE_DISABLE_OPTIMIZATION;
 
@@ -22,17 +23,22 @@ DECLARE_CYCLE_STAT(TEXT("RigidBodyNodeWithControl_LogControlsModifiersAndSets"),
 DECLARE_CYCLE_STAT(TEXT("RigidBodyNodeWithControl_ApplyControlAndModifierUpdatesAndParametersToRecords"), STAT_RigidBodyNodeWithControl_ApplyControlAndModifierUpdatesAndParametersToRecords, STATGROUP_Anim);
 DECLARE_CYCLE_STAT(TEXT("RigidBodyNodeWithControl_ApplyControlsAndModifierDatas"), STAT_RigidBodyNodeWithControl_ApplyControlsAndModifierDatas, STATGROUP_Anim);
 DECLARE_CYCLE_STAT(TEXT("RigidBodyNodeWithControl_ApplyControlsAndModifiers"), STAT_RigidBodyNodeWithControl_ApplyControlsAndModifiers, STATGROUP_Anim);
+
+DECLARE_CYCLE_STAT(TEXT("RigidBodyNodeWithControl_ApplyControls"), STAT_RigidBodyNodeWithControl_ApplyControls, STATGROUP_Anim);
+DECLARE_CYCLE_STAT(TEXT("RigidBodyNodeWithControl_ApplyModifiers"), STAT_RigidBodyNodeWithControl_ApplyModifiers, STATGROUP_Anim);
+DECLARE_CYCLE_STAT(TEXT("RigidBodyNodeWithControl_SetConstraintEnabled"), STAT_RigidBodyNodeWithControl_SetConstraintEnabled, STATGROUP_Anim);
+
 DECLARE_CYCLE_STAT(TEXT("RigidBodyNodeWithControl_ApplyKinematicTargets"), STAT_RigidBodyNodeWithControl_ApplyKinematicTargets, STATGROUP_Anim);
 DECLARE_CYCLE_STAT(TEXT("RigidBodyNodeWithControl_ApplyCurrentConstraintProfile"), STAT_RigidBodyNodeWithControl_ApplyCurrentConstraintProfile, STATGROUP_Anim);
 
-int32 MaxNumControlsOrModifiersPerName = 16;
-
 constexpr int32 ConstraintChildIndex = 0;
 constexpr int32 ConstraintParentIndex = 1;
 
 //======================================================================================================================
 static void SetConstraintEnabled(ImmediatePhysics::FJointHandle* const JointHandle, const bool bIsEnabled)
 {
+	SCOPE_CYCLE_COUNTER(STAT_RigidBodyNodeWithControl_SetConstraintEnabled);
+
 	if (JointHandle)
 	{
 		if (ImmediatePhysics::FJointHandle::FChaosConstraintHandle* ConstraintHandle = JointHandle->GetConstraint())
@@ -196,24 +202,6 @@ ImmediatePhysics::FActorHandle* FAnimNode_RigidBodyWithControl::FindBodyFromBone
 	return BodyHandle;
 }
 
-//======================================================================================================================
-FName FAnimNode_RigidBodyWithControl::FindParentBoneNameFromBoneName(const FName BoneName) const
-{
-	// TODO - would be better if we didn't need to access the physics structures to find the parent bone name.
-	FName ParentName;
-	const int32 JointIndex = FindBodyIndexFromBoneName(BoneName);
-	if (JointIndex != INDEX_NONE)
-	{
-		ImmediatePhysics::FJointHandle* const JointHandle = Joints[JointIndex];
-
-		if (JointHandle && JointHandle->GetActorHandles()[ConstraintParentIndex])
-		{
-			ParentName = JointHandle->GetActorHandles()[ConstraintParentIndex]->GetName();
-		}
-	}
-	return ParentName;
-}
-
 //======================================================================================================================
 FName FAnimNode_RigidBodyWithControl::CreateControl(
 	const FName ParentBoneName, const FName ChildBoneName, const FPhysicsControlData& ControlData)
@@ -224,7 +212,6 @@ FName FAnimNode_RigidBodyWithControl::CreateControl(
 	Control.ControlData = ControlData;
 
 	ImmediatePhysics::FJointHandle* JointHandle = nullptr;
-
 	ImmediatePhysics::FActorHandle* ParentBodyHandle;
 
 	if (ParentBoneName.IsNone())
@@ -244,51 +231,13 @@ FName FAnimNode_RigidBodyWithControl::CreateControl(
 			TEXT("Unable to create world space control constraint for bone %s"), *ChildBoneName.ToString());
 		return FName();
 	}
-	
+
 	FName ControlName = GetUniqueControlName(Control.ParentBoneName, Control.ChildBoneName);
 	ControlRecords.Add(ControlName, FRigidBodyControlRecord(Control, JointHandle));
 
 	return ControlName;
 }
 
-//======================================================================================================================
-void FAnimNode_RigidBodyWithControl::CreateControlsFromLimbBones(
-	const FName                     LimbName,
-	const FPhysicsControlLimbBones& LimbBones,
-	const EPhysicsControlType       ControlType,
-	const FPhysicsControlData&      ControlData)
-{
-	for (const FName ChildBoneName : LimbBones.BoneNames)
-	{
-		FName ParentBoneName;
-
-		if (ControlType == EPhysicsControlType::ParentSpace)
-		{
-			ParentBoneName = FindParentBoneNameFromBoneName(ChildBoneName);
-			if (ParentBoneName.IsNone())
-			{
-				// This happens for the pelvis, for example - we only create parent space
-				// controls if there's a parent!
-				continue;
-			}
-		}
-
-		FName ControlName = CreateControl(ParentBoneName, ChildBoneName, ControlData);
-		if (ControlName.IsNone())
-		{
-			UE_LOG(LogRigidBodyWithControl, Warning,
-				TEXT("Failed to make control for %s"), *ChildBoneName.ToString());
-		}
-		else
-		{
-			NameRecords.AddControl(ControlName, LimbName);
-			NameRecords.AddControl(ControlName, GetPhysicsControlTypeName(ControlType));
-			NameRecords.AddControl(ControlName, FName(
-				GetPhysicsControlTypeName(ControlType).ToString().Append("_").Append(LimbName.ToString())));
-		}
-	}
-}
-
 //======================================================================================================================
 FName FAnimNode_RigidBodyWithControl::GetBodyFromBoneName(const FName BoneName) const
 {
@@ -300,86 +249,37 @@ FName FAnimNode_RigidBodyWithControl::GetBodyFromBoneName(const FName BoneName)
 }
 
 //======================================================================================================================
-FName FAnimNode_RigidBodyWithControl::GetUniqueBodyModifierName(const FName BoneName)
+FName FAnimNode_RigidBodyWithControl::GetUniqueBodyModifierName(const FName BoneName) const
 {
-	FName BodyName = GetBodyFromBoneName(BoneName);
-
-	FString NameBase = TEXT("");
-	if (!BodyName.IsNone())
-	{
-		NameBase += BodyName.ToString();
-	}
-	else
-	{
-		NameBase = TEXT("Body");
-	}
-
 	TSet<FName> Keys;
 	ModifierRecords.GetKeys(Keys);
+	const FName UniqueName = FindUniqueBodyModifierName(GetBodyFromBoneName(BoneName), Keys);
 
-	if (!Keys.Find(FName(NameBase)))
+	if (UniqueName.IsNone())
 	{
-		return FName(NameBase);
+		UE_LOG(LogRigidBodyWithControl, Warning,
+			TEXT("Unable to find a suitable Body Modifier name - the limit of MaxNumControlsOrModifiersPerName (%d) has been exceeded"),
+			MaxNumControlsOrModifiersPerName);
 	}
 
-	// If the number gets too large, almost certainly we're in some nasty situation where this is
-	// getting called in a loop. Better to quit and fail, rather than allow the modifier set to
-	// increase without bound. 
-	for (int32 Index = 0; Index != MaxNumControlsOrModifiersPerName; ++Index)
-	{
-		FString NameStr = FString::Format(TEXT("{0}_{1}"), { NameBase, Index });
-		FName Name(NameStr);
-		if (!Keys.Find(Name))
-		{
-			return Name;
-		}
-	}
-	UE_LOG(LogRigidBodyWithControl, Warning,
-		TEXT("Unable to find a suitable Body Modifier name - the limit of MaxNumControlsOrModifiersPerName (%d) has been exceeded"),
-		MaxNumControlsOrModifiersPerName);
-	return FName();
+	return UniqueName;
 }
 
 //======================================================================================================================
-FName FAnimNode_RigidBodyWithControl::GetUniqueControlName(const FName ParentBoneName, const FName ChildBoneName)
+FName FAnimNode_RigidBodyWithControl::GetUniqueControlName(const FName ParentBoneName, const FName ChildBoneName) const
 {
-	FName ParentBodyName = GetBodyFromBoneName(ParentBoneName);
-	FName ChildBodyName = GetBodyFromBoneName(ChildBoneName);
-
-	FString NameBase = TEXT("");
-	if (!ParentBodyName.IsNone())
-	{
-		NameBase += ParentBodyName.ToString() + TEXT("_");
-	}
-	if (!ChildBodyName.IsNone())
-	{
-		NameBase += ChildBodyName.ToString();
-	}
-
 	TSet<FName> Keys;
 	ControlRecords.GetKeys(Keys);
+	const FName UniqueName = FindUniqueControlName(GetBodyFromBoneName(ParentBoneName), GetBodyFromBoneName(ChildBoneName), Keys);
 
-	if (!Keys.Find(FName(NameBase)))
+	if (UniqueName.IsNone())
 	{
-		return FName(NameBase);
+		UE_LOG(LogRigidBodyWithControl, Warning,
+			TEXT("Unable to find a suitable Control name - the limit of MaxNumControlsOrModifiersPerName (%d) has been exceeded"),
+			MaxNumControlsOrModifiersPerName);
 	}
 
-	// If the number gets too large, almost certainly we're in some nasty situation where this is
-	// getting called in a loop. Better to quit and fail, rather than allow the constraint set to
-	// increase without bound. 
-	for (int32 Index = 0; Index < MaxNumControlsOrModifiersPerName; ++Index)
-	{
-		FString NameStr = FString::Format(TEXT("{0}_{1}"), { NameBase, Index });
-		FName Name(NameStr);
-		if (!Keys.Find(Name))
-		{
-			return Name;
-		}
-	}
-	UE_LOG(LogRigidBodyWithControl, Warning,
-		TEXT("Unable to find a suitable Control name - the limit of MaxNumControlsOrModifiersPerName (%d) has been exceeded"),
-		MaxNumControlsOrModifiersPerName);
-	return FName();
+	return UniqueName;
 }
 
 //======================================================================================================================
@@ -402,105 +302,7 @@ FName FAnimNode_RigidBodyWithControl::CreateBodyModifier(FName BoneName, const F
 }
 
 //======================================================================================================================
-void FAnimNode_RigidBodyWithControl::CreateBodyModifiersFromLimbBones(
-	const FName                        LimbName,
-	const FPhysicsControlLimbBones&    LimbBones,
-	const FPhysicsControlModifierData& DefaultModifierData)
-{
-	for (const FName BoneName : LimbBones.BoneNames)
-	{
-		const FName BodyModifierName = CreateBodyModifier(BoneName, DefaultModifierData);
-		if (BodyModifierName.IsNone())
-		{
-			UE_LOG(LogRigidBodyWithControl, Warning,
-				TEXT("Failed to make body modifier for %s"), *BoneName.ToString());
-		}
-		else
-		{
-			NameRecords.AddBodyModifier(BodyModifierName, LimbName);
-		}
-	}
-}
-
-//======================================================================================================================
-TMap<FName, FPhysicsControlLimbBones> FAnimNode_RigidBodyWithControl::GetLimbBonesFromSkeletalMesh(
-	const TArray<FPhysicsControlLimbSetupData>& LimbSetupData,
-	USkeletalMeshComponent* const               SkeletalMeshComponent) const
-{
-	// Parse the skeleton tree to figure out which bones are associated with which limbs. 
-	
-	// TODO - Output limb bones are not in the order specified in the skeleton - would be better if they were
-	
-	TMap<FName, FPhysicsControlLimbBones> Result;
-
-	UPhysicsAsset* const PhysicsAsset = GetPhysicsAsset();
-
-	if (!PhysicsAsset)
-	{
-		UE_LOG(LogRigidBodyWithControl, Warning, TEXT("Physics asset missing"));
-		return Result;
-	}
-
-	if (!SkeletalMeshComponent)
-	{
-		UE_LOG(LogRigidBodyWithControl, Warning, TEXT("Invalid Skeletal Mesh Component"));
-		return Result;
-	}
-
-	USkeletalMesh* const SkelMesh = SkeletalMeshComponent->GetSkeletalMeshAsset();
-
-	TSet<FName> AllBones;
-
-	for (const FPhysicsControlLimbSetupData& LimbSetup : LimbSetupData)
-	{
-		FPhysicsControlLimbBones& LimbBones = Result.Add(LimbSetup.LimbName);
-
-		LimbBones.bFirstBoneIsAdditional = false;
-		LimbBones.bCreateWorldSpaceControls = LimbSetup.bCreateWorldSpaceControls;
-		LimbBones.bCreateParentSpaceControls = LimbSetup.bCreateParentSpaceControls;
-		LimbBones.bCreateBodyModifiers = LimbSetup.bCreateBodyModifiers;
-
-		if (LimbSetup.bIncludeParentBone)
-		{		
-			const int32 StartBoneIndex = SkelMesh->GetRefSkeleton().FindBoneIndex(LimbSetup.StartBone);
-			const int32 ParentBodyIndex = PhysicsAsset->FindParentBodyIndex(SkelMesh, StartBoneIndex);
-			
-			FName ParentBoneName;
-
-			if (PhysicsAsset->SkeletalBodySetups.IsValidIndex(ParentBodyIndex) && 
-				PhysicsAsset->SkeletalBodySetups[ParentBodyIndex])
-			{
-				ParentBoneName = PhysicsAsset->SkeletalBodySetups[ParentBodyIndex]->BoneName;
-			}
-
-			if (!ParentBoneName.IsNone())
-			{
-				LimbBones.BoneNames.Add(ParentBoneName);
-				AllBones.Add(ParentBoneName);
-				LimbBones.bFirstBoneIsAdditional = true;
-			}
-		}
-
-		TArray<int32> ChildBodyIndices;
-		PhysicsAsset->GetBodyIndicesBelow(ChildBodyIndices, LimbSetup.StartBone, SkelMesh);
-
-		for (int32 ChildBodyIndex : ChildBodyIndices)
-		{
-			const FName BoneName = PhysicsAsset->SkeletalBodySetups[ChildBodyIndex]->BoneName;
-
-			if (!AllBones.Find(BoneName))
-			{
-				LimbBones.BoneNames.Add(BoneName);
-				AllBones.Add(BoneName);
-			}
-		}
-	}
-
-	return Result;
-}
-
-//======================================================================================================================
-void FAnimNode_RigidBodyWithControl::InitControlsAndBodyModifiers(USkeletalMeshComponent* const SkeletalMeshComponent)
+void FAnimNode_RigidBodyWithControl::InitControlsAndBodyModifiers(const FReferenceSkeleton& RefSkeleton)
 {
 	SCOPE_CYCLE_COUNTER(STAT_RigidBodyNodeWithControl_InitControlsAndBodyModifiers);
 
@@ -508,33 +310,12 @@ void FAnimNode_RigidBodyWithControl::InitControlsAndBodyModifiers(USkeletalMeshC
 
 	// These functions will create the base set of controls and modifiers from SetupData
 	TMap<FName, FPhysicsControlLimbBones> AllLimbBones =
-		GetLimbBonesFromSkeletalMesh(SetupData.LimbSetupData, SkeletalMeshComponent);
-
-	for (const TMap<FName, FPhysicsControlLimbBones>::ElementType& LimbBoneEntry : AllLimbBones)
-	{
-		const FName LimbName = LimbBoneEntry.Key;
-		const FPhysicsControlLimbBones& LimbBones = LimbBoneEntry.Value;
-		if (LimbBones.bCreateWorldSpaceControls)
-		{
-			CreateControlsFromLimbBones(
-				LimbName, LimbBones, EPhysicsControlType::WorldSpace, SetupData.DefaultWorldSpaceControlData);
-		}
-		if (LimbBones.bCreateParentSpaceControls)
-		{
-			CreateControlsFromLimbBones(
-				LimbName, LimbBones, EPhysicsControlType::ParentSpace, SetupData.DefaultParentSpaceControlData);
-		}
-		if (LimbBones.bCreateBodyModifiers)
-		{
-			CreateBodyModifiersFromLimbBones(LimbName, LimbBones, SetupData.DefaultBodyModifierData);
-		}
-	}
+		GetLimbBones(SetupData.LimbSetupData, RefSkeleton, GetPhysicsAsset());
 
-	// Create any additional controls that have been requested
-	CreateAdditionalControls();
+	CreateOperatorsForNode(this, AllLimbBones, RefSkeleton, GetPhysicsAsset(), NameRecords);
 
 	// Create any additional sets that have been requested
-	CreateAdditionalSets();
+	CreateAdditionalSets(AdditionalSets, ModifierRecords, ControlRecords, NameRecords);
 
 	// Apply control and modifier parameters on a single-use basis
 	ApplyControlAndBodyModifierDatas(
@@ -554,7 +335,7 @@ void FAnimNode_RigidBodyWithControl::DestroyControlsAndBodyModifiers()
 	ModifierRecords.Reset();
 	NameRecords.Reset();
 
-	bHaveRunControlSeup = false;
+	bHaveSetupControls = false;
 }
 
 //======================================================================================================================
@@ -620,94 +401,6 @@ void FAnimNode_RigidBodyWithControl::LogControlsModifiersAndSets()
 #undef RBWC_LOG_LEVEL
 }
 
-
-//======================================================================================================================
-void FAnimNode_RigidBodyWithControl::CreateAdditionalControls()
-{
-	// Create additional controls.
-	for (const FRigidBodyControlCreation& Control : AdditionalControlsAndBodyModifiers.Controls)
-	{	
-		FName ControlName = CreateControl(
-			Control.Control.ParentBoneName, Control.Control.ChildBoneName, Control.Control.ControlData);
-		if (ControlName.IsNone())
-		{
-			UE_LOG(LogRigidBodyWithControl, Warning,
-				TEXT("CreateAdditionalControls: Failed to make control between %s and %s"), 
-				*Control.Control.ParentBoneName.ToString(), *Control.Control.ChildBoneName.ToString());
-		}
-		else
-		{
-			UE_LOG(LogRigidBodyWithControl, Verbose,
-				TEXT("Made control %s between %s and %s"),
-				*ControlName.ToString(), 
-				*Control.Control.ParentBoneName.ToString(), 
-				*Control.Control.ChildBoneName.ToString());
-			NameRecords.AddControl(ControlName, Control.Sets);
-		}
-	}
-
-	// Create additional modifiers.
-	for (const FRigidBodyModifierCreation& Modifier : AdditionalControlsAndBodyModifiers.Modifiers)
-	{
-		FName ModifierName = CreateBodyModifier(Modifier.Modifier.BoneName, Modifier.Modifier.ModifierData);
-		if (ModifierName.IsNone())
-		{
-			UE_LOG(LogRigidBodyWithControl, Warning,
-				TEXT("CreateAdditionalControls: Failed to make body modifier for %s"), *Modifier.Modifier.BoneName.ToString());
-		}
-		else
-		{
-			UE_LOG(LogRigidBodyWithControl, Verbose,
-				TEXT("Made modifier %s for %s"),
-				*ModifierName.ToString(),
-				*Modifier.Modifier.BoneName.ToString());
-			NameRecords.AddControl(ModifierName, Modifier.Sets);
-		}
-	}
-}
-
-//======================================================================================================================
-// Slightly annoying to have to add the names individually, but we want to check they exist
-void FAnimNode_RigidBodyWithControl::CreateAdditionalSets()
-{
-	for (const FPhysicsControlSetUpdate& Set : AdditionalSets.ControlSetUpdates)
-	{
-		TArray<FName> Names = ExpandSetNames(Set.Names, NameRecords.ControlSets);
-		for (FName Name : Names)
-		{
-			if (ControlRecords.Find(Name))
-			{
-				NameRecords.AddControl(Name, Set.SetName);
-			}
-			else
-			{
-				UE_LOG(LogRigidBodyWithControl, Warning,
-					TEXT("CreateAdditionalSets: Failed to find control with name %s to add to set %s"), 
-					*Name.ToString(), *Set.SetName.ToString());
-			}
-		}
-	}
-
-	for (const FPhysicsControlSetUpdate& Set : AdditionalSets.ModifierSetUpdates)
-	{
-		TArray<FName> Names = ExpandSetNames(Set.Names, NameRecords.BodyModifierSets);
-		for (FName Name : Names)
-		{
-			if (ModifierRecords.Find(Name))
-			{
-				NameRecords.AddBodyModifier(Name, Set.SetName);
-			}
-			else
-			{
-				UE_LOG(LogRigidBodyWithControl, Warning,
-					TEXT("CreateAdditionalSets: Failed to find body modifier with name %s to add to set %s"),
-					*Name.ToString(), *Set.SetName.ToString());
-			}
-		}
-	}
-}
-
-
 //======================================================================================================================
 template<typename TOut>
 static void ConvertStrengthToSpringParams(
@@ -930,7 +623,7 @@ void FAnimNode_RigidBodyWithControl::ApplyControlAndBodyModifierDatas(
 	// This updates the "original" controls and modifiers based on the parameters.
 	for (const FPhysicsControlNamedControlParameters& ControlParameters : InControlParameters)
 	{
-		TArray<FName> Names = ExpandSetName(ControlParameters.Name, NameRecords.ControlSets);
+		TArray<FName> Names = ExpandName(ControlParameters.Name, NameRecords.ControlSets);
 		for (FName Name : Names)
 		{
 			const FPhysicsControlSparseData& ControlData = ControlParameters.Data;
@@ -948,7 +641,7 @@ void FAnimNode_RigidBodyWithControl::ApplyControlAndBodyModifierDatas(
 
 	for (const FPhysicsControlNamedModifierParameters& ModifierParameters : InModifierParameters)
 	{
-		TArray<FName> Names = ExpandSetName(ModifierParameters.Name, NameRecords.BodyModifierSets);
+		TArray<FName> Names = ExpandName(ModifierParameters.Name, NameRecords.BodyModifierSets);
 		for (FName Name : Names)
 		{
 			const FPhysicsControlModifierSparseData& ModifierData = ModifierParameters.Data;
@@ -990,21 +683,28 @@ void FAnimNode_RigidBodyWithControl::ApplyControlsAndModifiers(const FVector& Si
 	if (!PoseData.IsEmpty())
 	{
 		// Apply Controls.
-		for (TMap<FName, FRigidBodyControlRecord>::ElementType& NameRecordPair : ControlRecords)
 		{
-			FRigidBodyControlRecord& ControlRecord = NameRecordPair.Value;
-			if (ControlRecord.IsEnabled())
+			SCOPE_CYCLE_COUNTER(STAT_RigidBodyNodeWithControl_ApplyControls);
+
+			for (TMap<FName, FRigidBodyControlRecord>::ElementType& NameRecordPair : ControlRecords)
 			{
-				ApplyControl(ControlRecord, DeltaTime);
-			}
+				FRigidBodyControlRecord& ControlRecord = NameRecordPair.Value;
+				if (ControlRecord.IsEnabled())
+				{
+					ApplyControl(ControlRecord, DeltaTime);
+				}
 
-			SetConstraintEnabled(ControlRecord.JointHandle, ControlRecord.IsEnabled());
+				SetConstraintEnabled(ControlRecord.JointHandle, ControlRecord.IsEnabled());
+			}
 		}
 
 		// Apply Body Modifiers.
-		for (const TMap<FName, FRigidBodyModifierRecord>::ElementType& NameRecordPair : ModifierRecords)
 		{
-			ApplyModifier(NameRecordPair.Value, SimSpaceGravity);
+			SCOPE_CYCLE_COUNTER(STAT_RigidBodyNodeWithControl_ApplyModifiers);
+			for (const TMap<FName, FRigidBodyModifierRecord>::ElementType& NameRecordPair : ModifierRecords)
+			{
+				ApplyModifier(NameRecordPair.Value, SimSpaceGravity);
+			}
 		}
 	}
 }
diff --git a/Engine/Plugins/Experimental/PhysicsControl/Source/PhysicsControl/Public/AnimNode_RigidBodyWithControl.h b/Engine/Plugins/Experimental/PhysicsControl/Source/PhysicsControl/Public/AnimNode_RigidBodyWithControl.h
--- a/Engine/Plugins/Experimental/PhysicsControl/Source/PhysicsControl/Public/AnimNode_RigidBodyWithControl.h
+++ b/Engine/Plugins/Experimental/PhysicsControl/Source/PhysicsControl/Public/AnimNode_RigidBodyWithControl.h
@@ -268,7 +268,7 @@ private:
 	uint8 bEnabled : 1;
 	uint8 bSimulationStarted : 1;
 	uint8 bCheckForBodyTransformInit : 1;
-	uint8 bHaveRunControlSeup : 1;
+	uint8 bHaveSetupControls : 1;
 
 public:
 	void PostSerialize(const FArchive& Ar);
@@ -291,12 +291,15 @@ public:
 	const FTransform GetBodyTransform(const int32 BodyIndex) const;
 	const FTransform GetWorldSpaceControlRootTransform() const;
 
+	FName CreateControl(const FName ParentBoneName, const FName ChildBoneName, const FPhysicsControlData& ControlData);
+
+	FName CreateBodyModifier(const FName BoneName, const FPhysicsControlModifierData& ModifierData);
+
 private:
 
 	// Add a new physics body and update the body name to index map
 	int32 AddBody(ImmediatePhysics::FActorHandle* const BodyHandle); 
 	int32 FindBodyIndexFromBoneName(const FName BoneName) const;
-	FName FindParentBoneNameFromBoneName(const FName BoneName) const;
 	ImmediatePhysics::FActorHandle* FindBodyFromBoneName(const FName BoneName) const;
 
 	ImmediatePhysics::FJointHandle* CreateConstraint(
@@ -322,43 +325,20 @@ private:
 
 	FName GetBodyFromBoneName(const FName BoneName) const;
 
-	FName GetUniqueBodyModifierName(const FName BoneName);
-
-	FName GetUniqueControlName(const FName ParentBoneName, const FName ChildBoneName);
-
-	FName CreateControl(const FName ParentBoneName, const FName ChildBoneName, const FPhysicsControlData& ControlData);
-
-	FName CreateBodyModifier(const FName BoneName, const FPhysicsControlModifierData& ModifierData);
+	FName GetUniqueBodyModifierName(const FName BoneName) const;
 
-	void CreateControlsFromLimbBones(
-		const FName                     LimbName,
-		const FPhysicsControlLimbBones& LimbBones,
-		const EPhysicsControlType       ControlType,
-		const FPhysicsControlData&      ControlData);
+	FName GetUniqueControlName(const FName ParentBoneName, const FName ChildBoneName) const;
 
-	void CreateBodyModifiersFromLimbBones(
-		const FName                        LimbName,
-		const FPhysicsControlLimbBones&    LimbBones,
-		const FPhysicsControlModifierData& DefaultModifierData);
-
-	TMap<FName, FPhysicsControlLimbBones> GetLimbBonesFromSkeletalMesh(
-		const TArray<FPhysicsControlLimbSetupData>& LimbSetupData,
-		USkeletalMeshComponent* const               SkeletalMeshComponent) const;
+	void SetupControls(USkeletalMeshComponent* const SkeletalMeshComponent);
 
 	// This will walk through the skeleton, create controls and body modifiers, and create the sets.
-	void InitControlsAndBodyModifiers(USkeletalMeshComponent* const SkeletalMeshComponent);
+	void InitControlsAndBodyModifiers(const FReferenceSkeleton& RefSkeleton);
 
 	void DestroyControlsAndBodyModifiers();
 
 	// Helper to log overything that has been created to help the user
 	void LogControlsModifiersAndSets();
 
-	// Creates any controls as set in AdditionalControls
-	void CreateAdditionalControls();
-	
-	// Creates/updates sets - possibly making new ones, or adding to existing ones
-	void CreateAdditionalSets();
-
 	// Applies the overrides to the underlying controls and modifiers
 	void ApplyControlAndBodyModifierDatas(
 		const TArray<FPhysicsControlNamedControlParameters>& ControlParameters,
diff --git a/Engine/Plugins/Experimental/PhysicsControl/Source/PhysicsControl/Public/PhysicsControlLog.h b/Engine/Plugins/Experimental/PhysicsControl/Source/PhysicsControl/Public/PhysicsControlLog.h
new file mode 100644
--- /dev/null
+++ b/Engine/Plugins/Experimental/PhysicsControl/Source/PhysicsControl/Public/PhysicsControlLog.h
@@ -0,0 +1,15 @@
+// Copyright Epic Games, Inc. All Rights Reserved.
+
+#pragma once
+
+#include "Logging/LogMacros.h"
+
+#if UE_BUILD_SHIPPING || UE_BUILD_TEST
+PHYSICSCONTROL_API DECLARE_LOG_CATEGORY_EXTERN(LogRigidBodyWithControl, Log, Warning);
+PHYSICSCONTROL_API DECLARE_LOG_CATEGORY_EXTERN(LogPhysicsControlComponent, Log, Warning);
+#else
+PHYSICSCONTROL_API DECLARE_LOG_CATEGORY_EXTERN(LogRigidBodyWithControl, Log, All);
+PHYSICSCONTROL_API DECLARE_LOG_CATEGORY_EXTERN(LogPhysicsControlComponent, Log, All);
+#endif
+
+
diff --git a/Engine/Plugins/Experimental/PhysicsControl/Source/PhysicsControl/Public/PhysicsControlOperatorNameGeneration.h b/Engine/Plugins/Experimental/PhysicsControl/Source/PhysicsControl/Public/PhysicsControlOperatorNameGeneration.h
new file mode 100644
--- /dev/null
+++ b/Engine/Plugins/Experimental/PhysicsControl/Source/PhysicsControl/Public/PhysicsControlOperatorNameGeneration.h
@@ -0,0 +1,39 @@
+// Copyright Epic Games, Inc. All Rights Reserved.
+
+#pragma once
+
+#include "CoreMinimal.h"
+
+#include "PhysicsControlLog.h"
+
+class UPhysicsAsset;
+
+struct FAnimNode_RigidBodyWithControl;
+struct FPhysicsControlLimbBones;
+struct FPhysicsControlLimbSetupData;
+struct FPhysicsControlSetUpdates;
+struct FReferenceSkeleton;
+struct FRigidBodyControlRecord;
+struct FRigidBodyModifierRecord;
+struct FRigidBodyNameRecords;
+
+constexpr int32 MaxNumControlsOrModifiersPerName = 16;
+
+// Parse the skeleton tree to figure out which bones are associated with which limbs. 
+PHYSICSCONTROL_API TMap<FName, FPhysicsControlLimbBones> GetLimbBones(const TArray<FPhysicsControlLimbSetupData>& LimbSetupData, const FReferenceSkeleton& RefSkeleton, UPhysicsAsset* const PhysicsAsset);
+
+// Populates the supplied body modifier names, control names and name records structures with the names and sets that could be created for the supplied node, limb bones, skeleton and physics asset.
+PHYSICSCONTROL_API void CollectOperatorNames(const FAnimNode_RigidBodyWithControl* const Node, TMap<FName, FPhysicsControlLimbBones> AllLimbBones, const FReferenceSkeleton& RefSkeleton, UPhysicsAsset* const PhysicsAsset, TSet<FName>& BodyModifierNames, TSet<FName>& ControlNames, FRigidBodyNameRecords& NameRecords);
+
+// Creates the body modifiers, controls and sets for the supplied node, limb bones, skeleton and physics asset.
+PHYSICSCONTROL_API void CreateOperatorsForNode(FAnimNode_RigidBodyWithControl* const Node, TMap<FName, FPhysicsControlLimbBones> AllLimbBones, const FReferenceSkeleton& RefSkeleton, UPhysicsAsset* const PhysicsAsset, FRigidBodyNameRecords& NameRecords);
+
+// Adds the specified additional sets to the supplied Name Records structure.
+PHYSICSCONTROL_API void CreateAdditionalSets(const FPhysicsControlSetUpdates& AdditionalSets, const TSet<FName>& BodyModifierNames, const TSet<FName>& ControlNames, FRigidBodyNameRecords& NameRecords);
+
+// Adds the specified additional sets to the supplied Name Records structure.
+PHYSICSCONTROL_API void CreateAdditionalSets(const FPhysicsControlSetUpdates& AdditionalSets, const TMap<FName, FRigidBodyModifierRecord>& BodyModifierNames, const TMap<FName, FRigidBodyControlRecord>& ControlNames, FRigidBodyNameRecords& NameRecords);
+
+FName FindUniqueName(const FString& NameBase, const TSet<FName>& Keys, const int32 MaxNmeaIndex);
+FName FindUniqueBodyModifierName(const FName BodyName, const TSet<FName>& ExistingNames);
+FName FindUniqueControlName(const FName ParentBodyName, const FName ChildBodyName, const TSet<FName>& ExistingNames);
diff --git a/Engine/Plugins/Experimental/PhysicsControl/Source/PhysicsControl/Public/RigidBodyNameRecords.h b/Engine/Plugins/Experimental/PhysicsControl/Source/PhysicsControl/Public/RigidBodyNameRecords.h
--- a/Engine/Plugins/Experimental/PhysicsControl/Source/PhysicsControl/Public/RigidBodyNameRecords.h
+++ b/Engine/Plugins/Experimental/PhysicsControl/Source/PhysicsControl/Public/RigidBodyNameRecords.h
@@ -12,7 +12,7 @@
  * 3. Sets of control/modifiers - the obvious ones are limbs (e.g. a set of world-space controls on the leg)
  *    but users might want to make other sets - e.g. "UpperBody"
  */
-struct FRigidBodyNameRecords
+struct PHYSICSCONTROL_API FRigidBodyNameRecords
 {
 	/** Adds Name to SetName as well as adding it to the set "All" */
 	void AddControl(FName Name, FName SetName);
@@ -51,5 +51,5 @@ struct FRigidBodyNameRecords
 };
 
 /** Returns a copy of the supplied list of names in which any names that match keys in the supplied 'Sets' map are replaced with all the names in the array associated with that key. */
-TArray<FName> ExpandSetName(const FName InName, const TMap<FName, TArray<FName>>& Sets);
-TArray<FName> ExpandSetNames(const TArray<FName>& InNames, const TMap<FName, TArray<FName>>& Sets);
+PHYSICSCONTROL_API TArray<FName> ExpandName(const FName InName, const TMap<FName, TArray<FName>>& Sets);
+PHYSICSCONTROL_API TArray<FName> ExpandName(const TArray<FName>& InNames, const TMap<FName, TArray<FName>>& Sets);
\ No newline at end of file
diff --git a/Engine/Plugins/Experimental/PhysicsControl/Source/PhysicsControlEditor/PhysicsControlEditor.Build.cs b/Engine/Plugins/Experimental/PhysicsControl/Source/PhysicsControlEditor/PhysicsControlEditor.Build.cs
--- a/Engine/Plugins/Experimental/PhysicsControl/Source/PhysicsControlEditor/PhysicsControlEditor.Build.cs
+++ b/Engine/Plugins/Experimental/PhysicsControl/Source/PhysicsControlEditor/PhysicsControlEditor.Build.cs
@@ -39,15 +39,20 @@ public class PhysicsControlEditor : ModuleRules
 				"AnimationEditMode",
 				"AnimGraph",
 				"AnimGraphRuntime",
+				"ApplicationCore",
+				"BlueprintGraph",
 				"ComponentVisualizers",
 				"Core",
 				"CoreUObject",
 				"EditorFramework",
 				"Engine",
+				"InputCore",
 				"Persona",
 				"PhysicsControl",
 				"Slate",
 				"SlateCore",
+				"ToolMenus",
+				"ToolWidgets",
 				"UnrealEd",
 			}
 			);
@@ -62,6 +67,7 @@ public class PhysicsControlEditor : ModuleRules
 						"BlueprintGraph",
 						"EditorFramework",
 						"Kismet",
+						"PhysicsControlUncookedOnly",
 						"UnrealEd",
 				}
 			);
diff --git a/Engine/Plugins/Experimental/PhysicsControl/Source/PhysicsControlEditor/Private/OperatorEditor/OperatorEditor.cpp b/Engine/Plugins/Experimental/PhysicsControl/Source/PhysicsControlEditor/Private/OperatorEditor/OperatorEditor.cpp
new file mode 100644
--- /dev/null
+++ b/Engine/Plugins/Experimental/PhysicsControl/Source/PhysicsControlEditor/Private/OperatorEditor/OperatorEditor.cpp
@@ -0,0 +1,101 @@
+// Copyright Epic Games, Inc. All Rights Reserved.
+
+#include "OperatorEditor/OperatorEditor.h"
+
+#include "Widgets/Docking/SDockTab.h"
+#include "Widgets/Layout/SBox.h"
+
+#include "OperatorEditor/SOperatorEditorTabWidget.h"
+#include "OperatorEditor/OperatorEditor.h"
+
+// UE_DISABLE_OPTIMIZATION;
+
+static const FName PhysicsControlEditorModule_OperatorNamesTabWidget("PhysicsControlEditorModule_OperatorNamesTabWidget");
+
+#define LOCTEXT_NAMESPACE "PhysicsControlEditor"
+
+void FPhysicsControlOperatorEditor::Startup()
+{
+	// Editor Physics Operator Names Tool
+	FGlobalTabmanager::Get()->RegisterNomadTabSpawner(PhysicsControlEditorModule_OperatorNamesTabWidget, FOnSpawnTab::CreateRaw(this, &FPhysicsControlOperatorEditor::OnCreateTab))
+		.SetDisplayName(LOCTEXT("PhysicsAnimationEditor_OperatorNamesTabTitle", "Rigid Body With Control"))
+		.SetMenuType(ETabSpawnerMenuType::Hidden);
+}
+
+void FPhysicsControlOperatorEditor::Shutdown()
+{
+	FGlobalTabmanager::Get()->UnregisterNomadTabSpawner(PhysicsControlEditorModule_OperatorNamesTabWidget);
+}
+
+void FPhysicsControlOperatorEditor::OpenOperatorNamesTab()
+{
+	OperatorNamesTab = FGlobalTabmanager::Get()->FindExistingLiveTab(PhysicsControlEditorModule_OperatorNamesTabWidget);
+
+	if (!OperatorNamesTab)
+	{
+		OperatorNamesTab = FGlobalTabmanager::Get()->TryInvokeTab(PhysicsControlEditorModule_OperatorNamesTabWidget);
+	}
+}
+
+void FPhysicsControlOperatorEditor::CloseOperatorNamesTab()
+{
+	if (!OperatorNamesTab)
+	{
+		OperatorNamesTab = FGlobalTabmanager::Get()->FindExistingLiveTab(PhysicsControlEditorModule_OperatorNamesTabWidget);
+	}
+
+	OperatorNamesTab->RequestCloseTab();
+	OperatorNamesTab.Reset();
+}
+
+void FPhysicsControlOperatorEditor::ToggleOperatorNamesTab()
+{
+	if (IsOperatorNamesTabOpen())
+	{
+		CloseOperatorNamesTab();
+	}
+	else
+	{
+		OpenOperatorNamesTab();
+	}
+}
+
+bool FPhysicsControlOperatorEditor::IsOperatorNamesTabOpen()
+{
+	return OperatorNamesTab.IsValid();
+}
+
+void FPhysicsControlOperatorEditor::RequestRefresh()
+{
+	if (PersistantTabWidget)
+	{
+		PersistantTabWidget->RequestRefresh();
+	}
+}
+
+TSharedRef<SDockTab> FPhysicsControlOperatorEditor::OnCreateTab(const FSpawnTabArgs& SpawnTabArgs)
+{
+	const int TabIndex = 1;
+
+	PersistantTabWidget = SNew(SOperatorEditorTabWidget, TabIndex);
+
+	return SNew(SDockTab)
+		.TabRole(ETabRole::NomadTab)
+		.OnTabClosed_Lambda([this](TSharedRef<class SDockTab> InParentTab) { this->OnTabClosed(InParentTab); })
+		[
+			SNew(SBox)
+			.HAlign(HAlign_Fill)
+			.VAlign(VAlign_Fill)
+			[
+				PersistantTabWidget.ToSharedRef()
+			]
+		];
+}
+
+void FPhysicsControlOperatorEditor::OnTabClosed(TSharedRef<SDockTab> DockTab)
+{
+	OperatorNamesTab.Reset();
+	PersistantTabWidget.Reset();
+}
+
+#undef LOCTEXT_NAMESPACE
\ No newline at end of file
diff --git a/Engine/Plugins/Experimental/PhysicsControl/Source/PhysicsControlEditor/Private/OperatorEditor/OperatorTreeElements.cpp b/Engine/Plugins/Experimental/PhysicsControl/Source/PhysicsControlEditor/Private/OperatorEditor/OperatorTreeElements.cpp
new file mode 100644
--- /dev/null
+++ b/Engine/Plugins/Experimental/PhysicsControl/Source/PhysicsControlEditor/Private/OperatorEditor/OperatorTreeElements.cpp
@@ -0,0 +1,349 @@
+// Copyright Epic Games, Inc. All Rights Reserved.
+
+#include "OperatorEditor/OperatorTreeElements.h"
+
+#include "AnimGraphNode_RigidBodyWithControl.h"
+#include "BlueprintEditorModule.h"
+#include "Input/Reply.h"
+#include "Kismet2/BlueprintEditorUtils.h"
+#include "Styling/SlateColor.h"
+#include "Types/SlateEnums.h"
+#include "Widgets/Layout/SWrapBox.h"
+#include "Widgets/Text/STextBlock.h"
+#include "Widgets/Views/STableRow.h"
+#include "Widgets/Views/STreeView.h"
+
+// UE_DISABLE_OPTIMIZATION;
+
+#define LOCTEXT_NAMESPACE "PhysicsControlEditor"
+
+// class OperatorTreeItem //
+const EHorizontalAlignment OperatorTreeItem::HorizontalAlignment = HAlign_Left;
+const EVerticalAlignment OperatorTreeItem::VerticalAlignment = VAlign_Top;
+const float OperatorTreeItem::HorizontalPadding = 5.0f;
+const float OperatorTreeItem::VerticalPadding = 1.0f;
+
+static const FName OperatorTagControl("Control");
+static const FName OperatorTagModifier("Modifier");
+
+bool MatchSearchText(const FString& TextToSearch, const TArray<TArray<FString>>& StructuredSearchCriteria)
+{
+	for (const TArray<FString>& AndCriteria : StructuredSearchCriteria)
+	{
+		bool bMatchFound = true;
+
+		for (const FString& Token : AndCriteria)
+		{
+			bMatchFound &= TextToSearch.Contains(Token);
+		}
+
+		if (bMatchFound)
+		{
+			return true;
+		}
+	}
+
+	return false;
+}
+
+// class OperatorTreeItem //
+FReply OperatorTreeItem::OnDoubleClick()
+{ 
+	return FReply::Unhandled(); 
+}
+
+// class OperatorTreeControlItem //
+const TArray<FName> OperatorTreeControlItem::HiddenTags = { OperatorTagControl, OperatorTagModifier };
+
+FText OperatorTreeControlItem::Name() const
+{
+	return FText::FromName(NameAndTags.Key);
+}
+
+FText OperatorTreeControlItem::Description() const
+{
+	FString DescriptionStr;
+	
+	DescriptionStr += OperatorTypeName().ToString();
+	DescriptionStr += FString(" ");
+	DescriptionStr += NameAndTags.Key.ToString();
+
+	for (const FName TagName : NameAndTags.Value)
+	{
+		if (!HiddenTags.Contains(TagName))
+		{
+			DescriptionStr += " " + TagName.ToString();
+		}
+	}
+
+	DescriptionStr.TrimEndInline();
+
+	return FText::FromString(DescriptionStr);
+}
+
+FName OperatorTreeControlItem::OperatorTypeName() const
+{
+	if (HasTag(OperatorTagControl))
+	{
+		return OperatorTagControl;
+	}
+
+	if (HasTag(OperatorTagModifier))
+	{
+		return OperatorTagModifier;
+	}
+
+	return NAME_None;
+}
+
+TSharedRef<ITableRow> OperatorTreeControlItem::GenerateRow(const TSharedRef<STableViewBase>& OwnerTable) const
+{
+	TSharedPtr<SHorizontalBox> HorizontalBox = SNew(SHorizontalBox);
+
+	// Display Operator Type Indicator
+	{
+		FText OperatorTypeLabel;
+		FLinearColor OperatorTypeColor;
+
+		if (HasTag(OperatorTagControl))
+		{
+			OperatorTypeLabel = LOCTEXT("Control", "Control");
+			OperatorTypeColor = FLinearColor::Yellow;
+		}
+
+		if (HasTag(OperatorTagModifier))
+		{
+			OperatorTypeLabel = LOCTEXT("Modifier", "Modifier");
+			OperatorTypeColor = FLinearColor::Red;
+		}
+
+		HorizontalBox->AddSlot()
+			.HAlign(HorizontalAlignment)
+			.Padding(HorizontalPadding, VerticalPadding)
+			.AutoWidth()
+			[
+				SNew(STextBlock)
+				.Text(OperatorTypeLabel)
+				.ColorAndOpacity(OperatorTypeColor)
+			];
+	}
+
+	// Display Operator Name
+	{
+		HorizontalBox->AddSlot()
+			.HAlign(HorizontalAlignment)
+			.VAlign(VerticalAlignment)
+			.Padding(HorizontalPadding, VerticalPadding)
+			.FillWidth(0.4)
+			[
+				SNew(STextBlock)
+				.Text(FText::FromName(NameAndTags.Key))
+			];
+	}
+
+	// Display Sets
+	{
+		TSharedPtr<SWrapBox> TagBox = SNew(SWrapBox)
+			.UseAllottedSize(true)
+			.InnerSlotPadding(FVector2D(4, 4));
+
+		for (const FName TagName : NameAndTags.Value)
+		{
+			if (!HiddenTags.Contains(TagName))
+			{
+				TagBox->AddSlot()
+					[
+						SNew(STextBlock)
+						.Text(FText::FromName(TagName))
+						.ColorAndOpacity(FLinearColor::White)
+					];
+			}
+		}
+
+		TagBox->SetVisibility(EVisibility::Visible);
+
+		HorizontalBox->AddSlot()
+			.HAlign(HAlign_Fill)
+			.Padding(HorizontalPadding, VerticalPadding)
+			[
+				TagBox.ToSharedRef()
+			];
+	}
+
+	return SNew(STableRow<TSharedPtr<OperatorTreeItem>>, OwnerTable)
+		[
+			HorizontalBox.ToSharedRef()
+		];
+}
+
+bool OperatorTreeControlItem::MatchesFilterText(const TArray<TArray<FString>>& FilterStructuredCriteria) const
+{
+	if (FilterStructuredCriteria.IsEmpty())
+	{
+		return true;
+	}
+
+	FString TextToSearch = NameAndTags.Key.ToString();
+
+	for (const FName Tag : NameAndTags.Value)
+	{
+		TextToSearch += Tag.ToString();
+	}
+
+	return MatchSearchText(TextToSearch, FilterStructuredCriteria);
+}
+
+bool OperatorTreeControlItem::HasTag(const FName Tag) const
+{
+	return NameAndTags.Value.Contains(Tag);
+}
+
+// class OperatorTreeNodeItem //
+OperatorTreeNodeItem::OperatorTreeNodeItem(const UEdGraphNode* InNode, const FName& InName, const FString& InBlueprintName)
+: Node(InNode)
+{
+	const FString Name = Node->GetNodeTitle(ENodeTitleType::FullTitle).ToString() + " [" + InBlueprintName + "]";
+	NodeDisplayName = FText::FromString(Name);
+}
+
+FText OperatorTreeNodeItem::Name() const
+{
+	return NodeDisplayName;
+}
+
+TSharedRef<ITableRow> OperatorTreeNodeItem::GenerateRow(const TSharedRef<STableViewBase>& OwnerTable) const
+{
+	return SNew(STableRow<TSharedPtr<OperatorTreeItem>>, OwnerTable)
+		[
+			SNew(SHorizontalBox)
+			+ SHorizontalBox::Slot()
+			[
+				SNew(STextBlock)
+				.Text(NodeDisplayName) // TODO - update on changes
+				.ColorAndOpacity_Lambda([this]() { return this->ColorAndOpacity(); })
+			]
+		];
+}
+
+bool OperatorTreeNodeItem::IsSelectedInGraph() const
+{
+	return Node && Node->IsSelected();
+}
+
+FSlateColor OperatorTreeNodeItem::ColorAndOpacity() const
+{
+	return IsSelectedInGraph() ? FAppStyle::Get().GetSlateColor("Colors.AccentOrange") : FAppStyle::Get().GetSlateColor("Colors.White");
+}
+
+FReply OperatorTreeNodeItem::OnDoubleClick()
+{
+	if (UBlueprint* const Blueprint = FBlueprintEditorUtils::FindBlueprintForNode(Node))
+	{
+		if (IBlueprintEditor* BlueprintEditor = static_cast<IBlueprintEditor*>(GEditor->GetEditorSubsystem<UAssetEditorSubsystem>()->FindEditorForAsset(Blueprint, true)))
+		{
+			BlueprintEditor->JumpToHyperlink(Node, false);
+			return FReply::Handled();
+		}
+	}
+
+	return FReply::Unhandled();
+}
+
+bool operator==(const OperatorTreeNodeItem& Lhs, const OperatorTreeNodeItem& Rhs)
+{
+	return Lhs.GraphNode()->GetFName().IsEqual(Rhs.GraphNode()->GetFName()); // note: must use IsEqual as operator == doesn't account for the FName's number postfix.
+}
+
+// class OperatorTreeMessageItem //
+OperatorTreeMessageItem::OperatorTreeMessageItem(const UEdGraphNode* InNode, const MessageType InType, const FText& InText)
+	: Text(InText)
+	, Node(InNode)
+	, Type(InType)
+{}
+
+FText OperatorTreeMessageItem::Name() const
+{
+	return Text;
+}
+
+TSharedRef<ITableRow> OperatorTreeMessageItem::GenerateRow(const TSharedRef<STableViewBase>& OwnerTable) const
+{
+	FText MessageTypeLabel;
+	FLinearColor MessageTypeColor;
+
+	if (Type == MessageType::Log)
+	{
+		MessageTypeLabel = LOCTEXT("MessageItemLog", "Log");
+		MessageTypeColor = FLinearColor::Gray;
+	}
+
+	if (Type == MessageType::Warning)
+	{
+		MessageTypeLabel = LOCTEXT("MessageItemWarning", "Warning");
+		MessageTypeColor = FLinearColor::Yellow;
+	}
+
+	if (Type == MessageType::Error)
+	{
+		MessageTypeLabel = LOCTEXT("MessageItemError", "Error");
+		MessageTypeColor = FLinearColor::Red;
+	}
+
+	return SNew(STableRow<TSharedPtr<OperatorTreeItem>>, OwnerTable)
+		[
+			SNew(SHorizontalBox)
+			+ SHorizontalBox::Slot()
+			.HAlign(HorizontalAlignment)
+			.Padding(HorizontalPadding, VerticalPadding)
+			.AutoWidth()
+			[
+				SNew(STextBlock)
+				.Text(MessageTypeLabel)
+				.ColorAndOpacity(MessageTypeColor)
+			]
+			+ SHorizontalBox::Slot()
+			.HAlign(HorizontalAlignment)
+			.Padding(HorizontalPadding, VerticalPadding)
+			.AutoWidth()
+			[
+				SNew(STextBlock)
+				.Text(Text)
+				.ColorAndOpacity(MessageTypeColor)
+			]
+		];
+}
+
+template<typename TDerivedType> bool IsEqualAfterCast(const TSharedPtr<OperatorTreeItem> Lhs, const TSharedPtr<OperatorTreeItem> Rhs)
+{
+	TSharedPtr<TDerivedType> LhsDerived = StaticCastSharedPtr<TDerivedType>(Lhs);
+	TSharedPtr<TDerivedType> RhsDerived = StaticCastSharedPtr<TDerivedType>(Rhs);
+
+	return LhsDerived && RhsDerived && (*LhsDerived == *RhsDerived);
+}
+
+bool IsEqual(const TSharedPtr<OperatorTreeItem> Lhs, const TSharedPtr<OperatorTreeItem> Rhs)
+{
+	if (Lhs->GetType() == Rhs->GetType())
+	{
+		const OperatorTreeItem::OperatorTreeItemType ItemType = Lhs->GetType();
+
+		if (ItemType == OperatorTreeControlItem::StaticType())
+		{
+			return IsEqualAfterCast<OperatorTreeControlItem>(Lhs, Rhs);
+		}
+
+		if (ItemType == OperatorTreeMessageItem::StaticType())
+		{
+			return IsEqualAfterCast<OperatorTreeMessageItem>(Lhs, Rhs);
+		}
+
+		if (ItemType == OperatorTreeNodeItem::StaticType())
+		{
+			return IsEqualAfterCast<OperatorTreeNodeItem>(Lhs, Rhs);
+		}
+	}
+
+	return false;
+}
+
+#undef LOCTEXT_NAMESPACE
diff --git a/Engine/Plugins/Experimental/PhysicsControl/Source/PhysicsControlEditor/Private/OperatorEditor/SOperatorEditorTabWidget.cpp b/Engine/Plugins/Experimental/PhysicsControl/Source/PhysicsControlEditor/Private/OperatorEditor/SOperatorEditorTabWidget.cpp
new file mode 100644
--- /dev/null
+++ b/Engine/Plugins/Experimental/PhysicsControl/Source/PhysicsControlEditor/Private/OperatorEditor/SOperatorEditorTabWidget.cpp
@@ -0,0 +1,46 @@
+// Copyright Epic Games, Inc. All Rights Reserved.
+
+#include "OperatorEditor/SOperatorEditorTabWidget.h"
+#include "OperatorEditor/SOperatorTreeWidget.h"
+#include "Framework/Docking/TabManager.h"
+#include "Widgets/Docking/SDockTab.h"
+#include "Widgets/Input/SButton.h"
+
+// UE_DISABLE_OPTIMIZATION;
+
+#define LOCTEXT_NAMESPACE "PhysicsControlEditor"
+
+void SOperatorEditorTabWidget::Construct(const FArguments& InArgs, int32 InTabIndex)
+{
+	TSharedPtr<SVerticalBox> TreeViewBox;
+
+	ChildSlot.Padding(5, 5, 5, 5)
+		[
+			SNew(SVerticalBox)
+			+ SVerticalBox::Slot()
+			[
+				SAssignNew(TreeViewBox, SVerticalBox)
+			]
+	
+		];
+
+	// Tree View
+	{
+		TreeViewWidget = SNew(SOperatorTreeWidget);
+
+		TreeViewBox->AddSlot()
+			[
+				TreeViewWidget.ToSharedRef()
+			];
+	}
+}
+
+void SOperatorEditorTabWidget::RequestRefresh()
+{
+	if (TreeViewWidget.IsValid())
+	{
+		TreeViewWidget->RequestRefresh();
+	}
+}
+
+#undef LOCTEXT_NAMESPACE
\ No newline at end of file
diff --git a/Engine/Plugins/Experimental/PhysicsControl/Source/PhysicsControlEditor/Private/OperatorEditor/SOperatorTreeWidget.cpp b/Engine/Plugins/Experimental/PhysicsControl/Source/PhysicsControlEditor/Private/OperatorEditor/SOperatorTreeWidget.cpp
new file mode 100644
--- /dev/null
+++ b/Engine/Plugins/Experimental/PhysicsControl/Source/PhysicsControlEditor/Private/OperatorEditor/SOperatorTreeWidget.cpp
@@ -0,0 +1,430 @@
+// Copyright Epic Games, Inc. All Rights Reserved.
+
+#include "OperatorEditor/SOperatorTreeWidget.h"
+
+#include "AnimGraphNode_RigidBodyWithControl.h"
+#include "BlueprintEditorModule.h"
+#include "Filters/GenericFilter.h"
+#include "HAL/PlatformApplicationMisc.h"
+#include "Kismet2/BlueprintEditorUtils.h"
+#include "Widgets/Views/STableRow.h"
+#include "Widgets/Input/SSearchBox.h"
+#include "Widgets/Text/SInlineEditableTextBlock.h"
+
+// UE_DISABLE_OPTIMIZATION;
+
+#define LOCTEXT_NAMESPACE "PhysicsControl"
+
+template<typename TFilterItemType> bool PassesAnyFilters(const TFilterItemType InItem, const TSharedPtr<TFilterCollection<TFilterItemType>> FilterCollection)
+{
+	bool bPassesAnyFilter = true;
+
+	if (FilterCollection)
+	{
+		bPassesAnyFilter = false;
+
+		for (int32 FilterIndex = 0, FilterCount = FilterCollection->Num(); !bPassesAnyFilter && (FilterIndex < FilterCount); ++FilterIndex)
+		{
+			bPassesAnyFilter = FilterCollection->GetFilterAtIndex(FilterIndex)->PassesFilter(InItem);
+		}
+	}
+
+	return bPassesAnyFilter;
+}
+
+template<typename TFilterItemType> bool PassesAllFilters(const TFilterItemType InItem, const TSharedPtr<TFilterCollection<TFilterItemType>> FilterCollection)
+{
+	return !FilterCollection || FilterCollection->PassesAllFilters(InItem);
+}
+
+TArray<TArray<FString>> BuildStructuredSearchCriteria(const FText& SearchCriteriaText)
+{
+	const FString OROperatorChar(TEXT("|"));
+	const FString OROperatorText(TEXT(" OR "));
+
+	FString SearchCriteriaString = SearchCriteriaText.ToString();
+	SearchCriteriaString.ReplaceInline(*OROperatorText, *OROperatorChar);
+
+	TArray<FString> OrSeparatedCriteria;
+	SearchCriteriaString.ParseIntoArray(OrSeparatedCriteria, *OROperatorChar, true);
+
+	TArray<TArray<FString>> StructuredSearchCriteria;
+	const TCHAR* Delimeters[2] = { TEXT(" "), TEXT("+") };
+
+	for (const FString& CriteriaText : OrSeparatedCriteria)
+	{
+		StructuredSearchCriteria.Add(TArray<FString>());
+		CriteriaText.ParseIntoArray(StructuredSearchCriteria.Last(), Delimeters, 2, true);
+	}
+
+	return StructuredSearchCriteria;
+}
+
+void FormatStringForClipboard(FString& CopyBfr, const TCHAR* Delimeter)
+{
+	// Format buffer for pasting to an array in the details panel.
+	CopyBfr.ReplaceInline(Delimeter, TEXT("\",\""));
+	CopyBfr.InsertAt(0, TEXT("(\""));
+	CopyBfr.InsertAt(CopyBfr.Len(), TEXT("\")"));
+}
+
+// class SFilteredTreeWidget //
+void SOperatorTreeWidget::Construct(const FArguments& InArgs)
+{
+	Construct_Internal();
+
+	Refresh();
+}
+
+void SOperatorTreeWidget::Tick(const FGeometry& AllottedGeometry, const double InCurrentTime, const float InDeltaTime)
+{
+	Parent::Tick(AllottedGeometry, InCurrentTime, InDeltaTime);
+
+	const bool bIsScrollBarVisible = TreeView->GetScrollbarVisibility() == EVisibility::Visible;
+
+	if (bRefreshRequested)
+	{
+		Refresh();
+	}
+}
+
+void SOperatorTreeWidget::Construct_Internal()
+{
+	CreateCustomFilters();
+
+	SAssignNew(SearchBox, SSearchBox)
+		.SelectAllTextWhenFocused(true)
+		.OnTextChanged(this, &SOperatorTreeWidget::OnFilterTextChanged)
+		.HintText(LOCTEXT("SearchBoxHint", "Search the list of assets..."))
+		.AddMetaData<FTagMetaData>(TEXT("PerfToolAssetBrowser.Search"));
+
+	SAssignNew(FilterBar, FilterBarType)
+		.OnFilterChanged(this, &SOperatorTreeWidget::OnFilterBarFilterChanged)
+		.CustomFilters(CustomFilters);
+
+	TSharedPtr<SVerticalBox> VerticalBox = SNew(SVerticalBox);
+	TSharedPtr<SHorizontalBox> Toolbar = SNew(SHorizontalBox);
+
+	// Add Filter Menu Button
+	Toolbar->AddSlot()
+		.VAlign(VAlign_Center)
+		.Padding(0.0f, 0.0f, 2.0f, 0.0f)
+		.AutoWidth()
+		[
+			SBasicFilterBar<ItemType>::MakeAddFilterButton(FilterBar.ToSharedRef())
+		];
+
+	// Add the Search Text Box
+	Toolbar->AddSlot()
+		.VAlign(VAlign_Center)
+		[
+			SearchBox.ToSharedRef()
+		];
+
+	// Add the toolbar to the main layout vertical box
+	VerticalBox->AddSlot()
+		.AutoHeight()
+		.MaxHeight(24)
+		[
+			Toolbar.ToSharedRef()
+		];
+
+	// Add the filter bar to the main layout vertical box
+	VerticalBox->AddSlot()
+		.AutoHeight()
+		.MaxHeight(24)
+		.Padding(0.0f, 0.0f, 0.0f, 4.0f)
+		[
+			FilterBar.ToSharedRef()
+		];
+
+	SAssignNew(TreeView, STreeView<ItemType>)
+		.TreeItemsSource(&TreeItems)
+		.OnGetChildren(this, &SOperatorTreeWidget::OnItemGetChildren)
+		.OnGenerateRow(this, &SOperatorTreeWidget::GenerateItemRow)
+		.SelectionMode(ESelectionMode::Multi)
+		.OnKeyDownHandler(this, &SOperatorTreeWidget::OnKeyDown)
+		.OnMouseButtonDoubleClick(this, &SOperatorTreeWidget::OnItemDoubleClicked);
+
+	VerticalBox->AddSlot()
+		.VAlign(VAlign_Fill)
+		[
+			TreeView.ToSharedRef()
+		];
+
+	ChildSlot
+		[
+			VerticalBox.ToSharedRef()
+		];
+}
+
+void SOperatorTreeWidget::PushContextMenu(const FVector2D& Position)
+{
+	FMenuBuilder MenuBuilder(true, nullptr, nullptr);
+	MenuBuilder.BeginSection(NAME_None, LOCTEXT("ChannelDrawStyle", "Draw Style"));
+	{
+		MenuBuilder.AddMenuEntry(
+			LOCTEXT("OperatorTreeCopyName", "Copy Names"),
+			LOCTEXT("OperatorTreeCopyName_Tooltip", "Copy the names of all the selected operators to the clipboard."),
+			FSlateIcon(FAppStyle::GetAppStyleSetName(), "Icons.Copy"),
+			FUIAction(FExecuteAction::CreateSP(this, &SOperatorTreeWidget::CopySelectedItemsNamesToClipboard)));
+
+		MenuBuilder.AddMenuEntry(
+			LOCTEXT("OperatorTreeCopyName", "Copy All"),
+			LOCTEXT("OperatorTreeCopyName_Tooltip", "Copy the full description of all the selected operators to the clipboard."),
+			FSlateIcon(FAppStyle::GetAppStyleSetName(), "Icons.Copy"),
+			FUIAction(FExecuteAction::CreateSP(this, &SOperatorTreeWidget::CopySelectedItemsDescriptionsToClipboard)));
+	}
+	MenuBuilder.EndSection();
+
+	FSlateApplication::Get().PushMenu(
+		SharedThis(this),
+		FWidgetPath(),
+		MenuBuilder.MakeWidget(),
+		Position,
+		FPopupTransitionEffect(FPopupTransitionEffect::ContextMenu));
+}
+
+TSharedRef<ITableRow> SOperatorTreeWidget::GenerateItemRow(ItemType InItem, const TSharedRef<STableViewBase>& OwnerTable)
+{
+	return InItem->GenerateRow(OwnerTable);
+}
+
+void SOperatorTreeWidget::CopySelectedItemsNamesToClipboard() const
+{
+	const TCHAR* Delimeter = TEXT("\n");
+	FString CopyBfr;
+	uint32 ItemCount = 0;
+
+	for (ItemType SelectedItem : TreeView->GetSelectedItems())
+	{
+		CopyBfr += SelectedItem->Name().ToString() + Delimeter;
+		++ItemCount;
+	}
+
+	CopyBfr.RemoveFromEnd(Delimeter);
+
+	if (ItemCount > 1)
+	{
+		FormatStringForClipboard(CopyBfr, Delimeter);
+	}
+
+	FPlatformApplicationMisc::ClipboardCopy(*CopyBfr);
+}
+
+void SOperatorTreeWidget::CopySelectedItemsDescriptionsToClipboard() const
+{
+	constexpr TCHAR Delimeter = '\n';
+	FString CopyBfr;
+
+	for (ItemType SelectedItem : TreeView->GetSelectedItems())
+	{
+		CopyBfr += SelectedItem->Description().ToString() + Delimeter;
+	}
+
+	CopyBfr.TrimEndInline();
+
+	FPlatformApplicationMisc::ClipboardCopy(*CopyBfr);
+}
+
+FReply SOperatorTreeWidget::OnKeyDown(const FGeometry& MyGeometry, const FKeyEvent& InKeyEvent)
+{
+	if (InKeyEvent.IsControlDown() && (InKeyEvent.GetKey() == EKeys::C))
+	{
+		CopySelectedItemsNamesToClipboard();
+		return FReply::Handled();
+	}
+
+	return FReply::Unhandled();
+}
+
+FReply SOperatorTreeWidget::OnMouseButtonDown(const FGeometry& SenderGeometry, const FPointerEvent& MouseEvent)
+{
+	if (MouseEvent.GetEffectingButton() == EKeys::RightMouseButton)
+	{
+		PushContextMenu(MouseEvent.GetScreenSpacePosition());
+		return FReply::Handled();
+	}
+
+	return FReply::Unhandled();
+}
+
+void SOperatorTreeWidget::OnItemDoubleClicked(ItemType InItem)
+{
+	if (!InItem->OnDoubleClick().IsEventHandled())
+	{
+		TreeView->SetItemSelection(InItem, true, ESelectInfo::OnMouseClick);
+		CopySelectedItemsNamesToClipboard();
+	}
+}
+
+void SOperatorTreeWidget::OnFilterBarFilterChanged()
+{
+	Refresh();
+}
+
+void SOperatorTreeWidget::OnFilterTextChanged(const FText& SearchCriteriaText)
+{
+	FilterStructuredCriteria = BuildStructuredSearchCriteria(SearchCriteriaText);
+
+	for (const FName Set : SetNames)
+	{
+		if (MatchSearchText(Set.ToString(), FilterStructuredCriteria))
+		{
+			FilterSetNames.Add(Set);
+		}
+	}
+
+	Refresh();
+}
+
+void SOperatorTreeWidget::RequestRefresh()
+{
+	bRefreshRequested = true;
+}
+
+void SOperatorTreeWidget::Refresh()
+{
+	using OperatorNameAndTags = OperatorTreeControlItem::OperatorNameAndTags;
+
+	check(FilterBar.IsValid());
+
+	TSet<ItemType> ExpandedItems;
+	TreeView->GetExpandedItems(ExpandedItems);
+
+	TreeItems.Reset();
+	SetNames.Reset();
+
+	TArray<OperatorNameAndTags> OperatorNamesAndTags;
+	TSet<FString> ProcessedNodeNames;
+
+	for (TObjectIterator<UAnimGraphNode_RigidBodyWithControl> Itr; Itr; ++Itr)
+	{	
+		FString BlueprintName = "[Unknown]";
+
+		if (UAnimBlueprint* const AnimBlueprint = Cast<UAnimBlueprint>(FBlueprintEditorUtils::FindBlueprintForNode(*Itr)))
+		{
+			BlueprintName = AnimBlueprint->GetFriendlyName();
+		}
+
+		const FString NodeUniqueName = Itr->GetName() + " [" + BlueprintName + "]";
+
+		if (!ProcessedNodeNames.Contains(NodeUniqueName))
+		{
+			// This is the first time we've encountered this node.
+			ProcessedNodeNames.Add(NodeUniqueName);
+
+			OperatorNamesAndTags.Reset();
+			OperatorNamesAndTags = Itr->GenerateControlsAndBodyModifierNames();
+
+			for (const OperatorNameAndTags& Operator : OperatorNamesAndTags)
+			{
+				// TODO - Define a way to sort these operators.
+
+				for (const FName SetName : Operator.Value)
+				{
+					SetNames.Add(SetName);
+				}
+			}
+
+			// Create node item
+			NodeItemPtr Node = MakeShared<OperatorTreeNodeItem>(*Itr, FName(*NodeUniqueName), BlueprintName);
+
+			if (!Itr->Node.OverridePhysicsAsset)
+			{
+				Node->AddChild(MakeShared<OperatorTreeMessageItem>(*Itr, OperatorTreeMessageItem::Warning, LOCTEXT("OperatorNoPhysicsAssetDefinedText", "Can't create operators for limbs because no Physics Asset has been defined for this node.")));
+			}
+
+			// Create items for controls and modifiers
+			for (const OperatorNameAndTags& Operator : OperatorNamesAndTags)
+			{
+				OperatorTreeItemPtr Item = MakeShared<OperatorTreeControlItem>(Operator);
+
+				if (MatchesSearchAndFilter(Item))
+				{
+					Node->AddChild(Item);
+				}
+			}
+
+			if (OperatorNamesAndTags.IsEmpty())
+			{
+				Node->AddChild(MakeShared<OperatorTreeMessageItem>(*Itr, OperatorTreeMessageItem::Log, LOCTEXT("OperatorNoOperatorsDefinedText", "No operators defined for this node.")));
+			}
+			else if (!Node->HasChildren())
+			{
+				Node->AddChild(MakeShared<OperatorTreeMessageItem>(*Itr, OperatorTreeMessageItem::Log, LOCTEXT("OperatorNoOperatorsMatchSearchCriteria", "No operators match search criteria.")));
+			}
+			
+			TreeItems.Add(Node);
+		}
+	}
+
+	TreeView->RequestTreeRefresh();
+
+	// Find any new tree items that match old tree items that were expanded and expand them.
+	for (ItemType Item : ExpandedItems)
+	{
+		if (ItemType* const MatchingExpandedItem = TreeItems.FindByPredicate([Item](const ItemType InItem) -> bool { return IsEqual(InItem, Item); }))
+		{
+			TreeView->SetItemExpansion(*MatchingExpandedItem, true);
+		}
+	}
+
+	bRefreshRequested = false;
+}
+
+void SOperatorTreeWidget::OnItemGetChildren(ItemType InItem, TArray<ItemType>& OutChildren)
+{
+	OutChildren = InItem->Children;
+}
+
+void SOperatorTreeWidget::CreateCustomFilters()
+{
+	using FilterType = FGenericFilter<ItemType>;
+
+	TSharedPtr<FFilterCategory> OperatorTypeFilterCategory = MakeShared<FFilterCategory>(LOCTEXT("OperatorTypeFilterCategory", "Type"), FText::GetEmpty());
+	TSharedPtr<FFilterCategory> SetFilterCategory = MakeShared<FFilterCategory>(LOCTEXT("SetFilterCategory", "Set"), FText::GetEmpty());
+
+	auto CreateOperatorTypeFilter = [this, &OperatorTypeFilterCategory](const FName OperatorTypeName)
+	{
+		const FString FilterName = FString("OperatorType") + OperatorTypeName.ToString();
+
+		FilterType::FOnItemFiltered FilterDelegate = FilterType::FOnItemFiltered::CreateSP(this, &SOperatorTreeWidget::HasTag, OperatorTypeName);
+		TSharedPtr<FilterType> Filter = MakeShared<FilterType>(OperatorTypeFilterCategory, FilterName, FText::FromName(OperatorTypeName), FilterDelegate);
+		Filter->SetToolTipText(FText::Format(LOCTEXT("OperatorTypeFilterTooltip", "Only show {0}s."), FText::FromName(OperatorTypeName)));
+		this->CustomFilters.Add(Filter.ToSharedRef());
+	};
+
+	auto CreateSetFilter = [this, &SetFilterCategory](const FName SetName)
+	{
+		const FString FilterName = FString("Set") + SetName.ToString();
+
+		FilterType::FOnItemFiltered FilterDelegate = FilterType::FOnItemFiltered::CreateSP(this, &SOperatorTreeWidget::HasTag, SetName);
+		TSharedPtr<FilterType> Filter = MakeShared<FilterType>(SetFilterCategory, FilterName, FText::FromName(SetName), FilterDelegate);
+		Filter->SetToolTipText(FText::Format(LOCTEXT("SetFilterTooltip", "Only show controls or modifiers in the {0} set."), FText::FromName(SetName)));
+		this->CustomFilters.Add(Filter.ToSharedRef());
+	};
+
+	this->CustomFilters.Reset();
+
+	CreateOperatorTypeFilter(FName("Control"));
+	CreateOperatorTypeFilter(FName("Modifier"));
+
+	CreateSetFilter(FName("ParentSpace"));
+	CreateSetFilter(FName("WorldSpace"));	
+}
+
+bool SOperatorTreeWidget::HasTag(const ItemType Item, const FName InTag)
+{
+	return Item->HasTag(InTag);
+}
+
+bool SOperatorTreeWidget::MatchesSearchAndFilter(const ItemType InItem) const
+{
+	TSharedPtr<TFilterCollection<ItemType>> FilterCollection = FilterBar->GetAllActiveFilters();
+	check(FilterCollection);
+
+	return PassesAllFilters(InItem, FilterCollection) && (FilterStructuredCriteria.IsEmpty() || InItem->MatchesFilterText(FilterStructuredCriteria));
+}
+
+#undef LOCTEXT_NAMESPACE
diff --git a/Engine/Plugins/Experimental/PhysicsControl/Source/PhysicsControlEditor/Private/PhysicsControlEditorModule.cpp b/Engine/Plugins/Experimental/PhysicsControl/Source/PhysicsControlEditor/Private/PhysicsControlEditorModule.cpp
--- a/Engine/Plugins/Experimental/PhysicsControl/Source/PhysicsControlEditor/Private/PhysicsControlEditorModule.cpp
+++ b/Engine/Plugins/Experimental/PhysicsControl/Source/PhysicsControlEditor/Private/PhysicsControlEditorModule.cpp
@@ -3,10 +3,14 @@
 #include "PhysicsControlEditorModule.h"
 #include "PhysicsControlComponent.h"
 #include "PhysicsControlComponentVisualizer.h"
+#include "OperatorEditor/OperatorEditor.h"
 #include "PhysicsControlProfileAssetActions.h"
 #include "PhysicsControlProfileEditorMode.h"
 #include "UnrealEdGlobals.h"
 #include "Editor/UnrealEdEngine.h"
+#include "Features/IModularFeatures.h"
+
+static const FName PhysicsControlEditorModule_PhysicsControlEditorInterface("PhysicsControlEditorInterface");
 
 #define LOCTEXT_NAMESPACE "PhysicsControlModule"
 
@@ -36,11 +40,30 @@ void FPhysicsControlEditorModule::StartupModule()
 		Visualizer->OnRegister();
 		VisualizersToUnregisterOnShutdown.Add(UPhysicsControlComponent::StaticClass()->GetFName());
 	}
+
+	if (!EditorInterface)
+	{
+		EditorInterface = new FPhysicsControlOperatorEditor;
+	}
+
+	if (EditorInterface)
+	{
+		EditorInterface->Startup();
+		IModularFeatures::Get().RegisterModularFeature(PhysicsControlEditorModule_PhysicsControlEditorInterface, EditorInterface);
+	}
 }
 
 //======================================================================================================================
 void FPhysicsControlEditorModule::ShutdownModule()
 {
+	if (EditorInterface)
+	{
+		EditorInterface->Shutdown();
+		IModularFeatures::Get().UnregisterModularFeature(PhysicsControlEditorModule_PhysicsControlEditorInterface, EditorInterface);
+		delete EditorInterface;
+		EditorInterface = nullptr;
+	}
+	
 	// Physics Control Profile editor/asset is disabled for now
 #ifdef ENABLE_PHYSICS_CONTROL_PROFILE_EDITOR
 	FEditorModeRegistry::Get().UnregisterMode(FPhysicsControlProfileEditorMode::ModeName);
@@ -62,7 +85,5 @@ void FPhysicsControlEditorModule::ShutdownModule()
 }
 
 #undef LOCTEXT_NAMESPACE
-	
-IMPLEMENT_MODULE(FPhysicsControlEditorModule, PhysicsControlEditor)
-
 
+IMPLEMENT_MODULE(FPhysicsControlEditorModule, PhysicsControlEditor)
diff --git a/Engine/Plugins/Experimental/PhysicsControl/Source/PhysicsControlEditor/Public/OperatorEditor/IPhysicsControlOperatorEditorInterface.h b/Engine/Plugins/Experimental/PhysicsControl/Source/PhysicsControlEditor/Public/OperatorEditor/IPhysicsControlOperatorEditorInterface.h
new file mode 100644
--- /dev/null
+++ b/Engine/Plugins/Experimental/PhysicsControl/Source/PhysicsControlEditor/Public/OperatorEditor/IPhysicsControlOperatorEditorInterface.h
@@ -0,0 +1,18 @@
+// Copyright Epic Games, Inc. All Rights Reserved.
+
+#pragma once
+
+#include "CoreMinimal.h"
+#include "Features/IModularFeature.h"
+
+class IPhysicsControlOperatorEditorInterface : public IModularFeature
+{
+public:
+
+	virtual void OpenOperatorNamesTab() = 0;
+	virtual void CloseOperatorNamesTab() = 0;
+	virtual void ToggleOperatorNamesTab() = 0;
+	virtual bool IsOperatorNamesTabOpen() = 0;
+	virtual void RequestRefresh() = 0;
+
+};
\ No newline at end of file
diff --git a/Engine/Plugins/Experimental/PhysicsControl/Source/PhysicsControlEditor/Public/OperatorEditor/OperatorEditor.h b/Engine/Plugins/Experimental/PhysicsControl/Source/PhysicsControlEditor/Public/OperatorEditor/OperatorEditor.h
new file mode 100644
--- /dev/null
+++ b/Engine/Plugins/Experimental/PhysicsControl/Source/PhysicsControlEditor/Public/OperatorEditor/OperatorEditor.h
@@ -0,0 +1,35 @@
+// Copyright Epic Games, Inc. All Rights Reserved.
+
+#pragma once
+
+#include "CoreMinimal.h"
+#include "OperatorEditor/IPhysicsControlOperatorEditorInterface.h"
+#include "Containers/Ticker.h"
+
+class FSpawnTabArgs;
+class SOperatorEditorTabWidget;
+class SDockTab;
+
+class FPhysicsControlOperatorEditor : public IPhysicsControlOperatorEditorInterface
+{
+public:
+
+	virtual ~FPhysicsControlOperatorEditor() {}
+
+	virtual void OpenOperatorNamesTab() override;
+	virtual void CloseOperatorNamesTab() override;
+	virtual void ToggleOperatorNamesTab() override;
+	virtual bool IsOperatorNamesTabOpen() override;
+	virtual void RequestRefresh() override;
+
+	void Startup();
+	void Shutdown();
+
+	void OnTabClosed(TSharedRef<SDockTab> DockTab);
+
+private:
+	TSharedRef<SDockTab> OnCreateTab(const FSpawnTabArgs& SpawnTabArgs);
+
+	TSharedPtr<SOperatorEditorTabWidget> PersistantTabWidget;
+	TSharedPtr<SDockTab> OperatorNamesTab;
+};
\ No newline at end of file
diff --git a/Engine/Plugins/Experimental/PhysicsControl/Source/PhysicsControlEditor/Public/OperatorEditor/OperatorTreeElements.h b/Engine/Plugins/Experimental/PhysicsControl/Source/PhysicsControlEditor/Public/OperatorEditor/OperatorTreeElements.h
new file mode 100644
--- /dev/null
+++ b/Engine/Plugins/Experimental/PhysicsControl/Source/PhysicsControlEditor/Public/OperatorEditor/OperatorTreeElements.h
@@ -0,0 +1,189 @@
+// Copyright Epic Games, Inc. All Rights Reserved.
+
+#pragma once
+
+#include "CoreMinimal.h"
+
+enum EHorizontalAlignment : int;
+enum EVerticalAlignment : int;
+
+class FReply;
+class FString;
+class ITableRow;
+class STableViewBase;
+class UEdGraphNode;
+
+struct FSlateColor;
+
+class OperatorTreeItem;
+
+using OperatorTreeItemPtr = TSharedPtr<OperatorTreeItem>;
+
+class OperatorTreeItem
+{
+public:
+
+	virtual ~OperatorTreeItem() {}
+
+	enum OperatorTreeItemType
+	{
+		ControlType,
+		ModifierType,
+		TargetType,
+		NodeType,
+		MessageType,
+		None
+	};
+
+	virtual FText Name() const = 0;
+	virtual FText Description() const { return Name(); }
+
+	virtual TSharedRef<ITableRow> GenerateRow(const TSharedRef<STableViewBase>& OwnerTable) const = 0;
+	virtual FReply OnDoubleClick();
+	virtual bool MatchesFilterText(const TArray<TArray<FString>>& FilterStructuredCriteria) const { return true; };
+	virtual bool HasTag(const FName Tag) const { return false; }
+	virtual bool IsExpandable() const { return false; }
+	virtual bool IsExpanded() const { return true; }
+
+	void AddChild(OperatorTreeItemPtr ItemPtr) { Children.Add(ItemPtr); }
+	bool HasChildren() const { return !Children.IsEmpty(); }
+
+	virtual OperatorTreeItemType GetType() const { return None; }
+	bool HasType(const OperatorTreeItemType InType) const { return InType == GetType(); }
+
+	TArray<OperatorTreeItemPtr> Children;
+
+	static const EHorizontalAlignment HorizontalAlignment;
+	static const EVerticalAlignment VerticalAlignment;
+	static const float HorizontalPadding;
+	static const float VerticalPadding;
+};
+
+class OperatorTreeControlItem : public OperatorTreeItem
+{
+	using Super = OperatorTreeItem;
+
+public:
+
+	using OperatorNameAndTags = TPair<FName, TArray<FName>>;
+
+	OperatorTreeControlItem()
+	{}
+
+	OperatorTreeControlItem(const OperatorTreeControlItem&) = default;
+
+	OperatorTreeControlItem(const OperatorNameAndTags& InOperatorNameAndTags)
+	: NameAndTags(InOperatorNameAndTags)
+	{}
+
+	virtual FText Name() const override;
+	virtual FText Description() const override;
+
+	virtual TSharedRef<ITableRow> GenerateRow(const TSharedRef<STableViewBase>& OwnerTable) const override;
+	virtual bool MatchesFilterText(const TArray<TArray<FString>>& FilterStructuredCriteria) const override;
+	virtual bool HasTag(const FName Tag) const override;
+	FName OperatorTypeName() const;
+
+	virtual OperatorTreeItemType GetType() const override { return StaticType(); }
+	static OperatorTreeItemType StaticType() { return OperatorTreeItem::ControlType; }
+
+	OperatorNameAndTags NameAndTags;
+
+	static const TArray<FName> HiddenTags;
+};
+
+using ControlItemPtr = TSharedPtr<OperatorTreeControlItem>;
+
+inline bool operator==(const OperatorTreeControlItem& Lhs, const OperatorTreeControlItem& Rhs)
+{
+	return (Lhs.Name().ToString() == Rhs.Name().ToString()) && (Lhs.HasTag(FName("Modifier")) == Rhs.HasTag(FName("Modifier")));
+}
+
+class OperatorTreeNodeItem : public OperatorTreeItem
+{
+	using Super = OperatorTreeItem;
+
+public:
+
+	OperatorTreeNodeItem()
+	{}
+
+	OperatorTreeNodeItem(const OperatorTreeNodeItem&) = default;
+
+	OperatorTreeNodeItem(const UEdGraphNode* InNode, const FName& InName, const FString& InBlueprintName);
+
+	virtual FText Name() const override;
+
+	virtual TSharedRef<ITableRow> GenerateRow(const TSharedRef<STableViewBase>& OwnerTable) const override;
+	virtual FReply OnDoubleClick() override;
+
+	virtual OperatorTreeItemType GetType() const override { return StaticType(); }
+	static OperatorTreeItemType StaticType() { return OperatorTreeItem::NodeType; }
+
+	const UEdGraphNode* GraphNode() const { return Node; }
+	bool IsSelectedInGraph() const;
+	FSlateColor ColorAndOpacity() const;
+
+private:
+	const UEdGraphNode* Node;
+	FText NodeDisplayName;
+
+};
+
+using NodeItemPtr = TSharedPtr<OperatorTreeNodeItem>;
+
+bool operator==(const OperatorTreeNodeItem& Lhs, const OperatorTreeNodeItem& Rhs);
+
+class OperatorTreeMessageItem : public OperatorTreeItem
+{
+	using Super = OperatorTreeItem;
+
+public:
+
+	enum MessageType
+	{
+		Log,
+		Warning,
+		Error
+	};
+
+	OperatorTreeMessageItem() = default;
+	OperatorTreeMessageItem(const OperatorTreeMessageItem&) = default;
+	OperatorTreeMessageItem(const UEdGraphNode* InNode, const MessageType InType, const FText& InText);
+
+	virtual FText Name() const override;
+
+	virtual TSharedRef<ITableRow> GenerateRow(const TSharedRef<STableViewBase>& OwnerTable) const override;
+
+	virtual OperatorTreeItemType GetType() const override { return StaticType(); }
+	static OperatorTreeItemType StaticType() { return OperatorTreeItem::MessageType; }
+
+	const UEdGraphNode* GraphNode() const { return Node; }
+
+private:
+	FText Text;
+	const UEdGraphNode* Node;
+	MessageType Type;
+};
+
+using MessageItemPtr = TSharedPtr<OperatorTreeMessageItem>;
+
+inline bool operator==(const OperatorTreeMessageItem& Lhs, const OperatorTreeMessageItem& Rhs)
+{
+	return Lhs.Name().ToString() == Rhs.Name().ToString(); // TODO - don't convert to strings + incomplete
+}
+
+template<typename DerivedType> TSharedPtr<DerivedType> StaticCastOperatorItemPtr(const TSharedPtr<OperatorTreeItem> InItem)
+{
+	if (InItem->GetType() == DerivedType::StaticType())
+	{
+		return StaticCastSharedPtr<DerivedType>(InItem);
+	}
+
+	return TSharedPtr<DerivedType>();
+}
+
+bool IsEqual(const TSharedPtr<OperatorTreeItem> Lhs, const TSharedPtr<OperatorTreeItem> Rhs);
+
+// Returns true if the TextToSearch string matches the search criteria. Used to identify operator or set names that match user specified search strings in the editor.
+bool MatchSearchText(const FString& TextToSearch, const TArray<TArray<FString>>& StructuredSearchCriteria);
diff --git a/Engine/Plugins/Experimental/PhysicsControl/Source/PhysicsControlEditor/Public/OperatorEditor/SOperatorEditorTabWidget.h b/Engine/Plugins/Experimental/PhysicsControl/Source/PhysicsControlEditor/Public/OperatorEditor/SOperatorEditorTabWidget.h
new file mode 100644
--- /dev/null
+++ b/Engine/Plugins/Experimental/PhysicsControl/Source/PhysicsControlEditor/Public/OperatorEditor/SOperatorEditorTabWidget.h
@@ -0,0 +1,27 @@
+// Copyright Epic Games, Inc. All Rights Reserved.
+
+#pragma once
+
+#include "Templates/SharedPointer.h"
+//#include "Widgets/DeclarativeSyntaxSupport.h"
+#include "Widgets/SCompoundWidget.h"
+
+class SCompoundWidget;
+class SOperatorTreeWidget;
+
+class SOperatorEditorTabWidget : public SCompoundWidget
+{
+public:
+
+	SLATE_BEGIN_ARGS(SOperatorEditorTabWidget)
+	{}
+
+	SLATE_END_ARGS()
+
+public:
+	void Construct(const FArguments& InArgs, int32 InTabIndex);
+	void RequestRefresh();
+
+private:
+	TSharedPtr<SOperatorTreeWidget> TreeViewWidget;
+};
diff --git a/Engine/Plugins/Experimental/PhysicsControl/Source/PhysicsControlEditor/Public/OperatorEditor/SOperatorTreeWidget.h b/Engine/Plugins/Experimental/PhysicsControl/Source/PhysicsControlEditor/Public/OperatorEditor/SOperatorTreeWidget.h
new file mode 100644
--- /dev/null
+++ b/Engine/Plugins/Experimental/PhysicsControl/Source/PhysicsControlEditor/Public/OperatorEditor/SOperatorTreeWidget.h
@@ -0,0 +1,80 @@
+// Copyright Epic Games, Inc. All Rights Reserved.
+
+#pragma once
+
+#include "CoreMinimal.h"
+
+#include "Widgets/SCompoundWidget.h"
+#include "Filters/SFilterSearchBox.h"
+#include "Filters/SBasicFilterBar.h"
+#include "Framework/SlateDelegates.h"
+#include "Framework/Application/SlateApplication.h"
+#include "Styling/AppStyle.h"
+
+#include "Widgets/Views/STreeView.h"
+
+#include "AnimNode_RigidBodyWithControl.h"
+
+#include "OperatorEditor/OperatorTreeElements.h"
+
+/**
+ * class SOperatorTreeWidget
+ *
+ * 
+ */
+class SOperatorTreeWidget : public SCompoundWidget
+{
+	using Parent = SCompoundWidget;
+
+public:
+	using ItemType = OperatorTreeItemPtr;
+
+	using FilterBarType = SBasicFilterBar<ItemType>;
+
+	SLATE_BEGIN_ARGS(SOperatorTreeWidget) {}
+	SLATE_END_ARGS()
+
+	/** Constructs this widget with InArgs */
+	void Construct(const FArguments& InArgs);
+	
+	virtual void Tick(const FGeometry& AllottedGeometry, const double InCurrentTime, const float InDeltaTime) override;
+
+	TSharedRef<ITableRow> GenerateItemRow(ItemType InItem, const TSharedRef<STableViewBase>& OwnerTable);
+	void OnItemGetChildren(ItemType InItem, TArray<ItemType>& OutChildren);
+	void OnFilterBarFilterChanged();
+	void OnFilterTextChanged(const FText& SearchText);
+
+	bool MatchesSearchAndFilter(const ItemType InItem) const;
+
+	bool HasTag(const ItemType Item, const FName Tag);
+
+	virtual FReply OnKeyDown(const FGeometry& MyGeometry, const FKeyEvent& InKeyEvent) override;
+	virtual FReply OnMouseButtonDown(const FGeometry& SenderGeometry, const FPointerEvent& MouseEvent) override;
+
+	void OnItemDoubleClicked(ItemType InItem);
+	void PushContextMenu(const FVector2D& Position);
+	void RequestRefresh();
+
+	void CopySelectedItemsNamesToClipboard() const;
+	void CopySelectedItemsDescriptionsToClipboard() const;
+
+protected:
+
+	void Construct_Internal();
+	virtual void CreateCustomFilters();
+	virtual void Refresh();
+
+	TSharedPtr<STreeView<ItemType>> TreeView;
+	TSharedPtr<SSearchBox> SearchBox;
+	TSharedPtr<FilterBarType> FilterBar;
+
+	TArray<TSharedRef<FFilterBase<ItemType>>> CustomFilters;
+	TArray<ItemType> TreeItems;
+	TSet<FName> SetNames;
+	FText FilterText;
+
+	TArray<TArray<FString>> FilterStructuredCriteria;
+	TSet<FName> FilterSetNames;
+
+	bool bRefreshRequested;
+};
diff --git a/Engine/Plugins/Experimental/PhysicsControl/Source/PhysicsControlEditor/Public/PhysicsControlEditorModule.h b/Engine/Plugins/Experimental/PhysicsControl/Source/PhysicsControlEditor/Public/PhysicsControlEditorModule.h
--- a/Engine/Plugins/Experimental/PhysicsControl/Source/PhysicsControlEditor/Public/PhysicsControlEditorModule.h
+++ b/Engine/Plugins/Experimental/PhysicsControl/Source/PhysicsControlEditor/Public/PhysicsControlEditorModule.h
@@ -5,6 +5,9 @@
 #include "CoreMinimal.h"
 #include "Modules/ModuleInterface.h"
 
+
+
+class FPhysicsControlOperatorEditor;
 class FPhysicsControlProfileAssetActions;
 
 /**
@@ -19,7 +22,10 @@ public:
 	void ShutdownModule() override;
 
 private:
+
 	TArray<FName> VisualizersToUnregisterOnShutdown;
 	TSharedPtr<FPhysicsControlProfileAssetActions> PhysicsControlProfileAssetActions;
+
+	FPhysicsControlOperatorEditor* EditorInterface = nullptr;
 };
 
diff --git a/Engine/Plugins/Experimental/PhysicsControl/Source/PhysicsControlUncookedOnly/PhysicsControlUncookedOnly.Build.cs b/Engine/Plugins/Experimental/PhysicsControl/Source/PhysicsControlUncookedOnly/PhysicsControlUncookedOnly.Build.cs
--- a/Engine/Plugins/Experimental/PhysicsControl/Source/PhysicsControlUncookedOnly/PhysicsControlUncookedOnly.Build.cs
+++ b/Engine/Plugins/Experimental/PhysicsControl/Source/PhysicsControlUncookedOnly/PhysicsControlUncookedOnly.Build.cs
@@ -18,6 +18,7 @@ public class PhysicsControlUncookedOnly : ModuleRules
 				"Engine",
 				"Persona",
 				"PhysicsControl",
+				"PhysicsControlEditor",
 				"Slate",
 				"SlateCore",
 			}
diff --git a/Engine/Plugins/Experimental/PhysicsControl/Source/PhysicsControlUncookedOnly/Private/AnimGraphNode_RigidBodyWithControl.cpp b/Engine/Plugins/Experimental/PhysicsControl/Source/PhysicsControlUncookedOnly/Private/AnimGraphNode_RigidBodyWithControl.cpp
--- a/Engine/Plugins/Experimental/PhysicsControl/Source/PhysicsControlUncookedOnly/Private/AnimGraphNode_RigidBodyWithControl.cpp
+++ b/Engine/Plugins/Experimental/PhysicsControl/Source/PhysicsControlUncookedOnly/Private/AnimGraphNode_RigidBodyWithControl.cpp
@@ -16,7 +16,9 @@
 #include "AnimNode_RigidBodyWithControl.h"
 #include "EditorModeManager.h"
 #include "IPhysicsAssetRenderInterface.h"
-// #include "IAnimNodeEditMode.h" <- removed as it can't be found during build and doesn't appear to be necessary
+#include "OperatorEditor/IPhysicsControlOperatorEditorInterface.h"
+#include "Kismet2/BlueprintEditorUtils.h"
+#include "PhysicsControlOperatorNameGeneration.h"
 
 // Details includes
 #include "PropertyHandle.h"
@@ -25,10 +27,10 @@
 #include "DetailCategoryBuilder.h"
 #include "Widgets/Input/SButton.h"
 #include "Widgets/Text/STextBlock.h"
-
-
 #include "SceneManagement.h"
 
+// UE_DISABLE_OPTIMIZATION;
+
 /////////////////////////////////////////////////////
 // UAnimGraphNode_RigidBodyWithControl
 
@@ -147,43 +149,76 @@ void UAnimGraphNode_RigidBodyWithControl::CustomizeDetails(IDetailLayoutBuilder&
 	Super::CustomizeDetails(DetailBuilder);
 
 	IDetailCategoryBuilder& ViewportCategory = DetailBuilder.EditCategory(TEXT("Debug Visualization"));
-	FDetailWidgetRow& WidgetRow = ViewportCategory.AddCustomRow(LOCTEXT("ToggleDebugVisualizationButtonRow", "DebugVisualization"));
+	
 	FAnimNode_RigidBodyWithControl* const RigidBodyNode = static_cast<FAnimNode_RigidBodyWithControl*>(GetDebuggedAnimNode());
+	
+	{
+		FDetailWidgetRow& ControlsEditorWidgetRow = ViewportCategory.AddCustomRow(LOCTEXT("ToggleControlsEditorWidgetRowButtonRow", "ControlsEditor"));
 
-	WidgetRow
-		[
-			SNew(SHorizontalBox)
-			// Show/Hide Bodies button.
-			+ SHorizontalBox::Slot()
+		ControlsEditorWidgetRow
 			[
-				SNew(SButton)
-				.HAlign(HAlign_Center)
-				.VAlign(VAlign_Center)
-				.OnClicked_Lambda([this](){ this->ToggleBodyVisibility(); return FReply::Handled(); })
-				.ButtonColorAndOpacity_Lambda([this](){ return (AreAnyBodiesHidden()) ? FAppStyle::Get().GetSlateColor("Colors.AccentRed") : FAppStyle::Get().GetSlateColor("Colors.AccentGreen"); })
-				.Content()
+				SNew(SHorizontalBox)
+				+ SHorizontalBox::Slot()
 				[
-					SNew(STextBlock)
-					.Text_Lambda([this]() { return (AreAnyBodiesHidden()) ? LOCTEXT("ShowAllBodiesButtonText", "Show All Bodies") : LOCTEXT("HideAllBodiesButtonText", "Hide All Bodies"); })
-					.ToolTipText(LOCTEXT("ToggleBodyVisibilityButtonToolTip", "Toggle debug visualization of all physics bodies"))
+					SNew(SButton)
+					.HAlign(HAlign_Center)
+					.VAlign(VAlign_Center)
+					.OnClicked_Lambda([this]() { this->ToggleControlEditorTab(); return FReply::Handled(); })
+					.ButtonColorAndOpacity_Lambda([this]() { return (IsControlEditorTabOpen()) ? FAppStyle::Get().GetSlateColor("Colors.AccentRed") : FAppStyle::Get().GetSlateColor("Colors.AccentGreen"); })
+					.Content()
+					[
+						SNew(STextBlock)
+						.Text_Lambda([this]() { return (IsControlEditorTabOpen()) ? LOCTEXT("CloseControlEditorTabButtonText", "Close Control Editor") : LOCTEXT("OpenControlEditorTabButtonText", "Open Control Editor"); })
+						.ToolTipText(LOCTEXT("ToggleControlEditorTabButtonToolTip", "Toggle Control Editor Tab"))
+					]
 				]
-			]
-			// Show/Hide Constraints button.
-			+ SHorizontalBox::Slot()
+			];
+	}
+
+	{
+		FDetailWidgetRow& DebugVisualizationWidgetRow = ViewportCategory.AddCustomRow(LOCTEXT("ToggleDebugVisualizationButtonRow", "DebugVisualization"));
+
+		DebugVisualizationWidgetRow
 			[
-				SNew(SButton)
-				.HAlign(HAlign_Center)
-				.VAlign(VAlign_Center)
-				.OnClicked_Lambda([this](){ this->ToggleConstraintVisibility(); return FReply::Handled(); })
-				.ButtonColorAndOpacity_Lambda([this]() { return (AreAnyConstraintsHidden()) ? FAppStyle::Get().GetSlateColor("Colors.AccentRed") : FAppStyle::Get().GetSlateColor("Colors.AccentGreen"); })
-				.Content()
+				SNew(SHorizontalBox)
+				// Show/Hide Bodies button.
+				+ SHorizontalBox::Slot()
+				[
+					SNew(SButton)
+					.HAlign(HAlign_Center)
+					.VAlign(VAlign_Center)
+					.OnClicked_Lambda([this]() { this->ToggleBodyVisibility(); return FReply::Handled(); })
+					.ButtonColorAndOpacity_Lambda([this]() { return (AreAnyBodiesHidden()) ? FAppStyle::Get().GetSlateColor("Colors.AccentRed") : FAppStyle::Get().GetSlateColor("Colors.AccentGreen"); })
+					.Content()
+					[
+						SNew(STextBlock)
+						.Text_Lambda([this]() { return (AreAnyBodiesHidden()) ? LOCTEXT("ShowAllBodiesButtonText", "Show All Bodies") : LOCTEXT("HideAllBodiesButtonText", "Hide All Bodies"); })
+						.ToolTipText(LOCTEXT("ToggleBodyVisibilityButtonToolTip", "Toggle debug visualization of all physics bodies"))
+					]
+				]
+				// Show/Hide Constraints button.
+				+ SHorizontalBox::Slot()
 				[
-					SNew(STextBlock)
-					.Text_Lambda([this]() { return (AreAnyConstraintsHidden()) ? LOCTEXT("ShowAllConstraintsButtonText", "Show All Constraints") : LOCTEXT("HideAllConstraintsButtonText", "Hide All Constraints"); })
-					.ToolTipText(LOCTEXT("ToggleConstraintVisibilityButtonToolTip", "Toggle debug visualization of all physics constriants"))
+					SNew(SButton)
+					.HAlign(HAlign_Center)
+					.VAlign(VAlign_Center)
+					.OnClicked_Lambda([this]() { this->ToggleConstraintVisibility(); return FReply::Handled(); })
+					.ButtonColorAndOpacity_Lambda([this]() { return (AreAnyConstraintsHidden()) ? FAppStyle::Get().GetSlateColor("Colors.AccentRed") : FAppStyle::Get().GetSlateColor("Colors.AccentGreen"); })
+					.Content()
+					[
+						SNew(STextBlock)
+						.Text_Lambda([this]() { return (AreAnyConstraintsHidden()) ? LOCTEXT("ShowAllConstraintsButtonText", "Show All Constraints") : LOCTEXT("HideAllConstraintsButtonText", "Hide All Constraints"); })
+						.ToolTipText(LOCTEXT("ToggleConstraintVisibilityButtonToolTip", "Toggle debug visualization of all physics constriants"))
+					]
 				]
-			]
-		];
+			];
+	}
+}
+
+void UAnimGraphNode_RigidBodyWithControl::PostChange()
+{
+	IPhysicsControlOperatorEditorInterface& PhysicsControlEditorInterface = IModularFeatures::Get().GetModularFeature<IPhysicsControlOperatorEditorInterface>("PhysicsControlEditorInterface");
+	PhysicsControlEditorInterface.RequestRefresh();
 }
 
 void UAnimGraphNode_RigidBodyWithControl::PostEditChangeProperty(struct FPropertyChangedEvent& PropertyChangedEvent)
@@ -192,6 +227,26 @@ void UAnimGraphNode_RigidBodyWithControl::PostEditChangeProperty(struct FPropert
 
 	IPhysicsAssetRenderInterface& PhysicsAssetRenderInterface = IModularFeatures::Get().GetModularFeature<IPhysicsAssetRenderInterface>("PhysicsAssetRenderInterface");
 	PhysicsAssetRenderInterface.SaveConfig();
+
+	PostChange();
+}
+
+void UAnimGraphNode_RigidBodyWithControl::PostPlacedNewNode()
+{
+	Super::PostPlacedNewNode();
+	PostChange();
+}
+
+void UAnimGraphNode_RigidBodyWithControl::PostPasteNode()
+{
+	Super::PostPasteNode();
+	PostChange();
+}
+
+void UAnimGraphNode_RigidBodyWithControl::DestroyNode()
+{
+	Super::DestroyNode();
+	PostChange();
 }
 
 void UAnimGraphNode_RigidBodyWithControl::ToggleBodyVisibility()
@@ -243,4 +298,79 @@ bool UAnimGraphNode_RigidBodyWithControl::AreAnyConstraintsHidden() const
 	return false;
 }
 
+void UAnimGraphNode_RigidBodyWithControl::ToggleControlEditorTab()
+{
+	IPhysicsControlOperatorEditorInterface& PhysicsControlEditorInterface = IModularFeatures::Get().GetModularFeature<IPhysicsControlOperatorEditorInterface>("PhysicsControlEditorInterface");
+	PhysicsControlEditorInterface.ToggleOperatorNamesTab();
+}
+
+bool UAnimGraphNode_RigidBodyWithControl::IsControlEditorTabOpen() const
+{
+	IPhysicsControlOperatorEditorInterface& PhysicsControlEditorInterface = IModularFeatures::Get().GetModularFeature<IPhysicsControlOperatorEditorInterface>("PhysicsControlEditorInterface");
+	return PhysicsControlEditorInterface.IsOperatorNamesTabOpen();
+}
+
+TArray<TPair<FName, TArray<FName>>> UAnimGraphNode_RigidBodyWithControl::GenerateControlsAndBodyModifierNames() const
+{
+	using OperatorNameAndTags = TPair<FName, TArray<FName>>;
+
+	TArray<OperatorNameAndTags> GeneratedOperatorNames;
+
+	if (USkeleton* const Skeleton = GetSkeleton())
+	{
+		const FReferenceSkeleton& RefSkeleton = Skeleton->GetReferenceSkeleton();
+
+		// These functions will create the base set of controls and modifiers from SetupData
+		TMap<FName, FPhysicsControlLimbBones> AllLimbBones =
+			GetLimbBones(Node.SetupData.LimbSetupData, RefSkeleton, Node.OverridePhysicsAsset);
+
+		TSet<FName> BodyModifierNames;
+		TSet<FName> ControlNames;
+		FRigidBodyNameRecords NameRecords;
+
+		CollectOperatorNames(&Node, AllLimbBones, RefSkeleton, Node.OverridePhysicsAsset, BodyModifierNames, ControlNames, NameRecords);
+
+		// Create any additional sets that have been requested
+		CreateAdditionalSets(Node.AdditionalSets, BodyModifierNames, ControlNames, NameRecords);
+
+		auto TransformOperatorNamesAndTags = [&GeneratedOperatorNames](const FName TypeTag, const TSet<FName>& Names, const TMap<FName, TArray<FName>>& SetToOperatorNameMap)
+		{
+			for (const FName OperatorName : Names)
+			{
+				OperatorNameAndTags NameAndTagsPair;
+
+				NameAndTagsPair.Key = OperatorName;
+				NameAndTagsPair.Value.Add(TypeTag);
+
+				for (const TMap<FName, TArray<FName>>::ElementType& Set : SetToOperatorNameMap)
+				{
+					if (Set.Value.Contains(OperatorName))
+					{
+						NameAndTagsPair.Value.Add(Set.Key);
+					}
+				}
+
+				GeneratedOperatorNames.Add(NameAndTagsPair);
+			}
+		};
+
+		TransformOperatorNamesAndTags(FName("Modifier"), BodyModifierNames, NameRecords.BodyModifierSets);
+		TransformOperatorNamesAndTags(FName("Control"), ControlNames, NameRecords.ControlSets);
+	}
+
+	return GeneratedOperatorNames;
+}
+
+USkeleton* UAnimGraphNode_RigidBodyWithControl::GetSkeleton() const
+{
+	USkeleton* Skeleton = nullptr;
+
+	if (UAnimBlueprint* const AnimBlueprint = Cast<UAnimBlueprint>(FBlueprintEditorUtils::FindBlueprintForNode(this)))
+	{
+		Skeleton = AnimBlueprint->TargetSkeleton;
+	}
+
+	return Skeleton;
+}
+
 #undef LOCTEXT_NAMESPACE
diff --git a/Engine/Plugins/Experimental/PhysicsControl/Source/PhysicsControlUncookedOnly/Public/AnimGraphNode_RigidBodyWithControl.h b/Engine/Plugins/Experimental/PhysicsControl/Source/PhysicsControlUncookedOnly/Public/AnimGraphNode_RigidBodyWithControl.h
--- a/Engine/Plugins/Experimental/PhysicsControl/Source/PhysicsControlUncookedOnly/Public/AnimGraphNode_RigidBodyWithControl.h
+++ b/Engine/Plugins/Experimental/PhysicsControl/Source/PhysicsControlUncookedOnly/Public/AnimGraphNode_RigidBodyWithControl.h
@@ -40,8 +40,17 @@ public:
 	virtual void CustomizeDetails(IDetailLayoutBuilder& DetailBuilder) override;
 	virtual void PostEditChangeProperty(struct FPropertyChangedEvent& PropertyChangedEvent) override;
 
+	virtual void PostPlacedNewNode() override;
+	virtual void PostPasteNode() override;
+	virtual void DestroyNode() override;
+
 	void ToggleBodyVisibility();
 	void ToggleConstraintVisibility();
+	void ToggleControlEditorTab();
+
+	TArray<TPair<FName, TArray<FName>>> GenerateControlsAndBodyModifierNames() const;
+
+	USkeleton* GetSkeleton() const;
 
 protected:
 	// UAnimGraphNode_SkeletalControlBase interface
@@ -51,6 +60,9 @@ protected:
 
 	bool AreAnyBodiesHidden() const;
 	bool AreAnyConstraintsHidden() const;
+	bool IsControlEditorTabOpen() const;
+
+	void PostChange();
 
 private:
 	TWeakObjectPtr<UPoseWatchElement> PoseWatchElementBodies;
diff --git a/Engine/Source/Runtime/Engine/Classes/PhysicsEngine/PhysicsAsset.h b/Engine/Source/Runtime/Engine/Classes/PhysicsEngine/PhysicsAsset.h
--- a/Engine/Source/Runtime/Engine/Classes/PhysicsEngine/PhysicsAsset.h
+++ b/Engine/Source/Runtime/Engine/Classes/PhysicsEngine/PhysicsAsset.h
@@ -279,6 +279,7 @@ public:
 	// Find the index of the physics bone that is controlling this graphics bone.
 	ENGINE_API int32		FindControllingBodyIndex(class USkeletalMesh* skelMesh, int32 BoneIndex);
 	ENGINE_API int32		FindParentBodyIndex(class USkeletalMesh * skelMesh, int32 StartBoneIndex) const;
+	ENGINE_API int32		FindParentBodyIndex(const FReferenceSkeleton& RefSkeleton, const int32 StartBoneIndex) const;
 	ENGINE_API int32		FindConstraintIndex(FName ConstraintName);
 	ENGINE_API int32		FindConstraintIndex(FName Bone1Name, FName Bone2Name);
 	FName					FindConstraintBoneName(int32 ConstraintIndex);
@@ -286,6 +287,7 @@ public:
 
 	/** Utility for getting indices of all bodies below (and including) the one with the supplied name. */
 	ENGINE_API void			GetBodyIndicesBelow(TArray<int32>& OutBodyIndices, FName InBoneName, USkeletalMesh* InSkelMesh, bool bIncludeParent = true);
+	ENGINE_API void			GetBodyIndicesBelow(TArray<int32>& OutBodyIndices, const FName InBoneName, const FReferenceSkeleton& RefSkeleton, const bool bIncludeParent = true);
 
 	ENGINE_API void			GetNearestBodyIndicesBelow(TArray<int32> & OutBodyIndices, FName InBoneName, USkeletalMesh * InSkelMesh);
 
diff --git a/Engine/Source/Runtime/Engine/Private/PhysicsEngine/PhysicsAsset.cpp b/Engine/Source/Runtime/Engine/Private/PhysicsEngine/PhysicsAsset.cpp
--- a/Engine/Source/Runtime/Engine/Private/PhysicsEngine/PhysicsAsset.cpp
+++ b/Engine/Source/Runtime/Engine/Private/PhysicsEngine/PhysicsAsset.cpp
@@ -496,13 +496,23 @@ int32	UPhysicsAsset::FindControllingBodyIndex(class USkeletalMesh* skelMesh, int
 	return INDEX_NONE; // Shouldn't reach here.
 }
 
-int32 UPhysicsAsset::FindParentBodyIndex(class USkeletalMesh * skelMesh, int32 StartBoneIndex) const
+int32 UPhysicsAsset::FindParentBodyIndex(class USkeletalMesh* skelMesh, int32 StartBoneIndex) const
+{
+	if (skelMesh)
+	{
+		return FindParentBodyIndex(skelMesh->GetRefSkeleton(), StartBoneIndex);
+	}
+
+	return INDEX_NONE;
+}
+
+int32 UPhysicsAsset::FindParentBodyIndex(const FReferenceSkeleton& RefSkeleton, const int32 StartBoneIndex) const
 {
 	int32 BoneIndex = StartBoneIndex;
-	while ((BoneIndex = skelMesh->GetRefSkeleton().GetParentIndex(BoneIndex)) != INDEX_NONE)
+	while ((BoneIndex = RefSkeleton.GetParentIndex(BoneIndex)) != INDEX_NONE)
 	{
-		FName BoneName = skelMesh->GetRefSkeleton().GetBoneName(BoneIndex);
-		int32 BodyIndex = FindBodyIndex(BoneName);
+		const FName BoneName = RefSkeleton.GetBoneName(BoneIndex);
+		const int32 BodyIndex = FindBodyIndex(BoneName);
 
 		if (StartBoneIndex == BoneIndex)
 			return INDEX_NONE;
@@ -577,7 +587,15 @@ int32 UPhysicsAsset::FindMirroredBone(USkeletalMesh* SkelMesh,  int32 BoneIndex)
 
 void UPhysicsAsset::GetBodyIndicesBelow(TArray<int32>& OutBodyIndices, FName InBoneName, USkeletalMesh* SkelMesh, bool bIncludeParent /*= true*/)
 {
-	int32 BaseIndex = SkelMesh->GetRefSkeleton().FindBoneIndex(InBoneName);
+	if (SkelMesh)
+	{
+		GetBodyIndicesBelow(OutBodyIndices, InBoneName, SkelMesh->GetRefSkeleton(), bIncludeParent);
+	}
+}
+
+void UPhysicsAsset::GetBodyIndicesBelow(TArray<int32>& OutBodyIndices, const FName InBoneName, const FReferenceSkeleton& RefSkeleton, const bool bIncludeParent /*= true*/)
+{
+	const int32 BaseIndex = RefSkeleton.FindBoneIndex(InBoneName);
 
 	if (BaseIndex != INDEX_NONE)
 	{
@@ -590,9 +608,9 @@ void UPhysicsAsset::GetBodyIndicesBelow(TArray<int32>& OutBodyIndices, FName InB
 				continue;
 			}
 			FName TestName = BS->BoneName;
-			int32 TestIndex = SkelMesh->GetRefSkeleton().FindBoneIndex(TestName);
+			int32 TestIndex = RefSkeleton.FindBoneIndex(TestName);
 
-			if ((bIncludeParent && TestIndex == BaseIndex) || SkelMesh->GetRefSkeleton().BoneIsChildOf(TestIndex, BaseIndex))
+			if ((bIncludeParent && TestIndex == BaseIndex) || RefSkeleton.BoneIsChildOf(TestIndex, BaseIndex))
 			{
 				OutBodyIndices.Add(i);
 			}
