[UE] Possible fix for circular dependancy in Physics Control module build files
- Moved module from private to public dependency
- Main CIS Issue 358032

diff --git a/Engine/Plugins/Experimental/PhysicsControl/Source/PhysicsControlEditor/PhysicsControlEditor.Build.cs b/Engine/Plugins/Experimental/PhysicsControl/Source/PhysicsControlEditor/PhysicsControlEditor.Build.cs
--- a/Engine/Plugins/Experimental/PhysicsControl/Source/PhysicsControlEditor/PhysicsControlEditor.Build.cs
+++ b/Engine/Plugins/Experimental/PhysicsControl/Source/PhysicsControlEditor/PhysicsControlEditor.Build.cs
@@ -26,7 +26,8 @@ public class PhysicsControlEditor : ModuleRules
 		PublicDependencyModuleNames.AddRange(
 			new string[]
 			{
-				"Core"
+				"Core",
+				"PhysicsControlUncookedOnly",
 				// ... add other public dependencies that you statically link with here ...
 			}
 			);
@@ -67,7 +68,6 @@ public class PhysicsControlEditor : ModuleRules
 						"BlueprintGraph",
 						"EditorFramework",
 						"Kismet",
-						"PhysicsControlUncookedOnly",
 						"UnrealEd",
 				}
 			);
