[UE] Fix for circular modules dependancy in PhysicsControl
- moved file declaring Module Interface class into the UncookedOnly module
- updated include statements

diff --git a/Engine/Plugins/Experimental/PhysicsControl/Source/PhysicsControlEditor/Public/OperatorEditor/IPhysicsControlOperatorEditorInterface.h b/Engine/Plugins/Experimental/PhysicsControl/Source/PhysicsControlEditor/Public/OperatorEditor/IPhysicsControlOperatorEditorInterface.h
deleted file mode 100644
diff --git a/Engine/Plugins/Experimental/PhysicsControl/Source/PhysicsControlEditor/Public/OperatorEditor/OperatorEditor.h b/Engine/Plugins/Experimental/PhysicsControl/Source/PhysicsControlEditor/Public/OperatorEditor/OperatorEditor.h
--- a/Engine/Plugins/Experimental/PhysicsControl/Source/PhysicsControlEditor/Public/OperatorEditor/OperatorEditor.h
+++ b/Engine/Plugins/Experimental/PhysicsControl/Source/PhysicsControlEditor/Public/OperatorEditor/OperatorEditor.h
@@ -3,7 +3,7 @@
 #pragma once
 
 #include "CoreMinimal.h"
-#include "OperatorEditor/IPhysicsControlOperatorEditorInterface.h"
+#include "IPhysicsControlOperatorEditorInterface.h"
 #include "Containers/Ticker.h"
 
 class FSpawnTabArgs;
diff --git a/Engine/Plugins/Experimental/PhysicsControl/Source/PhysicsControlUncookedOnly/PhysicsControlUncookedOnly.Build.cs b/Engine/Plugins/Experimental/PhysicsControl/Source/PhysicsControlUncookedOnly/PhysicsControlUncookedOnly.Build.cs
--- a/Engine/Plugins/Experimental/PhysicsControl/Source/PhysicsControlUncookedOnly/PhysicsControlUncookedOnly.Build.cs
+++ b/Engine/Plugins/Experimental/PhysicsControl/Source/PhysicsControlUncookedOnly/PhysicsControlUncookedOnly.Build.cs
@@ -18,7 +18,6 @@ public class PhysicsControlUncookedOnly : ModuleRules
 				"Engine",
 				"Persona",
 				"PhysicsControl",
-				"PhysicsControlEditor",
 				"Slate",
 				"SlateCore",
 			}
diff --git a/Engine/Plugins/Experimental/PhysicsControl/Source/PhysicsControlUncookedOnly/Private/AnimGraphNode_RigidBodyWithControl.cpp b/Engine/Plugins/Experimental/PhysicsControl/Source/PhysicsControlUncookedOnly/Private/AnimGraphNode_RigidBodyWithControl.cpp
--- a/Engine/Plugins/Experimental/PhysicsControl/Source/PhysicsControlUncookedOnly/Private/AnimGraphNode_RigidBodyWithControl.cpp
+++ b/Engine/Plugins/Experimental/PhysicsControl/Source/PhysicsControlUncookedOnly/Private/AnimGraphNode_RigidBodyWithControl.cpp
@@ -16,7 +16,7 @@
 #include "AnimNode_RigidBodyWithControl.h"
 #include "EditorModeManager.h"
 #include "IPhysicsAssetRenderInterface.h"
-#include "OperatorEditor/IPhysicsControlOperatorEditorInterface.h"
+#include "IPhysicsControlOperatorEditorInterface.h"
 #include "Kismet2/BlueprintEditorUtils.h"
 #include "PhysicsControlOperatorNameGeneration.h"
 
diff --git a/Engine/Plugins/Experimental/PhysicsControl/Source/PhysicsControlUncookedOnly/Public/IPhysicsControlOperatorEditorInterface.h b/Engine/Plugins/Experimental/PhysicsControl/Source/PhysicsControlUncookedOnly/Public/IPhysicsControlOperatorEditorInterface.h
new file mode 100644
--- /dev/null
+++ b/Engine/Plugins/Experimental/PhysicsControl/Source/PhysicsControlUncookedOnly/Public/IPhysicsControlOperatorEditorInterface.h
@@ -0,0 +1,18 @@
+// Copyright Epic Games, Inc. All Rights Reserved.
+
+#pragma once
+
+#include "CoreMinimal.h"
+#include "Features/IModularFeature.h"
+
+class IPhysicsControlOperatorEditorInterface : public IModularFeature
+{
+public:
+
+	virtual void OpenOperatorNamesTab() = 0;
+	virtual void CloseOperatorNamesTab() = 0;
+	virtual void ToggleOperatorNamesTab() = 0;
+	virtual bool IsOperatorNamesTabOpen() = 0;
+	virtual void RequestRefresh() = 0;
+
+};
\ No newline at end of file
