[UE] Fix right click context menu in physics control edtior operator names tab
- Slate consumes right mouse button presses when scrolling is enabled on a list view, which was preventing the context menu from opening on right click
- Fix was to use the OnContextMenuOpening binding on the list view widget instead of using the OnMouseButtonDown binding

diff --git a/Engine/Plugins/Experimental/PhysicsControl/Source/PhysicsControlEditor/Private/OperatorEditor/SOperatorTreeWidget.cpp b/Engine/Plugins/Experimental/PhysicsControl/Source/PhysicsControlEditor/Private/OperatorEditor/SOperatorTreeWidget.cpp
--- a/Engine/Plugins/Experimental/PhysicsControl/Source/PhysicsControlEditor/Private/OperatorEditor/SOperatorTreeWidget.cpp
+++ b/Engine/Plugins/Experimental/PhysicsControl/Source/PhysicsControlEditor/Private/OperatorEditor/SOperatorTreeWidget.cpp
@@ -143,6 +143,7 @@ void SOperatorTreeWidget::Construct_Internal()
 		.OnGetChildren(this, &SOperatorTreeWidget::OnItemGetChildren)
 		.OnGenerateRow(this, &SOperatorTreeWidget::GenerateItemRow)
 		.SelectionMode(ESelectionMode::Multi)
+		.OnContextMenuOpening(this, &SOperatorTreeWidget::CreateContextMenu)
 		.OnKeyDownHandler(this, &SOperatorTreeWidget::OnKeyDown)
 		.OnMouseButtonDoubleClick(this, &SOperatorTreeWidget::OnItemDoubleClicked);
 
@@ -158,7 +159,7 @@ void SOperatorTreeWidget::Construct_Internal()
 		];
 }
 
-void SOperatorTreeWidget::PushContextMenu(const FVector2D& Position)
+TSharedPtr<SWidget> SOperatorTreeWidget::CreateContextMenu()
 {
 	FMenuBuilder MenuBuilder(true, nullptr, nullptr);
 	MenuBuilder.BeginSection(NAME_None, LOCTEXT("ChannelDrawStyle", "Draw Style"));
@@ -177,12 +178,7 @@ void SOperatorTreeWidget::PushContextMenu(const FVector2D& Position)
 	}
 	MenuBuilder.EndSection();
 
-	FSlateApplication::Get().PushMenu(
-		SharedThis(this),
-		FWidgetPath(),
-		MenuBuilder.MakeWidget(),
-		Position,
-		FPopupTransitionEffect(FPopupTransitionEffect::ContextMenu));
+	return MenuBuilder.MakeWidget();
 }
 
 TSharedRef<ITableRow> SOperatorTreeWidget::GenerateItemRow(ItemType InItem, const TSharedRef<STableViewBase>& OwnerTable)
@@ -238,17 +234,6 @@ FReply SOperatorTreeWidget::OnKeyDown(const FGeometry& MyGeometry, const FKeyEve
 	return FReply::Unhandled();
 }
 
-FReply SOperatorTreeWidget::OnMouseButtonDown(const FGeometry& SenderGeometry, const FPointerEvent& MouseEvent)
-{
-	if (MouseEvent.GetEffectingButton() == EKeys::RightMouseButton)
-	{
-		PushContextMenu(MouseEvent.GetScreenSpacePosition());
-		return FReply::Handled();
-	}
-
-	return FReply::Unhandled();
-}
-
 void SOperatorTreeWidget::OnItemDoubleClicked(ItemType InItem)
 {
 	if (!InItem->OnDoubleClick().IsEventHandled())
diff --git a/Engine/Plugins/Experimental/PhysicsControl/Source/PhysicsControlEditor/Public/OperatorEditor/SOperatorTreeWidget.h b/Engine/Plugins/Experimental/PhysicsControl/Source/PhysicsControlEditor/Public/OperatorEditor/SOperatorTreeWidget.h
--- a/Engine/Plugins/Experimental/PhysicsControl/Source/PhysicsControlEditor/Public/OperatorEditor/SOperatorTreeWidget.h
+++ b/Engine/Plugins/Experimental/PhysicsControl/Source/PhysicsControlEditor/Public/OperatorEditor/SOperatorTreeWidget.h
@@ -49,10 +49,10 @@ public:
 	bool HasTag(const ItemType Item, const FName Tag);
 
 	virtual FReply OnKeyDown(const FGeometry& MyGeometry, const FKeyEvent& InKeyEvent) override;
-	virtual FReply OnMouseButtonDown(const FGeometry& SenderGeometry, const FPointerEvent& MouseEvent) override;
 
 	void OnItemDoubleClicked(ItemType InItem);
-	void PushContextMenu(const FVector2D& Position);
+	TSharedPtr<SWidget> CreateContextMenu();
+
 	void RequestRefresh();
 
 	void CopySelectedItemsNamesToClipboard() const;
