Fix for duplicate localization keys

diff --git a/Engine/Plugins/Experimental/PhysicsControl/Source/PhysicsControlEditor/Private/OperatorEditor/SOperatorTreeWidget.cpp b/Engine/Plugins/Experimental/PhysicsControl/Source/PhysicsControlEditor/Private/OperatorEditor/SOperatorTreeWidget.cpp
--- a/Engine/Plugins/Experimental/PhysicsControl/Source/PhysicsControlEditor/Private/OperatorEditor/SOperatorTreeWidget.cpp
+++ b/Engine/Plugins/Experimental/PhysicsControl/Source/PhysicsControlEditor/Private/OperatorEditor/SOperatorTreeWidget.cpp
@@ -165,14 +165,14 @@ TSharedPtr<SWidget> SOperatorTreeWidget::CreateContextMenu()
 	MenuBuilder.BeginSection(NAME_None, LOCTEXT("ChannelDrawStyle", "Draw Style"));
 	{
 		MenuBuilder.AddMenuEntry(
-			LOCTEXT("OperatorTreeCopyName", "Copy Names"),
-			LOCTEXT("OperatorTreeCopyName_Tooltip", "Copy the names of all the selected operators to the clipboard."),
+			LOCTEXT("OperatorTreeCopyNames", "Copy Names"),
+			LOCTEXT("OperatorTreeCopyNames_Tooltip", "Copy the names of all the selected operators to the clipboard."),
 			FSlateIcon(FAppStyle::GetAppStyleSetName(), "Icons.Copy"),
 			FUIAction(FExecuteAction::CreateSP(this, &SOperatorTreeWidget::CopySelectedItemsNamesToClipboard)));
 
 		MenuBuilder.AddMenuEntry(
-			LOCTEXT("OperatorTreeCopyName", "Copy All"),
-			LOCTEXT("OperatorTreeCopyName_Tooltip", "Copy the full description of all the selected operators to the clipboard."),
+			LOCTEXT("OperatorTreeCopyAll", "Copy All"),
+			LOCTEXT("OperatorTreeCopyAll_Tooltip", "Copy the full description of all the selected operators to the clipboard."),
 			FSlateIcon(FAppStyle::GetAppStyleSetName(), "Icons.Copy"),
 			FUIAction(FExecuteAction::CreateSP(this, &SOperatorTreeWidget::CopySelectedItemsDescriptionsToClipboard)));
 	}
