Support better interactions with collidable objects in the RigidBodyWithControl

diff --git a/Engine/Plugins/Experimental/PhysicsControl/Source/PhysicsControl/Private/AnimNode_RigidBodyWithControl.cpp b/Engine/Plugins/Experimental/PhysicsControl/Source/PhysicsControl/Private/AnimNode_RigidBodyWithControl.cpp
--- a/Engine/Plugins/Experimental/PhysicsControl/Source/PhysicsControl/Private/AnimNode_RigidBodyWithControl.cpp
+++ b/Engine/Plugins/Experimental/PhysicsControl/Source/PhysicsControl/Private/AnimNode_RigidBodyWithControl.cpp
@@ -126,6 +126,7 @@ FAnimNode_RigidBodyWithControl::FAnimNode_RigidBodyWithControl()
 	, ComponentAppliedLinearAccClamp(10000.0f)
 	, SimSpaceSettings()
 	, CachedBoundsScale(1.2f)
+	, UpdateCacheEveryFrame(true)
 	, BaseBoneRef()
 	, OverlapChannel(ECC_WorldStatic)
 	, SimulationSpace(ESimulationSpace::ComponentSpace)
@@ -842,7 +843,7 @@ void FAnimNode_RigidBodyWithControl::EvaluateSkeletalControl_AnyThread(FComponen
 				SimulationLinearAcceleration,
 				SimulationAngularAcceleration);
 
-			UpdateWorldObjects(SimulationTransform);
+			UpdateWorldObjects(SimulationTransform, DeltaSeconds);
 			UpdateClothColliderObjects(SimulationTransform);
 
 			PhysicsSimulation->UpdateSimulationSpace(
@@ -1347,14 +1348,15 @@ void FAnimNode_RigidBodyWithControl::UpdateWorldGeometry(const UWorld& World, co
 
 	// If we have moved outside of the bounds we checked for world objects we need to gather new world objects
 	FSphere Bounds = SKC.CalcBounds(SKC.GetComponentToWorld()).GetSphere();
-	if (!Bounds.IsInside(CachedBounds))
+	if (!Bounds.IsInside(CachedBounds) || UpdateCacheEveryFrame)
 	{
 		// Since the cached bounds are no longer valid, update them.
 		CachedBounds = Bounds;
 		CachedBounds.W *= CachedBoundsScale;
 
-		// Cache the PhysScene and World for use in UpdateWorldForces and CollectWorldObjects
-		// When these are non-null it is an indicator that we need to update the collected world objects list
+		// Cache the PhysScene and World for use in UpdateWorldForces and CollectWorldObjects. When
+		// these are non-null it is an indicator that we need to update the collected world objects
+		// list in CollectWorldObjects
 		PhysScene = World.GetPhysicsScene();
 		UnsafeWorld = &World;
 		UnsafeOwner = SKC.GetOwner();
@@ -1689,9 +1691,16 @@ void FAnimNode_RigidBodyWithControl::CollectWorldObjects()
 					const bool bIsSelf = (UnsafeOwner == OverlapComp->GetOwner());
 					if (!bIsSelf)
 					{
-						// Create a kinematic actor. Not using Static as world-static objects may move in the simulation's frame of reference
-						ImmediatePhysics::FActorHandle* ActorHandle = PhysicsSimulation->CreateActor(ImmediatePhysics::EActorType::KinematicActor, &OverlapComp->BodyInstance, OverlapComp->GetComponentTransform());
+						// Note that the TM here isn't correct unless we are actually using
+						// world-space simulation, but it will be updated.
+						FTransform TM = OverlapComp->GetComponentTransform();
+						// Create a kinematic actor. Not using Static as world-static objects may
+						// move in the simulation's frame of reference
+						ImmediatePhysics::FActorHandle* ActorHandle = PhysicsSimulation->CreateActor(ImmediatePhysics::EActorType::KinematicActor, &OverlapComp->BodyInstance, TM);
 						PhysicsSimulation->AddToCollidingPairs(ActorHandle);
+						// This seems necessary in order for subsequent calls to set the transform
+						// to update things correctly.
+						//ActorHandle->InitWorldTransform(TM);
 						ComponentsInSim.Add(OverlapComp, FWorldObject(ActorHandle, ComponentsInSimTick));
 					}
 				}
@@ -1750,32 +1759,50 @@ void FAnimNode_RigidBodyWithControl::PurgeExpiredWorldObjects()
 	}
 }
 
-// Update the transforms of the world objects we added to the sim. This is required
-// if we have a component- or bone-space simulation as even world-static objects
-// will be moving in the simulation's frame of reference.
-void FAnimNode_RigidBodyWithControl::UpdateWorldObjects(const FTransform& SpaceTransform)
+// Update the transforms of the world objects we added to the sim. This could be because we're a
+// component-based simulation and need to update stationary objects into our space, or simply
+// because the objects are moving.
+void FAnimNode_RigidBodyWithControl::UpdateWorldObjects(const FTransform& SpaceTransform, const float DeltaSeconds)
 {
 	LLM_SCOPE_BYNAME(TEXT("Animation/RigidBodyWithControl")); 
 
-	if (SimulationSpace != ESimulationSpace::WorldSpace)
-	{
-		for (const auto& WorldEntry : ComponentsInSim)
-		{ 
-			const UPrimitiveComponent* OverlapComp = WorldEntry.Key;
-			if (OverlapComp != nullptr)
-			{
-				ImmediatePhysics::FActorHandle* ActorHandle = WorldEntry.Value.ActorHandle;
+	for (const auto& WorldEntry : ComponentsInSim)
+	{ 
+		const UPrimitiveComponent* OverlapComp = WorldEntry.Key;
+		if (OverlapComp != nullptr)
+		{
+			ImmediatePhysics::FActorHandle* ActorHandle = WorldEntry.Value.ActorHandle;
 
-				// Calculate the sim-space transform of this object
-				const FTransform CompWorldTransform = OverlapComp->BodyInstance.GetUnrealWorldTransform();
-				FTransform CompSpaceTransform;
-				CompSpaceTransform.SetTranslation(SpaceTransform.InverseTransformPosition(CompWorldTransform.GetLocation()));
-				CompSpaceTransform.SetRotation(SpaceTransform.InverseTransformRotation(CompWorldTransform.GetRotation()));
-				CompSpaceTransform.SetScale3D(FVector::OneVector);	// TODO - sort out scale for world objects in local sim
+			// Calculate the sim-space transform of this object
+			const FTransform CompWorldTransform = OverlapComp->BodyInstance.GetUnrealWorldTransform();
+			const FVector CompSpacePosition = SpaceTransform.InverseTransformPosition(CompWorldTransform.GetLocation());
+			const FQuat CompSpaceOrientation = SpaceTransform.InverseTransformRotation(CompWorldTransform.GetRotation());
 
-				// Update the sim's copy of the world object
-				ActorHandle->SetWorldTransform(CompSpaceTransform);
-			}
+			// TODO - sort out scale for world objects in local sim
+			FTransform CompSpaceTransform(CompSpaceOrientation, CompSpacePosition, FVector::OneVector);
+
+			// This is what we don't do as it does not track velocity. Also, it only seems to
+			// work if InitWorldTransform has been called first (after actor creation).
+			// ActorHandle->SetWorldTransform(CompSpaceTransform);
+
+			// We also need to set the velocity - either because the object or the space is
+			// moving. Ideally would do this by setting the velocity (with transformations), but
+			// it's not possible to set the velocity of kinematics. Use the kinematic target to
+			// force the velocities. This should handle initialization and teleportation too.
+			const FVector Velocity = OverlapComp->BodyInstance.GetUnrealWorldVelocity();
+			const FVector AngularVelocity = OverlapComp->BodyInstance.GetUnrealWorldAngularVelocityInRadians();
+
+			const FVector CompSpaceVelocity = SpaceTransform.GetRotation().Inverse().RotateVector(Velocity);
+			const FVector CompSpaceAngularVelocity = SpaceTransform.GetRotation().Inverse().RotateVector(AngularVelocity);
+
+			const FVector PrevCompSpacePosition = CompSpacePosition - CompSpaceVelocity * DeltaSeconds;
+			const FQuat CompSpaceDeltaOrientation = FQuat::MakeFromRotationVector(CompSpaceAngularVelocity * DeltaSeconds);
+			const FQuat PrevCompSpaceOrientation = CompSpaceDeltaOrientation.Inverse() * CompSpaceOrientation;
+
+			const FTransform PrevCompSpaceTransform(PrevCompSpaceOrientation, PrevCompSpacePosition, FVector::OneVector);
+
+			ActorHandle->InitWorldTransform(PrevCompSpaceTransform);
+			ActorHandle->SetKinematicTarget(CompSpaceTransform);
 		}
 	}
 }
diff --git a/Engine/Plugins/Experimental/PhysicsControl/Source/PhysicsControl/Public/AnimNode_RigidBodyWithControl.h b/Engine/Plugins/Experimental/PhysicsControl/Source/PhysicsControl/Public/AnimNode_RigidBodyWithControl.h
--- a/Engine/Plugins/Experimental/PhysicsControl/Source/PhysicsControl/Public/AnimNode_RigidBodyWithControl.h
+++ b/Engine/Plugins/Experimental/PhysicsControl/Source/PhysicsControl/Public/AnimNode_RigidBodyWithControl.h
@@ -117,13 +117,21 @@ public:
 
 
 	/**
-	 * Scale of cached bounds (vs. actual bounds).
+	 * Scale of cached bounds (vs. actual bounds) used for obtaining the list of objects we might collide with.
 	 * Increasing this may improve performance, but overlaps may not work as well.
 	 * (A value of 1.0 effectively disables cached bounds).
 	 */
 	UPROPERTY(EditAnywhere, Category = Settings, meta = (ClampMin="1.0", ClampMax="2.0"))
 	float CachedBoundsScale;
 
+	/**
+	 * This can be used to force the cached collision objects to be updated every frame, even when we are not 
+	 * moving. This can be expensive, but may be needed if we wish to detect collisions with objects that are 
+	 * moving, otherwise the cache will only be updated if/when we move.
+	 */
+	UPROPERTY(EditAnywhere, Category = Settings, meta = (PinHiddenByDefault))
+	bool UpdateCacheEveryFrame;
+
 	/** Matters if SimulationSpace is BaseBone */
 	UPROPERTY(EditAnywhere, Category = Settings)
 	FBoneReference BaseBoneRef;
@@ -386,7 +394,7 @@ private:
 	void PurgeExpiredWorldObjects();
 
 	// Update sim-space transforms of world objects
-	void UpdateWorldObjects(const FTransform& SpaceTransform);
+	void UpdateWorldObjects(const FTransform& SpaceTransform, const float DeltaSeconds);
 
 	// Advances the simulation by a given timestep
 	void RunPhysicsSimulation(float DeltaSeconds, const FVector& SimSpaceGravity);
@@ -505,9 +513,10 @@ private:
 		uint32 BoneIndex; // Bone within parent skeleton that drives physics actors transform.
 	};
 
-	// List of actors in the sim that represent objects collected from other parts of this character.
+	// List of cloth colliders in the sim that represent objects collected from other parts of this character.
 	TArray<FClothCollider> ClothColliders; 
 	
+	// List of actors in the sim that represent objects collected from the world.
 	TMap<const UPrimitiveComponent*, FWorldObject> ComponentsInSim;
 	int32 ComponentsInSimTick;
 
