Support for exposing the parent of the world-space controls in Physics Control Component

diff --git a/Engine/Plugins/Experimental/PhysicsControl/Source/PhysicsControl/Private/PhysicsControlComponent.cpp b/Engine/Plugins/Experimental/PhysicsControl/Source/PhysicsControl/Private/PhysicsControlComponent.cpp
--- a/Engine/Plugins/Experimental/PhysicsControl/Source/PhysicsControl/Private/PhysicsControlComponent.cpp
+++ b/Engine/Plugins/Experimental/PhysicsControl/Source/PhysicsControl/Private/PhysicsControlComponent.cpp
@@ -288,9 +288,10 @@ FName UPhysicsControlComponent::CreateControl(
 	const FPhysicsControlData     ControlData, 
 	const FPhysicsControlTarget   ControlTarget, 
 	const FPhysicsControlSettings ControlSettings,
-	const FName                   Set)
+	const FName                   Set,
+	const FString                 NamePrefix)
 {
-	const FName Name = Implementation->GetUniqueControlName(ParentBoneName, ChildBoneName);
+	const FName Name = Implementation->GetUniqueControlName(ParentBoneName, ChildBoneName, NamePrefix);
 	if (CreateNamedControl(
 		Name, ParentMeshComponent, ParentBoneName, ChildMeshComponent, ChildBoneName, 
 		ControlData, ControlTarget, ControlSettings, Set))
@@ -609,7 +610,10 @@ TMap<FName, FPhysicsControlNames> UPhysicsControlComponent::CreateControlsFromLi
 	const TMap<FName, FPhysicsControlLimbBones>& LimbBones,
 	const EPhysicsControlType                    ControlType,
 	const FPhysicsControlData                    ControlData,
-	const FPhysicsControlSettings                ControlSettings)
+	const FPhysicsControlSettings                ControlSettings,
+	UMeshComponent*                              WorldComponent,
+	FName                                        WorldBoneName,
+	FString                                      NamePrefix)
 {
 	TMap<FName, FPhysicsControlNames> Result;
 	Result.Reserve(LimbBones.Num());
@@ -631,7 +635,7 @@ TMap<FName, FPhysicsControlNames> UPhysicsControlComponent::CreateControlsFromLi
 			continue;
 		}
 
-		USkeletalMeshComponent* ParentMeshComponent =
+		USkeletalMeshComponent* ParentSkeletalMeshComponent =
 			(ControlType == EPhysicsControlType::ParentSpace) ? BonesInLimb.SkeletalMeshComponent : nullptr;
 
 		const int32 NumBonesInLimb = BonesInLimb.BoneNames.Num();
@@ -640,6 +644,9 @@ TMap<FName, FPhysicsControlNames> UPhysicsControlComponent::CreateControlsFromLi
 		LimbResult.Names.Reserve(NumBonesInLimb);
 		AllControls.Names.Reserve(AllControls.Names.Num() + NumBonesInLimb);
 
+		FString SetName = 
+			NamePrefix + GetControlTypeName(ControlType).ToString().Append("_").Append(LimbName.ToString());
+
 		for (int32 BoneIndex = 0 ; BoneIndex != NumBonesInLimb ; ++BoneIndex)
 		{
 			// Don't create the parent space control if it's the first bone in a limb that had bIncludeParentBone
@@ -652,19 +659,28 @@ TMap<FName, FPhysicsControlNames> UPhysicsControlComponent::CreateControlsFromLi
 			FBodyInstance* ChildBone = BonesInLimb.SkeletalMeshComponent->GetBodyInstance(ChildBoneName);
 
 			FName ParentBoneName;
-			if (ParentMeshComponent)
+			if (ParentSkeletalMeshComponent)
 			{
 				ParentBoneName = UE::PhysicsControlComponent::GetPhysicalParentBone(
-					ParentMeshComponent, ChildBoneName);
+					ParentSkeletalMeshComponent, ChildBoneName);
 				if (ParentBoneName.IsNone())
 				{
 					continue;
 				}
 			}
+
+			UMeshComponent* ParentMeshComponent = ParentSkeletalMeshComponent;
+			if (!ParentMeshComponent && WorldComponent)
+			{
+				ParentMeshComponent = WorldComponent;
+				ParentBoneName = WorldBoneName;
+			}
+
 			const FName ControlName = CreateControl(
 				ParentMeshComponent, ParentBoneName, BonesInLimb.SkeletalMeshComponent, ChildBoneName,
 				ControlData, FPhysicsControlTarget(), ControlSettings, 
-				FName(GetControlTypeName(ControlType).ToString().Append("_").Append(LimbName.ToString())));
+				FName(SetName));
+
 			if (!ControlName.IsNone())
 			{
 				LimbResult.Names.Add(ControlName);
@@ -820,6 +836,56 @@ void UPhysicsControlComponent::SetControlsInSetEnabled(FName SetName, bool bEnab
 	SetControlsEnabled(GetControlNamesInSet(SetName), bEnable);
 }
 
+//======================================================================================================================
+bool UPhysicsControlComponent::SetControlParent(
+	const FName     Name,
+	UMeshComponent* ParentMeshComponent,
+	const FName     ParentBoneName)
+{
+	FPhysicsControlRecord* Record = Implementation->FindControlRecord(Name);
+	if (Record)
+	{
+		if (USkeletalMeshComponent* SkeletalMeshComponent = 
+			Cast<USkeletalMeshComponent>(Record->PhysicsControl.ParentMeshComponent))
+		{
+			Implementation->RemoveSkeletalMeshReferenceForCaching(SkeletalMeshComponent);
+		}
+
+		if (USkeletalMeshComponent* SkeletalMeshComponent = Cast<USkeletalMeshComponent>(ParentMeshComponent))
+		{
+			Implementation->AddSkeletalMeshReferenceForCaching(SkeletalMeshComponent);
+		}
+
+		Record->PhysicsControl.ParentMeshComponent = ParentMeshComponent;
+		Record->PhysicsControl.ParentBoneName = ParentBoneName;
+		return Record->InitConstraint(this, Name);
+	}
+	return false;
+}
+
+//======================================================================================================================
+void UPhysicsControlComponent::SetControlParents(
+		const TArray<FName>& Names,
+		UMeshComponent*      ParentMeshComponent,
+		const FName          ParentBoneName)
+{
+	for (FName Name : Names)
+	{
+		SetControlParent(Name, ParentMeshComponent, ParentBoneName);
+	}
+}
+
+//======================================================================================================================
+void UPhysicsControlComponent::SetControlParentsInSet(
+	const FName     SetName,
+	UMeshComponent* ParentMeshComponent,
+	const FName     ParentBoneName)
+{
+	SetControlParents(GetControlNamesInSet(SetName), ParentMeshComponent, ParentBoneName);
+}
+
+
+
 //======================================================================================================================
 // Note - some params passed by value to allow them to be unconnected and enable inline BP editing
 bool UPhysicsControlComponent::SetControlData(
@@ -1827,14 +1893,16 @@ void UPhysicsControlComponent::CreateControlsAndBodyModifiersFromLimbBones(
 	const FPhysicsControlSettings               ParentSpaceControlSettings,
 	const EPhysicsMovementType                  PhysicsMovementType,
 	const float                                 GravityMultiplier,
-	const float                                 PhysicsBlendWeight)
+	const float                                 PhysicsBlendWeight,
+	UMeshComponent*                             WorldComponent,
+	FName                                       WorldBoneName)
 {
 	TMap<FName, FPhysicsControlLimbBones> LimbBones = 
 		GetLimbBonesFromSkeletalMesh(SkeletalMeshComponent, LimbSetupData);
 
 	LimbWorldSpaceControls = CreateControlsFromLimbBones(
 		AllWorldSpaceControls, LimbBones, EPhysicsControlType::WorldSpace, 
-		WorldSpaceControlData, WorldSpaceControlSettings);
+		WorldSpaceControlData, WorldSpaceControlSettings, WorldComponent, WorldBoneName);
 
 	LimbParentSpaceControls = CreateControlsFromLimbBones(
 		AllParentSpaceControls, LimbBones, EPhysicsControlType::ParentSpace,
diff --git a/Engine/Plugins/Experimental/PhysicsControl/Source/PhysicsControl/Private/PhysicsControlComponentImpl.cpp b/Engine/Plugins/Experimental/PhysicsControl/Source/PhysicsControl/Private/PhysicsControlComponentImpl.cpp
--- a/Engine/Plugins/Experimental/PhysicsControl/Source/PhysicsControl/Private/PhysicsControlComponentImpl.cpp
+++ b/Engine/Plugins/Experimental/PhysicsControl/Source/PhysicsControl/Private/PhysicsControlComponentImpl.cpp
@@ -423,9 +423,10 @@ bool FPhysicsControlComponentImpl::RemoveSkeletalMeshReferenceForModifier(
 }
 
 //======================================================================================================================
-FName FPhysicsControlComponentImpl::GetUniqueControlName(const FName ParentBoneName, const FName ChildBoneName) const
+FName FPhysicsControlComponentImpl::GetUniqueControlName(
+	const FName ParentBoneName, const FName ChildBoneName, const FString& NamePrefix) const
 {
-	FString NameBase = TEXT("");
+	FString NameBase = NamePrefix;
 	if (!ParentBoneName.IsNone())
 	{
 		NameBase += ParentBoneName.ToString() + TEXT("_");
@@ -521,6 +522,16 @@ void FPhysicsControlComponentImpl::CalculateControlTargetData(
 			bHaveParentBoneData = GetBoneData(
 				ParentBoneData, ParentSkeletalMeshComponent, Record.PhysicsControl.ParentBoneName);
 		}
+		else if (Record.PhysicsControl.ParentMeshComponent)
+		{
+			FTransform ParentTM = Record.PhysicsControl.ParentMeshComponent->GetComponentTransform();
+			ParentBoneData.Position = ParentTM.GetLocation();
+			ParentBoneData.Orientation = ParentTM.GetRotation();
+			ParentBoneData.Velocity = Record.PhysicsControl.ParentMeshComponent->GetPhysicsLinearVelocity();
+			ParentBoneData.AngularVelocity = Record.PhysicsControl.ParentMeshComponent->GetPhysicsAngularVelocityInRadians();
+			bHaveParentBoneData = true;
+			bCalculateVelocity = false;
+		}
 
 		// Note that the TargetTM/velocity calculated so far are supposed to be interpreted as
 		// expressed relative to the skeletal animation pose.
@@ -723,6 +734,8 @@ void FPhysicsControlComponentImpl::ApplyControl(FPhysicsControlRecord& Record)
 		return;
 	}
 
+	// TODO UE-197351 disable constraint features fully when strengths/damping are zero
+
 	// Set strengths etc
 	if (ApplyControlStrengths(Record, ConstraintInstance))
 	{
@@ -735,6 +748,7 @@ void FPhysicsControlComponentImpl::ApplyControl(FPhysicsControlRecord& Record)
 		ConstraintInstance->SetAngularOrientationTarget(TargetTM.GetRotation());
 		ConstraintInstance->SetLinearVelocityTarget(TargetVelocity);
 		ConstraintInstance->SetAngularVelocityTarget(TargetAngularVelocity / UE_TWO_PI); // In rev/sec
+		ConstraintInstance->SetParentDominates(Record.PhysicsControl.ControlSettings.bOnlyControlChildObject);
 
 		if (ParentBodyInstance)
 		{
diff --git a/Engine/Plugins/Experimental/PhysicsControl/Source/PhysicsControl/Private/PhysicsControlComponentImpl.h b/Engine/Plugins/Experimental/PhysicsControl/Source/PhysicsControl/Private/PhysicsControlComponentImpl.h
--- a/Engine/Plugins/Experimental/PhysicsControl/Source/PhysicsControl/Private/PhysicsControlComponentImpl.h
+++ b/Engine/Plugins/Experimental/PhysicsControl/Source/PhysicsControl/Private/PhysicsControlComponentImpl.h
@@ -156,7 +156,7 @@ public:
 		const EDestroyBehavior DestroyBehavior = EDestroyBehavior::RemoveRecord);
 
 	/** Calculates and returns a unique name based on the bone names and any existing records */
-	FName GetUniqueControlName(const FName ParentBoneName, const FName ChildBoneName) const;
+	FName GetUniqueControlName(const FName ParentBoneName, const FName ChildBoneName, const FString& NamePrefix) const;
 
 	/** Calculates and returns a unique name based on the bone name and any existing records */
 	FName GetUniqueBodyModifierName(const FName BoneName) const;
diff --git a/Engine/Plugins/Experimental/PhysicsControl/Source/PhysicsControl/Private/PhysicsControlRecord.cpp b/Engine/Plugins/Experimental/PhysicsControl/Source/PhysicsControl/Private/PhysicsControlRecord.cpp
--- a/Engine/Plugins/Experimental/PhysicsControl/Source/PhysicsControl/Private/PhysicsControlRecord.cpp
+++ b/Engine/Plugins/Experimental/PhysicsControl/Source/PhysicsControl/Private/PhysicsControlRecord.cpp
@@ -18,14 +18,13 @@ void FPhysicsControlState::Reset()
 }
 
 //======================================================================================================================
-// A constraint created with identity frames will just have the child frame at the mesh origin,
-// which is not necessarily where the center of gravity is.
-FConstraintInstance* FPhysicsControlRecord::CreateConstraint(UObject* ConstraintDebugOwner, FName ControlName)
+bool FPhysicsControlRecord::InitConstraint(UObject* ConstraintDebugOwner, FName ControlName)
 {
 	if (!PhysicsControlState.ConstraintInstance)
 	{
-		PhysicsControlState.ConstraintInstance = MakeShared<FConstraintInstance>();
+		return false;
 	}
+
 	FConstraintInstance* ConstraintInstance = PhysicsControlState.ConstraintInstance.Get();
 
 	FBodyInstance* ParentBody = UE::PhysicsControlComponent::GetBodyInstance(
@@ -35,40 +34,55 @@ FConstraintInstance* FPhysicsControlRecord::CreateConstraint(UObject* Constraint
 
 	if (PhysicsControl.ParentMeshComponent && !PhysicsControl.ParentBoneName.IsNone() && !ParentBody)
 	{
-		UE_LOG(LogPhysicsControlComponent, Warning, 
-			TEXT("Failed to find expected parent body %s when making constraint for control %s"), 
+		UE_LOG(LogPhysicsControlComponent, Warning,
+			TEXT("Failed to find expected parent body %s when making constraint for control %s"),
 			*PhysicsControl.ParentBoneName.ToString(),
 			*ControlName.ToString());
-		return nullptr;
+		return false;
 	}
 	if (PhysicsControl.ChildMeshComponent && !PhysicsControl.ChildBoneName.IsNone() && !ChildBody)
 	{
-		UE_LOG(LogPhysicsControlComponent, Warning, 
-			TEXT("Failed to find expected child body %s when making constraint for control %s"), 
+		UE_LOG(LogPhysicsControlComponent, Warning,
+			TEXT("Failed to find expected child body %s when making constraint for control %s"),
 			*PhysicsControl.ChildBoneName.ToString(),
 			*ControlName.ToString());
-		return nullptr;
+		return false;
 	}
 
 	ConstraintInstance->InitConstraint(ChildBody, ParentBody, 1.0f, ConstraintDebugOwner);
-
-	// Ensure the control point is set
-	UpdateConstraintControlPoint();
-
+	ConstraintInstance->SetDisableCollision(PhysicsControl.ControlSettings.bDisableCollision);
+	// These things won't change so set them once here
 	ConstraintInstance->SetLinearXMotion(ELinearConstraintMotion::LCM_Free);
 	ConstraintInstance->SetLinearYMotion(ELinearConstraintMotion::LCM_Free);
 	ConstraintInstance->SetLinearZMotion(ELinearConstraintMotion::LCM_Free);
 	ConstraintInstance->SetAngularSwing1Motion(EAngularConstraintMotion::ACM_Free);
 	ConstraintInstance->SetAngularSwing2Motion(EAngularConstraintMotion::ACM_Free);
 	ConstraintInstance->SetAngularTwistMotion(EAngularConstraintMotion::ACM_Free);
-
-	ConstraintInstance->SetLinearPositionDrive(true, true, true);
-	ConstraintInstance->SetLinearVelocityDrive(true, true, true);
 	ConstraintInstance->SetAngularDriveMode(EAngularDriveMode::SLERP);
+
 	ConstraintInstance->SetOrientationDriveSLERP(true);
 	ConstraintInstance->SetAngularVelocityDriveSLERP(true);
+	ConstraintInstance->SetLinearPositionDrive(true, true, true);
+	ConstraintInstance->SetLinearVelocityDrive(true, true, true);
 
-	ConstraintInstance->SetDisableCollision(PhysicsControl.ControlSettings.bDisableCollision);
+	// Ensure the control point is set
+	UpdateConstraintControlPoint();
+
+	return true;
+}
+
+//======================================================================================================================
+// A constraint created with identity frames will just have the child frame at the mesh origin,
+// which is not necessarily where the center of gravity is.
+FConstraintInstance* FPhysicsControlRecord::CreateConstraint(UObject* ConstraintDebugOwner, FName ControlName)
+{
+	if (!PhysicsControlState.ConstraintInstance)
+	{
+		PhysicsControlState.ConstraintInstance = MakeShared<FConstraintInstance>();
+	}
+	FConstraintInstance* ConstraintInstance = PhysicsControlState.ConstraintInstance.Get();
+
+	InitConstraint(ConstraintDebugOwner, ControlName);
 
 	return ConstraintInstance;
 }
diff --git a/Engine/Plugins/Experimental/PhysicsControl/Source/PhysicsControl/Private/PhysicsControlRecord.h b/Engine/Plugins/Experimental/PhysicsControl/Source/PhysicsControl/Private/PhysicsControlRecord.h
--- a/Engine/Plugins/Experimental/PhysicsControl/Source/PhysicsControl/Private/PhysicsControlRecord.h
+++ b/Engine/Plugins/Experimental/PhysicsControl/Source/PhysicsControl/Private/PhysicsControlRecord.h
@@ -36,6 +36,11 @@ struct FPhysicsControlRecord
 	 */
 	FConstraintInstance* CreateConstraint(UObject* ConstraintDebugOwner, FName ControlName);
 
+	/** 
+	 * Initialises the constraint with the bodies assigned to us. Requires that the constraint itself 
+	 * has been made. Returns true/false on success/failure */
+	bool InitConstraint(UObject* ConstraintDebugOwner, FName ControlName);
+
 	/** Ensures the constraint frame matches the control point in the record. */
 	void UpdateConstraintControlPoint();
 
diff --git a/Engine/Plugins/Experimental/PhysicsControl/Source/PhysicsControl/Public/PhysicsControlComponent.h b/Engine/Plugins/Experimental/PhysicsControl/Source/PhysicsControl/Public/PhysicsControlComponent.h
--- a/Engine/Plugins/Experimental/PhysicsControl/Source/PhysicsControl/Public/PhysicsControlComponent.h
+++ b/Engine/Plugins/Experimental/PhysicsControl/Source/PhysicsControl/Public/PhysicsControlComponent.h
@@ -93,6 +93,7 @@ public:
 	 *        gets added to the set "All"
 	 * @param bEnabled If true then the control will be enabled immediately. If false you will need to call
 	 *        SetControlEnabled(true) in order to activate it.
+	 * @param NamePrefix Optional string that is prefixed to the control that is created.
 	 * @return The name of the new control
 	 */
 	UFUNCTION(BlueprintCallable, Category = PhysicsControl)
@@ -104,7 +105,8 @@ public:
 		const FPhysicsControlData     ControlData,
 		const FPhysicsControlTarget   ControlTarget, 
 		const FPhysicsControlSettings ControlSettings,
-		FName                         Set
+		FName                         Set,
+		FString                       NamePrefix = TEXT("")
 	);
 
 	/**
@@ -263,12 +265,15 @@ public:
 	 *
 	 * @param AllControls A single container for all the controls that have been created
 	 * @param LimbBones A map relating the limbs and the bones that they contain. Typically create this 
-	 *                  using GetLimbBonesFromSkeletalMesh
+	 *        using GetLimbBonesFromSkeletalMesh
 	 * @param ControlType What type of control to create. This determines what the parent will be for each control
 	 * @param ControlData Describes the initial strength etc of the new control
 	 * @param ControlSettings General settings for the control
-	 * @param bEnabled If true then the control will be enabled immediately. If false you will need to call
-	 *                 SetControlEnabled(true) in order to enable it.
+	 * @param WorldComponent Optional component to use as the parent object for any "world-space" controls 
+	 *        that are created. Will be ignored if the controls being created are not world-space.
+	 * @param WorldBoneName Additional bone name to identify the world object if the WorldComponent is actually 
+	 *        a skeletal mesh component.
+	 * @param NamePrefix Optional string that is prefixed to each control that is created. 
 	 * @return A map containing the controls for each limb
 	 */
 	UFUNCTION(BlueprintCallable, Category = PhysicsControl)
@@ -277,7 +282,10 @@ public:
 		const TMap<FName, FPhysicsControlLimbBones>& LimbBones,
 		const EPhysicsControlType                    ControlType,
 		const FPhysicsControlData                    ControlData,
-		const FPhysicsControlSettings                ControlSettings);
+		const FPhysicsControlSettings                ControlSettings,
+		UMeshComponent*                              WorldComponent = nullptr,
+		FName                                        WorldBoneName = NAME_None,
+		FString                                      NamePrefix = TEXT(""));
 
 	/**
 	 * Creates a collection of ParentSpace controls controlling a skeletal mesh, grouped together in limbs, initializing
@@ -334,6 +342,32 @@ public:
 	UFUNCTION(BlueprintCallable, Category = PhysicsControl)
 	void DestroyControlsInSet(const FName Set);
 
+	/**
+	 * Updates the parent object part of a control. Note that this won't change the name of the control (which may
+	 * subsequently be misleading), or any set it is included in, etc.
+	 */
+	UFUNCTION(BlueprintCallable, Category = PhysicsControl)
+	bool SetControlParent(
+		const FName     Name, 
+		UMeshComponent* ParentMeshComponent,
+		const FName     ParentBoneName);
+
+	/**
+	 * Updates the parent object part of controls. Note that this won't change the name of the controls (which may
+	 * subsequently be misleading), or any set they are included in, etc.
+	 */
+	UFUNCTION(BlueprintCallable, Category = PhysicsControl)
+	void SetControlParents(
+		const TArray<FName>& Names,
+		UMeshComponent*      ParentMeshComponent,
+		const FName          ParentBoneName);
+
+	UFUNCTION(BlueprintCallable, Category = PhysicsControl)
+	void SetControlParentsInSet(
+		const FName     Set,
+		UMeshComponent* ParentMeshComponent,
+		const FName     ParentBoneName);
+
 	/**
 	 * Modifies an existing control data - i.e. the strengths etc of the control driving towards the target
 	 *
@@ -1284,8 +1318,9 @@ public:
 		const FPhysicsControlSettings               ParentSpaceControlSettings,
 		const EPhysicsMovementType                  PhysicsMovementType = EPhysicsMovementType::Static,
 		const float                                 GravityMultiplier = 1.0f,
-		const float                                 PhysicsBlendWeight = 1.0f
-		);
+		const float                                 PhysicsBlendWeight = 1.0f,
+		UMeshComponent*                             WorldComponent = nullptr,
+		FName                                       WorldBoneName = NAME_None);
 
 	/**
 	 * Adds a Control to a Set. This will add a new set if necessary. For example, you might
diff --git a/Engine/Plugins/Experimental/PhysicsControl/Source/PhysicsControl/Public/PhysicsControlData.h b/Engine/Plugins/Experimental/PhysicsControl/Source/PhysicsControl/Public/PhysicsControlData.h
--- a/Engine/Plugins/Experimental/PhysicsControl/Source/PhysicsControl/Public/PhysicsControlData.h
+++ b/Engine/Plugins/Experimental/PhysicsControl/Source/PhysicsControl/Public/PhysicsControlData.h
@@ -529,6 +529,7 @@ struct PHYSICSCONTROL_API FPhysicsControlSettings
 		, bUseSkeletalAnimation(true)
 		, SkeletalAnimationVelocityMultiplier(1.0f)
 		, bDisableCollision(false)
+		, bOnlyControlChildObject(false)
 	{
 	}
 
@@ -553,6 +554,12 @@ struct PHYSICSCONTROL_API FPhysicsControlSettings
 	 */
 	UPROPERTY(EditAnywhere, BlueprintReadWrite, Category = PhysicsControl)
 	bool bDisableCollision;
+
+	/**
+	 * If set to true, then the control will only affect the child object, not the parent
+	 */
+	UPROPERTY(EditAnywhere, BlueprintReadWrite, Category = PhysicsControl)
+	bool bOnlyControlChildObject;
 };
 
 /**
diff --git a/Engine/Source/Runtime/Engine/Classes/PhysicsEngine/ConstraintInstance.h b/Engine/Source/Runtime/Engine/Classes/PhysicsEngine/ConstraintInstance.h
--- a/Engine/Source/Runtime/Engine/Classes/PhysicsEngine/ConstraintInstance.h
+++ b/Engine/Source/Runtime/Engine/Classes/PhysicsEngine/ConstraintInstance.h
@@ -970,6 +970,7 @@ public:
 	/** Enable/Disable parent dominates (meaning the parent body cannot be be affected at all by a child) */
 	ENGINE_API void EnableParentDominates();
 	ENGINE_API void DisableParentDominates();
+	ENGINE_API void SetParentDominates(bool bParentDominates);
 
 	/** Whether mass conditioning is enabled. @see FConstraintProfileProperties::bEnableMassConditioning */
 	bool IsMassConditioningEnabled() const
diff --git a/Engine/Source/Runtime/Engine/Private/PhysicsEngine/ConstraintInstance.cpp b/Engine/Source/Runtime/Engine/Private/PhysicsEngine/ConstraintInstance.cpp
--- a/Engine/Source/Runtime/Engine/Private/PhysicsEngine/ConstraintInstance.cpp
+++ b/Engine/Source/Runtime/Engine/Private/PhysicsEngine/ConstraintInstance.cpp
@@ -1343,22 +1343,22 @@ void FConstraintInstance::DisableProjection()
 
 void FConstraintInstance::EnableParentDominates()
 {
-	ProfileInstance.bParentDominates = true;
-	
-	FPhysicsCommand::ExecuteWrite(ConstraintHandle, [&](const FPhysicsConstraintHandle& Constraint)
-	{
-		FPhysicsInterface::SetParentDominates_AssumesLocked(Constraint, true);
-	});
+	SetParentDominates(true);
 }
 
 void FConstraintInstance::DisableParentDominates()
 {
-	ProfileInstance.bParentDominates = false;
-	
-	FPhysicsCommand::ExecuteWrite(ConstraintHandle, [&](const FPhysicsConstraintHandle& Constraint)
-	{
-		FPhysicsInterface::SetParentDominates_AssumesLocked(Constraint, false);
-	});
+	SetParentDominates(false);
+}
+
+void FConstraintInstance::SetParentDominates(bool bParentDominates)
+{
+	ProfileInstance.bParentDominates = bParentDominates;
+
+	FPhysicsCommand::ExecuteWrite(ConstraintHandle, [this, bParentDominates](const FPhysicsConstraintHandle& Constraint)
+		{
+			FPhysicsInterface::SetParentDominates_AssumesLocked(Constraint, bParentDominates);
+		});
 }
 
 void FConstraintInstance::EnableMassConditioning()
