Fix indexing into angular limits for the RigidBodyWithControl node when using constraint profiles

diff --git a/Engine/Plugins/Experimental/PhysicsControl/Source/PhysicsControl/Private/RigidBodyWithControl.cpp b/Engine/Plugins/Experimental/PhysicsControl/Source/PhysicsControl/Private/RigidBodyWithControl.cpp
--- a/Engine/Plugins/Experimental/PhysicsControl/Source/PhysicsControl/Private/RigidBodyWithControl.cpp
+++ b/Engine/Plugins/Experimental/PhysicsControl/Source/PhysicsControl/Private/RigidBodyWithControl.cpp
@@ -837,13 +837,13 @@ void FAnimNode_RigidBodyWithControl::ApplyCurrentConstraintProfile()
 					JointSettings.LinearLimit = Profile.LinearLimit.Limit; // Is float to vector the best way?!
 
 					// Order is twist, swing1, swing2 and in degrees
-					JointSettings.AngularMotionTypes[0] = ConvertMotionType(Profile.TwistLimit.TwistMotion);
-					JointSettings.AngularMotionTypes[1] = ConvertMotionType(Profile.ConeLimit.Swing1Motion);
-					JointSettings.AngularMotionTypes[2] = ConvertMotionType(Profile.ConeLimit.Swing2Motion);
+					JointSettings.AngularMotionTypes[(int32)Chaos::EJointAngularConstraintIndex::Twist] = ConvertMotionType(Profile.TwistLimit.TwistMotion);
+					JointSettings.AngularMotionTypes[(int32)Chaos::EJointAngularConstraintIndex::Swing1] = ConvertMotionType(Profile.ConeLimit.Swing1Motion);
+					JointSettings.AngularMotionTypes[(int32)Chaos::EJointAngularConstraintIndex::Swing2] = ConvertMotionType(Profile.ConeLimit.Swing2Motion);
 
-					JointSettings.AngularLimits[0] = FMath::DegreesToRadians(Profile.TwistLimit.TwistLimitDegrees);
-					JointSettings.AngularLimits[1] = FMath::DegreesToRadians(Profile.ConeLimit.Swing1LimitDegrees);
-					JointSettings.AngularLimits[2] = FMath::DegreesToRadians(Profile.ConeLimit.Swing2LimitDegrees);
+					JointSettings.AngularLimits[(int32)Chaos::EJointAngularConstraintIndex::Twist] = FMath::DegreesToRadians(Profile.TwistLimit.TwistLimitDegrees);
+					JointSettings.AngularLimits[(int32)Chaos::EJointAngularConstraintIndex::Swing1] = FMath::DegreesToRadians(Profile.ConeLimit.Swing1LimitDegrees);
+					JointSettings.AngularLimits[(int32)Chaos::EJointAngularConstraintIndex::Swing2] = FMath::DegreesToRadians(Profile.ConeLimit.Swing2LimitDegrees);
 
 					JointSettings.bSoftLinearLimitsEnabled = Profile.LinearLimit.bSoftConstraint;
 					JointSettings.bSoftTwistLimitsEnabled = Profile.TwistLimit.bSoftConstraint;
