[UE] RBwC modify physics constraint transforms to match the incomming Skeleton in init physics
- Added a new function to optionally update each constraint's parent space transform to compensate for any difference between the skeleton used to create the physics asset and the current skeleton
- The skeleton used to create the physics asset can be specified in the RBwC node's details panel
- The constraint's parent transforms can be snapped to the default positions relative to the parent bone if the original skeleton is undefined.
- Exposes the previously editor only function CalculateRelativeBoneTransform to the runtime.

diff --git a/Engine/Plugins/Experimental/PhysicsControl/Source/PhysicsControl/Private/AnimNode_RigidBodyWithControl.cpp b/Engine/Plugins/Experimental/PhysicsControl/Source/PhysicsControl/Private/AnimNode_RigidBodyWithControl.cpp
--- a/Engine/Plugins/Experimental/PhysicsControl/Source/PhysicsControl/Private/AnimNode_RigidBodyWithControl.cpp
+++ b/Engine/Plugins/Experimental/PhysicsControl/Source/PhysicsControl/Private/AnimNode_RigidBodyWithControl.cpp
@@ -1090,6 +1090,11 @@ void FAnimNode_RigidBodyWithControl::InitPhysics(const UAnimInstance* InAnimInst
 			FPhysicsAggregateHandle(),
 			bCreateBodiesInRefPose);
 
+		if (bModifyConstraintTransformsToMatchSkeleton)
+		{
+			TransformConstraintsToMatchSkeletalMesh(SkeletalMeshAsset, HighLevelConstraintInstances);
+		}
+
 		TMap<FName, ImmediatePhysics::FActorHandle*> NamesToHandles;
 		TArray<ImmediatePhysics::FActorHandle*> IgnoreCollisionActors;
 
diff --git a/Engine/Plugins/Experimental/PhysicsControl/Source/PhysicsControl/Private/RigidBodyWithControl.cpp b/Engine/Plugins/Experimental/PhysicsControl/Source/PhysicsControl/Private/RigidBodyWithControl.cpp
--- a/Engine/Plugins/Experimental/PhysicsControl/Source/PhysicsControl/Private/RigidBodyWithControl.cpp
+++ b/Engine/Plugins/Experimental/PhysicsControl/Source/PhysicsControl/Private/RigidBodyWithControl.cpp
@@ -14,6 +14,7 @@
 #include "PhysicsEngine/PhysicsSettings.h"
 #include "Chaos/ChaosConstraintSettings.h"
 #include "PhysicsControlLog.h"
+#include "Logging/StructuredLog.h"
 #include "PhysicsControlOperatorNameGeneration.h"
 
 //UE_DISABLE_OPTIMIZATION;
@@ -959,3 +960,70 @@ void FAnimNode_RigidBodyWithControl::ApplyCurrentConstraintProfile()
 	}
 }
 
+//======================================================================================================================
+void FAnimNode_RigidBodyWithControl::TransformConstraintsToMatchSkeletalMesh(const USkeletalMesh* const SkeletalMeshAsset, TArray<FConstraintInstance*>& ConstraintInstances)
+{
+	// Bone1 = Child
+	// Bone2 = Parent 
+
+	if (SkeletalMeshAsset && PhysicsAssetAuthoredSkeletalMesh)
+	{
+		if (SkeletalMeshAsset != PhysicsAssetAuthoredSkeletalMesh)
+		{
+			UE_LOGFMT(LogRigidBodyWithControl, Log, "Modify Constraint parent transforms to correct for the difference between the Skeleton used to create the Physics asset \"{0}\" and the current skeleton \"{1}\".", PhysicsAssetAuthoredSkeletalMesh->GetName(), SkeletalMeshAsset->GetName());
+
+			const FReferenceSkeleton& OriginalReferenceSkeleton = PhysicsAssetAuthoredSkeletalMesh->GetRefSkeleton();
+			const FReferenceSkeleton& CurrentReferenceSkeleton = SkeletalMeshAsset->GetRefSkeleton();
+
+			for (FConstraintInstance* const ConstraintInstance : ConstraintInstances)
+			{
+#if !NO_LOGGING
+				const FVector LogPreviousConstraintPositionRelParent = ConstraintInstance->Pos2;
+#endif
+				const FTransform CurrentParentRelChildTM = CalculateRelativeBoneTransform(ConstraintInstance->ConstraintBone1, ConstraintInstance->ConstraintBone2, CurrentReferenceSkeleton);
+				const FTransform OriginalParentRelChildTM = CalculateRelativeBoneTransform(ConstraintInstance->ConstraintBone1, ConstraintInstance->ConstraintBone2, OriginalReferenceSkeleton);
+
+				// Find the transform that maps the parent-bone-relative-to-the-child-bone transform in the original skeleton to the parent-bone-relative-to-the-child-bone transform in the current skeleton.
+				const FTransform OriginalToCurrentParentRelChildTM = CurrentParentRelChildTM.GetRelativeTransform(OriginalParentRelChildTM); // Should be equivalent to CurrentParentRelChildTM * OriginalParentRelChildTM.Inverse()
+
+				// Update the constraints transform relative to the parent bone.
+				const FTransform OriginalRefFrame = ConstraintInstance->GetRefFrame(EConstraintFrame::Frame2);
+				const FTransform CurrentRefFrame = OriginalToCurrentParentRelChildTM * OriginalRefFrame;
+				ConstraintInstance->SetRefFrame(EConstraintFrame::Frame2, CurrentRefFrame);
+
+#if !NO_LOGGING
+				UE_LOGFMT(LogRigidBodyWithControl, Log, "Matched Constraint {0} - {1} Parent Transform - position was {2} now {3}.", ConstraintInstance->ConstraintBone1.ToString(), ConstraintInstance->ConstraintBone2.ToString(), LogPreviousConstraintPositionRelParent.ToCompactString(), ConstraintInstance->Pos2.ToCompactString());
+#endif
+			}
+		}
+#if !NO_LOGGING
+		else
+		{
+			UE_LOGFMT(LogRigidBodyWithControl, Log, "Do not modify constraint parent transforms as the Skeleton used to create the Physics asset \"{0}\" matches the current skeleton \"{1}\".", PhysicsAssetAuthoredSkeletalMesh->GetName(), SkeletalMeshAsset->GetName());
+		}
+#endif
+	}
+	else if (SkeletalMeshAsset)
+	{
+		UE_LOGFMT(LogRigidBodyWithControl, Log, "Snap Constraint parent transforms to the current skeleton \"{0}\" (Authored Skeleton is undefined in node details).", SkeletalMeshAsset->GetName());
+
+		const FReferenceSkeleton& CurrentReferenceSkeleton = SkeletalMeshAsset->GetRefSkeleton();
+
+		// Move all the constraints to the default (snapped) location relative to the parent bone.
+		for (FConstraintInstance* const ConstraintInstance : ConstraintInstances)
+		{
+#if !NO_LOGGING
+			const FVector LogPreviousConstraintPositionRelParent = ConstraintInstance->Pos2;
+#endif
+
+			const FTransform ParentRelChildTM = CalculateRelativeBoneTransform(ConstraintInstance->ConstraintBone1, ConstraintInstance->ConstraintBone2, CurrentReferenceSkeleton);			
+			ConstraintInstance->SetRefFrame(EConstraintFrame::Frame2, ParentRelChildTM);
+
+#if !NO_LOGGING
+			UE_LOGFMT(LogRigidBodyWithControl, Log, "Snapped Constraint {0} - {1} Parent Transform - position was {2} now {3}.", ConstraintInstance->ConstraintBone1.ToString(), ConstraintInstance->ConstraintBone2.ToString(), LogPreviousConstraintPositionRelParent.ToCompactString(), ConstraintInstance->Pos2.ToCompactString());
+#endif
+		}
+	}
+}
+
+
diff --git a/Engine/Plugins/Experimental/PhysicsControl/Source/PhysicsControl/Public/AnimNode_RigidBodyWithControl.h b/Engine/Plugins/Experimental/PhysicsControl/Source/PhysicsControl/Public/AnimNode_RigidBodyWithControl.h
--- a/Engine/Plugins/Experimental/PhysicsControl/Source/PhysicsControl/Public/AnimNode_RigidBodyWithControl.h
+++ b/Engine/Plugins/Experimental/PhysicsControl/Source/PhysicsControl/Public/AnimNode_RigidBodyWithControl.h
@@ -184,6 +184,12 @@ public:
 	UPROPERTY(EditAnywhere, Category = Settings)
 	uint8 bClampLinearTranslationLimitToRefPose : 1;
 
+	/**
+		Change the parent space transforms of constraints read from the physics asset to match the relative bone transforms in the in-coming skeleton.
+	*/
+	UPROPERTY(EditAnywhere, Category = PhysicsAssetConditioning, meta = (InlineEditConditionToggle))
+	uint8 bModifyConstraintTransformsToMatchSkeleton : 1;
+
 	/**
 		For world-space simulations, if the magnitude of the component's 3D scale is less than WorldSpaceMinimumScale, do not update the node.
 	*/
@@ -266,6 +272,15 @@ public:
 	UPROPERTY(EditAnywhere, Category = Controls, meta = (PinShownByDefault))
 	FRigidBodyKinematicTargets KinematicTargets;
 
+	/**
+	 * When enabled, each Constraint's transform relative to its parent bone will be updated to account for any difference 
+	 * in the transform of its child bone relative to its parent bone between the skeleton used to author the physics asset
+	 * and the current skeleton (if the authored skeleton is defined) or snapped to the default transform relative to the
+	 * parent bone (if the authored skeleton is not defined).
+	 */
+	UPROPERTY(EditAnywhere, Category = PhysicsAssetConditioning, meta = (DisplayName = "Map constraints to Skeleton", editcondition = "bModifyConstraintTransformsToMatchSkeleton"))
+	USkeletalMesh* PhysicsAssetAuthoredSkeletalMesh;
+
 	/**
 	 * The constraint profile to use on all the joints in the physics asset
 	 */
@@ -375,6 +390,9 @@ private:
 		FVector& SpaceLinearAcc,
 		FVector& SpaceAngularAcc);
 
+	// Modify Constraint transforms relative to the parent bone to correct for the difference between the Skeleton used to create the Physics asset and the current skeleton.
+	void TransformConstraintsToMatchSkeletalMesh(const USkeletalMesh* const SkeletalMeshAsset, TArray<FConstraintInstance*>& ConstraintInstances);
+
 	// Gather cloth collision sources from the supplied Skeltal Mesh and add a kinematic actor representing each one of them to the sim.
 	void CollectClothColliderObjects(const USkeletalMeshComponent* SkeletalMeshComp);
 	
diff --git a/Engine/Source/Runtime/Engine/Classes/PhysicsEngine/ConstraintInstance.h b/Engine/Source/Runtime/Engine/Classes/PhysicsEngine/ConstraintInstance.h
--- a/Engine/Source/Runtime/Engine/Classes/PhysicsEngine/ConstraintInstance.h
+++ b/Engine/Source/Runtime/Engine/Classes/PhysicsEngine/ConstraintInstance.h
@@ -22,6 +22,7 @@ class FMaterialRenderProxy;
 class FPrimitiveDrawInterface;
 class FMaterialRenderProxy;
 class UPhysicsAsset;
+struct FReferenceSkeleton;
 
 UENUM()
 enum class EConstraintTransformComponentFlags : uint8
@@ -42,6 +43,9 @@ enum class EConstraintTransformComponentFlags : uint8
 
 ENUM_CLASS_FLAGS(EConstraintTransformComponentFlags);
 
+/** Returns the 'To' bone's transform relative to the 'From' bone. */
+ENGINE_API FTransform CalculateRelativeBoneTransform(const FName ToBoneName, const FName FromBoneName, const FReferenceSkeleton& ReferenceSkeleton);
+
 /** Container for properties of a physics constraint that can be easily swapped at runtime. This is useful for switching different setups when going from ragdoll to standup for example */
 USTRUCT()
 struct FConstraintProfileProperties
diff --git a/Engine/Source/Runtime/Engine/Private/PhysicsEngine/ConstraintInstance.cpp b/Engine/Source/Runtime/Engine/Private/PhysicsEngine/ConstraintInstance.cpp
--- a/Engine/Source/Runtime/Engine/Private/PhysicsEngine/ConstraintInstance.cpp
+++ b/Engine/Source/Runtime/Engine/Private/PhysicsEngine/ConstraintInstance.cpp
@@ -65,10 +65,8 @@ FVector RevolutionsToRads(const FVector Revolutions)
 	return Revolutions * 2.f * UE_PI;
 }
 
-#if WITH_EDITORONLY_DATA
-
 /** Returns the 'To' bone's transform relative to the 'From' bone. */
-FTransform CalculateRelativeBoneTransform(const FName ToBoneName, const FName FromBoneName, FReferenceSkeleton& ReferenceSkeleton)
+FTransform CalculateRelativeBoneTransform(const FName ToBoneName, const FName FromBoneName, const FReferenceSkeleton& ReferenceSkeleton)
 {
 	FTransform RelativeBoneTransform = FTransform::Identity;
 
@@ -113,8 +111,6 @@ FTransform CalculateRelativeBoneTransform(const FName ToBoneName, const FName Fr
 	return RelativeBoneTransform;
 }
 
-#endif // WITH_EDITORONLY_DATA
-
 #if WITH_EDITOR
 void FConstraintProfileProperties::SyncChangedConstraintProperties(FPropertyChangedChainEvent& PropertyChangedEvent)
 {
