[UE] Initialise member to fix asset reference check error
- CIS 389954

diff --git a/Engine/Plugins/Experimental/PhysicsControl/Source/PhysicsControl/Private/AnimNode_RigidBodyWithControl.cpp b/Engine/Plugins/Experimental/PhysicsControl/Source/PhysicsControl/Private/AnimNode_RigidBodyWithControl.cpp
--- a/Engine/Plugins/Experimental/PhysicsControl/Source/PhysicsControl/Private/AnimNode_RigidBodyWithControl.cpp
+++ b/Engine/Plugins/Experimental/PhysicsControl/Source/PhysicsControl/Private/AnimNode_RigidBodyWithControl.cpp
@@ -138,9 +138,11 @@ FAnimNode_RigidBodyWithControl::FAnimNode_RigidBodyWithControl()
 	, bTransferBoneVelocities(false)
 	, bFreezeIncomingPoseOnStart(false)
 	, bClampLinearTranslationLimitToRefPose(false)
+	, bModifyConstraintTransformsToMatchSkeleton(false)
 	, WorldSpaceMinimumScale(0.01f)
 	, EvaluationResetTime(0.01f)
 	, bEnableControls(false)
+	, PhysicsAssetAuthoredSkeletalMesh(nullptr)
 	, bEnabled(false)
 	, bSimulationStarted(false)
 	, bCheckForBodyTransformInit(false)
