Code tidy - no functional changes

diff --git a/Engine/Plugins/Experimental/PhysicsControl/Source/PhysicsControl/Private/RigidBodyWithControl.cpp b/Engine/Plugins/Experimental/PhysicsControl/Source/PhysicsControl/Private/RigidBodyWithControl.cpp
--- a/Engine/Plugins/Experimental/PhysicsControl/Source/PhysicsControl/Private/RigidBodyWithControl.cpp
+++ b/Engine/Plugins/Experimental/PhysicsControl/Source/PhysicsControl/Private/RigidBodyWithControl.cpp
@@ -961,7 +961,9 @@ void FAnimNode_RigidBodyWithControl::ApplyCurrentConstraintProfile()
 }
 
 //======================================================================================================================
-void FAnimNode_RigidBodyWithControl::TransformConstraintsToMatchSkeletalMesh(const USkeletalMesh* const SkeletalMeshAsset, TArray<FConstraintInstance*>& ConstraintInstances)
+void FAnimNode_RigidBodyWithControl::TransformConstraintsToMatchSkeletalMesh(
+	const USkeletalMesh*          const SkeletalMeshAsset, 
+	TArray<FConstraintInstance*>& ConstraintInstances)
 {
 	// Bone1 = Child
 	// Bone2 = Parent 
@@ -970,7 +972,9 @@ void FAnimNode_RigidBodyWithControl::TransformConstraintsToMatchSkeletalMesh(con
 	{
 		if (SkeletalMeshAsset != PhysicsAssetAuthoredSkeletalMesh)
 		{
-			UE_LOGFMT(LogRigidBodyWithControl, Log, "Modify Constraint parent transforms to correct for the difference between the Skeleton used to create the Physics asset \"{0}\" and the current skeleton \"{1}\".", PhysicsAssetAuthoredSkeletalMesh->GetName(), SkeletalMeshAsset->GetName());
+			UE_LOGFMT(LogRigidBodyWithControl, Log, 
+				"Modify Constraint parent transforms to correct for the difference between the Skeleton used to create the Physics asset \"{0}\" and the current skeleton \"{1}\".", 
+				PhysicsAssetAuthoredSkeletalMesh->GetName(), SkeletalMeshAsset->GetName());
 
 			const FReferenceSkeleton& OriginalReferenceSkeleton = PhysicsAssetAuthoredSkeletalMesh->GetRefSkeleton();
 			const FReferenceSkeleton& CurrentReferenceSkeleton = SkeletalMeshAsset->GetRefSkeleton();
@@ -980,11 +984,17 @@ void FAnimNode_RigidBodyWithControl::TransformConstraintsToMatchSkeletalMesh(con
 #if !NO_LOGGING
 				const FVector LogPreviousConstraintPositionRelParent = ConstraintInstance->Pos2;
 #endif
-				const FTransform CurrentParentRelChildTM = CalculateRelativeBoneTransform(ConstraintInstance->ConstraintBone1, ConstraintInstance->ConstraintBone2, CurrentReferenceSkeleton);
-				const FTransform OriginalParentRelChildTM = CalculateRelativeBoneTransform(ConstraintInstance->ConstraintBone1, ConstraintInstance->ConstraintBone2, OriginalReferenceSkeleton);
-
-				// Find the transform that maps the parent-bone-relative-to-the-child-bone transform in the original skeleton to the parent-bone-relative-to-the-child-bone transform in the current skeleton.
-				const FTransform OriginalToCurrentParentRelChildTM = CurrentParentRelChildTM.GetRelativeTransform(OriginalParentRelChildTM); // Should be equivalent to CurrentParentRelChildTM * OriginalParentRelChildTM.Inverse()
+				const FTransform CurrentParentRelChildTM = CalculateRelativeBoneTransform(
+					ConstraintInstance->ConstraintBone1, ConstraintInstance->ConstraintBone2, CurrentReferenceSkeleton);
+				const FTransform OriginalParentRelChildTM = CalculateRelativeBoneTransform(
+					ConstraintInstance->ConstraintBone1, ConstraintInstance->ConstraintBone2, OriginalReferenceSkeleton);
+
+				// Find the transform that maps the parent-bone-relative-to-the-child-bone transform
+				// in the original skeleton to the parent-bone-relative-to-the-child-bone transform
+				// in the current skeleton.
+				// Should be equivalent to CurrentParentRelChildTM * OriginalParentRelChildTM.Inverse()
+				const FTransform OriginalToCurrentParentRelChildTM = 
+					CurrentParentRelChildTM.GetRelativeTransform(OriginalParentRelChildTM); 
 
 				// Update the constraints transform relative to the parent bone.
 				const FTransform OriginalRefFrame = ConstraintInstance->GetRefFrame(EConstraintFrame::Frame2);
@@ -992,20 +1002,29 @@ void FAnimNode_RigidBodyWithControl::TransformConstraintsToMatchSkeletalMesh(con
 				ConstraintInstance->SetRefFrame(EConstraintFrame::Frame2, CurrentRefFrame);
 
 #if !NO_LOGGING
-				UE_LOGFMT(LogRigidBodyWithControl, Log, "Matched Constraint {0} - {1} Parent Transform - position was {2} now {3}.", ConstraintInstance->ConstraintBone1.ToString(), ConstraintInstance->ConstraintBone2.ToString(), LogPreviousConstraintPositionRelParent.ToCompactString(), ConstraintInstance->Pos2.ToCompactString());
+				UE_LOGFMT(LogRigidBodyWithControl, Log, 
+					"Matched Constraint {0} - {1} Parent Transform - position was {2} now {3}.", 
+					ConstraintInstance->ConstraintBone1.ToString(), 
+					ConstraintInstance->ConstraintBone2.ToString(), 
+					LogPreviousConstraintPositionRelParent.ToCompactString(), 
+					ConstraintInstance->Pos2.ToCompactString());
 #endif
 			}
 		}
 #if !NO_LOGGING
 		else
 		{
-			UE_LOGFMT(LogRigidBodyWithControl, Log, "Do not modify constraint parent transforms as the Skeleton used to create the Physics asset \"{0}\" matches the current skeleton \"{1}\".", PhysicsAssetAuthoredSkeletalMesh->GetName(), SkeletalMeshAsset->GetName());
+			UE_LOGFMT(LogRigidBodyWithControl, Log, 
+				"Do not modify constraint parent transforms as the Skeleton used to create the Physics asset \"{0}\" matches the current skeleton \"{1}\".", 
+				PhysicsAssetAuthoredSkeletalMesh->GetName(), SkeletalMeshAsset->GetName());
 		}
 #endif
 	}
 	else if (SkeletalMeshAsset)
 	{
-		UE_LOGFMT(LogRigidBodyWithControl, Log, "Snap Constraint parent transforms to the current skeleton \"{0}\" (Authored Skeleton is undefined in node details).", SkeletalMeshAsset->GetName());
+		UE_LOGFMT(LogRigidBodyWithControl, Log, 
+			"Snap Constraint parent transforms to the current skeleton \"{0}\" (Authored Skeleton is undefined in node details).", 
+			SkeletalMeshAsset->GetName());
 
 		const FReferenceSkeleton& CurrentReferenceSkeleton = SkeletalMeshAsset->GetRefSkeleton();
 
@@ -1016,11 +1035,17 @@ void FAnimNode_RigidBodyWithControl::TransformConstraintsToMatchSkeletalMesh(con
 			const FVector LogPreviousConstraintPositionRelParent = ConstraintInstance->Pos2;
 #endif
 
-			const FTransform ParentRelChildTM = CalculateRelativeBoneTransform(ConstraintInstance->ConstraintBone1, ConstraintInstance->ConstraintBone2, CurrentReferenceSkeleton);			
+			const FTransform ParentRelChildTM = CalculateRelativeBoneTransform(
+				ConstraintInstance->ConstraintBone1, ConstraintInstance->ConstraintBone2, CurrentReferenceSkeleton);			
 			ConstraintInstance->SetRefFrame(EConstraintFrame::Frame2, ParentRelChildTM);
 
 #if !NO_LOGGING
-			UE_LOGFMT(LogRigidBodyWithControl, Log, "Snapped Constraint {0} - {1} Parent Transform - position was {2} now {3}.", ConstraintInstance->ConstraintBone1.ToString(), ConstraintInstance->ConstraintBone2.ToString(), LogPreviousConstraintPositionRelParent.ToCompactString(), ConstraintInstance->Pos2.ToCompactString());
+			UE_LOGFMT(LogRigidBodyWithControl, Log, 
+				"Snapped Constraint {0} - {1} Parent Transform - position was {2} now {3}.", 
+				ConstraintInstance->ConstraintBone1.ToString(), 
+				ConstraintInstance->ConstraintBone2.ToString(), 
+				LogPreviousConstraintPositionRelParent.ToCompactString(), 
+				ConstraintInstance->Pos2.ToCompactString());
 #endif
 		}
 	}
diff --git a/Engine/Plugins/Experimental/PhysicsControl/Source/PhysicsControl/Public/AnimNode_RigidBodyWithControl.h b/Engine/Plugins/Experimental/PhysicsControl/Source/PhysicsControl/Public/AnimNode_RigidBodyWithControl.h
--- a/Engine/Plugins/Experimental/PhysicsControl/Source/PhysicsControl/Public/AnimNode_RigidBodyWithControl.h
+++ b/Engine/Plugins/Experimental/PhysicsControl/Source/PhysicsControl/Public/AnimNode_RigidBodyWithControl.h
@@ -185,20 +185,23 @@ public:
 	uint8 bClampLinearTranslationLimitToRefPose : 1;
 
 	/**
-		Change the parent space transforms of constraints read from the physics asset to match the relative bone transforms in the in-coming skeleton.
+		Change the parent space transforms of constraints read from the physics asset to match the relative 
+		bone transforms in the in-coming skeleton.
 	*/
 	UPROPERTY(EditAnywhere, Category = PhysicsAssetConditioning, meta = (InlineEditConditionToggle))
 	uint8 bModifyConstraintTransformsToMatchSkeleton : 1;
 
 	/**
-		For world-space simulations, if the magnitude of the component's 3D scale is less than WorldSpaceMinimumScale, do not update the node.
+		For world-space simulations, if the magnitude of the component's 3D scale is less than 
+		WorldSpaceMinimumScale, do not update the node.
 	*/
 	UPROPERTY(EditAnywhere, Category = Settings)
 	float WorldSpaceMinimumScale;
 
 	/**
-		If the node is not evaluated for this amount of time (seconds), either because a lower LOD was in use for a while or the component was
-		not visible, reset the simulation to the default pose on the next evaluation. Set to 0 to disable time-based reset.
+		If the node is not evaluated for this amount of time (seconds), either because a lower LOD was in use 
+		for a while or the component was not visible, reset the simulation to the default pose on the next 
+		evaluation. Set to 0 to disable time-based reset.
 	*/
 	UPROPERTY(EditAnywhere, Category = Settings)
 	float EvaluationResetTime;
@@ -273,10 +276,14 @@ public:
 	FRigidBodyKinematicTargets KinematicTargets;
 
 	/**
-	 * When enabled, each Constraint's transform relative to its parent bone will be updated to account for any difference 
-	 * in the transform of its child bone relative to its parent bone between the skeleton used to author the physics asset
-	 * and the current skeleton (if the authored skeleton is defined) or snapped to the default transform relative to the
-	 * parent bone (if the authored skeleton is not defined).
+	 * If this option is enabled, each Constraint's parent transform will be updated to adjust its position and 
+	 * orientation to account for the difference in the constraints child transform between the skeleton used 
+	 * to author the physics asset and the current skeleton (if the authored skeleton is defined). If the authored
+	 * skeleton is not defined/unavailable then the parent bone's transform is set to the default transform that
+	 * would have been used in the physics asset.
+	 * This can be used to created a simulated character that has somewhat different bone length/orientations compared
+	 * to the one used to create the physics asset, in order to avoid the need to customize the physics asset
+	 * for every skeleton you use.
 	 */
 	UPROPERTY(EditAnywhere, Category = PhysicsAssetConditioning, meta = (DisplayName = "Map constraints to Skeleton", editcondition = "bModifyConstraintTransformsToMatchSkeleton"))
 	USkeletalMesh* PhysicsAssetAuthoredSkeletalMesh;
@@ -390,10 +397,13 @@ private:
 		FVector& SpaceLinearAcc,
 		FVector& SpaceAngularAcc);
 
-	// Modify Constraint transforms relative to the parent bone to correct for the difference between the Skeleton used to create the Physics asset and the current skeleton.
-	void TransformConstraintsToMatchSkeletalMesh(const USkeletalMesh* const SkeletalMeshAsset, TArray<FConstraintInstance*>& ConstraintInstances);
+	// Modify Constraint transforms relative to the parent bone to correct for the difference
+	// between the Skeleton used to create the Physics asset and the current skeleton.
+	void TransformConstraintsToMatchSkeletalMesh(
+		const USkeletalMesh* const SkeletalMeshAsset, TArray<FConstraintInstance*>& ConstraintInstances);
 
-	// Gather cloth collision sources from the supplied Skeltal Mesh and add a kinematic actor representing each one of them to the sim.
+	// Gather cloth collision sources from the supplied Skeltal Mesh and add a kinematic actor
+	// representing each one of them to the sim.
 	void CollectClothColliderObjects(const USkeletalMeshComponent* SkeletalMeshComp);
 	
 	// Remove all cloth collider objects from the sim.
@@ -425,7 +435,10 @@ private:
 
 public:
 
-	/* Whether the physics simulation runs synchronously with the node's evaluation or is run in the background until the next frame. */
+	/* 
+	 * Whether the physics simulation runs synchronously with the node's evaluation or is run in the 
+	 * background until the next frame. 
+	 */
 	UPROPERTY(EditAnywhere, Category=Settings, AdvancedDisplay)
 	ESimulationTiming SimulationTiming;
 
@@ -481,7 +494,8 @@ private:
 	struct FWorldObject
 	{
 		FWorldObject() : ActorHandle(nullptr), LastSeenTick(0), bExpired(false) {}
-		FWorldObject(ImmediatePhysics::FActorHandle* InActorHandle, int32 InLastSeenTick) : ActorHandle(InActorHandle), LastSeenTick(InLastSeenTick), bExpired(false) {}
+		FWorldObject(ImmediatePhysics::FActorHandle* InActorHandle, int32 InLastSeenTick) 
+			: ActorHandle(InActorHandle), LastSeenTick(InLastSeenTick), bExpired(false) {}
 
 		ImmediatePhysics::FActorHandle* ActorHandle;
 		int32 LastSeenTick;
@@ -520,7 +534,10 @@ private:
 	// Information required to identify and update a kinematic object representing a cloth collision source in the sim.
 	struct FClothCollider
 	{
-		FClothCollider(ImmediatePhysics::FActorHandle* const InActorHandle, const USkeletalMeshComponent* const InSkeletalMeshComponent, const uint32 InBoneIndex)
+		FClothCollider(
+			ImmediatePhysics::FActorHandle* const InActorHandle, 
+			const USkeletalMeshComponent* const   InSkeletalMeshComponent, 
+			const uint32                          InBoneIndex)
 			: ActorHandle(InActorHandle)
 			, SkeletalMeshComponent(InSkeletalMeshComponent)
 			, BoneIndex(InBoneIndex)
@@ -549,13 +566,12 @@ private:
 
 	FPhysScene* PhysScene;
 
-	// Used by CollectWorldObjects and UpdateWorldGeometry in Task Thread
-	// Typically, World should never be accessed off the Game Thread.
-	// However, since we're just doing overlaps this should be OK.
+	// Used by CollectWorldObjects and UpdateWorldGeometry in the Task Thread. Typically, World
+	// should never be accessed off the Game Thread. However, since we're just doing overlaps this
+	// should be OK.
 	const UWorld* UnsafeWorld;
 
-	// Used by CollectWorldObjects and UpdateWorldGeometry in Task Thread
-	// Only used for a pointer comparison.
+	// Used by CollectWorldObjects and UpdateWorldGeometry in Task Thread. Only used for a pointer comparison.
 	const AActor* UnsafeOwner;
 
 	FBoneContainer CapturedBoneVelocityBoneContainer;
