Support transitions between states in RBwC

#rnx

diff --git a/Engine/Plugins/Experimental/PhysicsControl/Source/PhysicsControl/Public/RigidBodyControlData.h b/Engine/Plugins/Experimental/PhysicsControl/Source/PhysicsControl/Public/RigidBodyControlData.h
--- a/Engine/Plugins/Experimental/PhysicsControl/Source/PhysicsControl/Public/RigidBodyControlData.h
+++ b/Engine/Plugins/Experimental/PhysicsControl/Source/PhysicsControl/Public/RigidBodyControlData.h
@@ -120,9 +120,9 @@ struct PHYSICSCONTROL_API FRigidBodyControlAndBodyModifierCreations
 	TArray<FRigidBodyModifierCreation> Modifiers;
 };
 
-FPhysicsControlData Interpolate(
+PHYSICSCONTROL_API FPhysicsControlData Interpolate(
 	const FPhysicsControlData& A, const FPhysicsControlData& B, const float Weight);
-FPhysicsControlSparseData Interpolate(
+PHYSICSCONTROL_API FPhysicsControlSparseData Interpolate(
 	const FPhysicsControlSparseData& A, const FPhysicsControlSparseData& B, const float Weight);
 
 /**
@@ -146,9 +146,9 @@ struct PHYSICSCONTROL_API FPhysicsControlNamedControlParameters
 	FPhysicsControlSparseData Data;
 };
 
-FPhysicsControlModifierData Interpolate(
+PHYSICSCONTROL_API FPhysicsControlModifierData Interpolate(
 	const FPhysicsControlModifierData& A, const FPhysicsControlModifierData& B, const float Weight);
-FPhysicsControlModifierSparseData Interpolate(
+PHYSICSCONTROL_API FPhysicsControlModifierSparseData Interpolate(
 	const FPhysicsControlModifierSparseData& A, const FPhysicsControlModifierSparseData& B, const float Weight);
 
 /**
