[UE] Add CVar to enable/disable matching physics asset constraints to the current skeleton at runtime.
- CVar is called "p.RigidBodyNodeWithControl.EnableMatchingConstraintsToSkeleton"

#rnx

diff --git a/Engine/Plugins/Experimental/PhysicsControl/Source/PhysicsControl/Private/AnimNode_RigidBodyWithControl.cpp b/Engine/Plugins/Experimental/PhysicsControl/Source/PhysicsControl/Private/AnimNode_RigidBodyWithControl.cpp
--- a/Engine/Plugins/Experimental/PhysicsControl/Source/PhysicsControl/Private/AnimNode_RigidBodyWithControl.cpp
+++ b/Engine/Plugins/Experimental/PhysicsControl/Source/PhysicsControl/Private/AnimNode_RigidBodyWithControl.cpp
@@ -51,6 +51,7 @@ CSV_DECLARE_CATEGORY_MODULE_EXTERN(ENGINE_API, Animation);
 
 TAutoConsoleVariable<int32> CVarEnableRigidBodyNodeWithControl(TEXT("p.RigidBodyNodeWithControl"), 1, TEXT("Enables/disables the whole rigid body node system. When disabled, avoids all allocations and runtime costs. Can be used to disable RB Nodes on low-end platforms."), ECVF_Scalability);
 TAutoConsoleVariable<int32> CVarEnableRigidBodyNodeWithControlSimulation(TEXT("p.RigidBodyNodeWithControl.EnableSimulation"), 1, TEXT("Runtime Enable/Disable RB Node Simulation for debugging and testing (node is initialized and bodies and constraints are created, even when disabled.)"), ECVF_Default);
+TAutoConsoleVariable<int32> CVarEnableRigidBodyNodeWithControlMatchingConstraintsToSkeleton(TEXT("p.RigidBodyNodeWithControl.EnableMatchingConstraintsToSkeleton"), 1, TEXT("Enables/disables the code that modifies physics asset constraint transforms to match the Skeleton at runtime."), ECVF_Scalability);
 TAutoConsoleVariable<int32> CVarRigidBodyNodeWithControlLODThreshold(TEXT("p.RigidBodyWithControlLODThreshold"), -1, TEXT("Max LOD that rigid body node is allowed to run on. Provides a global threshold that overrides per-node the LODThreshold property. -1 means no override."), ECVF_Scalability);
 
 int32 RBANWithControl_MaxSubSteps = 4;
@@ -1092,7 +1093,7 @@ void FAnimNode_RigidBodyWithControl::InitPhysics(const UAnimInstance* InAnimInst
 			FPhysicsAggregateHandle(),
 			bCreateBodiesInRefPose);
 
-		if (bModifyConstraintTransformsToMatchSkeleton)
+		if (bModifyConstraintTransformsToMatchSkeleton && (CVarEnableRigidBodyNodeWithControlMatchingConstraintsToSkeleton.GetValueOnAnyThread() > 0))
 		{
 			TransformConstraintsToMatchSkeletalMesh(SkeletalMeshAsset, HighLevelConstraintInstances);
 		}
