Small speedups for the RigidBodyWithControl node

#tests locally using PIE

diff --git a/Engine/Plugins/Experimental/PhysicsControl/Source/PhysicsControl/Private/AnimNode_RigidBodyWithControl.cpp b/Engine/Plugins/Experimental/PhysicsControl/Source/PhysicsControl/Private/AnimNode_RigidBodyWithControl.cpp
--- a/Engine/Plugins/Experimental/PhysicsControl/Source/PhysicsControl/Private/AnimNode_RigidBodyWithControl.cpp
+++ b/Engine/Plugins/Experimental/PhysicsControl/Source/PhysicsControl/Private/AnimNode_RigidBodyWithControl.cpp
@@ -45,6 +45,7 @@ DECLARE_CYCLE_STAT(TEXT("RigidBodyNodeWithControl_Simulation"), STAT_RigidBodyNo
 DECLARE_CYCLE_STAT(TEXT("RigidBodyNodeWithControl_SimulationWait"), STAT_RigidBodyNodeWithControl_SimulationWait, STATGROUP_Anim);
 DECLARE_CYCLE_STAT(TEXT("RigidBodyNodeWithControl_PreUpdate"), STAT_RigidBodyNodeWithControl_PreUpdate, STATGROUP_Anim);
 DECLARE_CYCLE_STAT(TEXT("RigidBodyNodeWithControl_Update"), STAT_RigidBodyNodeWithControl_Update, STATGROUP_Anim);
+DECLARE_CYCLE_STAT(TEXT("RigidBodyNodeWithControl_PoseUpdate"), STAT_RigidBodyNodeWithControl_PoseUpdate, STATGROUP_Anim);
 DECLARE_CYCLE_STAT(TEXT("FAnimNode_RigidBodyWithControl::EvaluateSkeletalControl_AnyThread"), STAT_ImmediateEvaluateSkeletalControl, STATGROUP_ImmediatePhysics);
 
 CSV_DECLARE_CATEGORY_MODULE_EXTERN(ENGINE_API, Animation);
@@ -138,7 +139,6 @@ FAnimNode_RigidBodyWithControl::FAnimNode_RigidBodyWithControl()
 	, bOverrideWorldGravity(false)
 	, bTransferBoneVelocities(false)
 	, bFreezeIncomingPoseOnStart(false)
-	, bClampLinearTranslationLimitToRefPose(false)
 	, bModifyConstraintTransformsToMatchSkeleton(false)
 	, WorldSpaceMinimumScale(0.01f)
 	, EvaluationResetTime(0.01f)
@@ -308,7 +308,6 @@ void FAnimNode_RigidBodyWithControl::InitializeNewBodyTransformsDuringSimulation
 
 				const FTransform WSBodyTM = BodyRelativeTransform * Bodies[OutputData.ParentBodyIndex]->GetWorldTransform();
 				Bodies[BodyIndex]->InitWorldTransform(WSBodyTM);
-				BodyAnimData[BodyIndex].RefPoseLength = (float) BodyRelativeTransform.GetLocation().Size();
 			}
 			// If we don't have a parent body, then we can just grab the incoming pose in component space.
 			else
@@ -626,14 +625,14 @@ void FAnimNode_RigidBodyWithControl::EvaluateSkeletalControl_AnyThread(FComponen
 								BodyData.TransferedBoneLinearVelocity = ((NextSSTM.GetLocation() - PrevSSTM.GetLocation()) / DeltaSeconds);
 								
 								// Angular Velocity
-								const FQuat DeltaRotation = (NextSSTM.GetRotation().Inverse() * PrevSSTM.GetRotation());
-								const float RotationAngle = (float) DeltaRotation.GetAngle() / DeltaSeconds;
-								BodyData.TransferedBoneAngularVelocity = (FQuat(DeltaRotation.GetRotationAxis(), RotationAngle)); 
+								FQuat DeltaRotation = (NextSSTM.GetRotation().Inverse() * PrevSSTM.GetRotation());
+								DeltaRotation.EnforceShortestArcWith(FQuat::Identity);
+								BodyData.TransferedBoneAngularVelocity = DeltaRotation.ToRotationVector() / DeltaSeconds;
 							}
 							else
 							{
-								BodyData.TransferedBoneLinearVelocity = (FVector::ZeroVector);
-								BodyData.TransferedBoneAngularVelocity = (FQuat::Identity); 
+								BodyData.TransferedBoneLinearVelocity = FVector::ZeroVector;
+								BodyData.TransferedBoneAngularVelocity = FVector::ZeroVector;
 							}
 
 						}
@@ -667,11 +666,6 @@ void FAnimNode_RigidBodyWithControl::EvaluateSkeletalControl_AnyThread(FComponen
 
 						BodyTM = ConvertCSTransformToSimSpace(SimulationSpace, ComponentSpaceTM, CompWorldSpaceTM, BaseBoneTM);
 						Bodies[BodyIndex]->SetWorldTransform(BodyTM);
-						if (OutputData.ParentBodyIndex != INDEX_NONE)
-						{
-							BodyAnimData[BodyIndex].RefPoseLength = (float) 
-								BodyTM.GetRelativeTransform(Bodies[OutputData.ParentBodyIndex]->GetWorldTransform()).GetLocation().Size();
-						}
 					}
 				}
 				break;
@@ -691,11 +685,6 @@ void FAnimNode_RigidBodyWithControl::EvaluateSkeletalControl_AnyThread(FComponen
 						const FTransform& ComponentSpaceTM = Output.Pose.GetComponentSpaceTransform(OutputData.CompactPoseBoneIndex);
 						const FTransform BodyTM = ConvertCSTransformToSimSpace(SimulationSpace, ComponentSpaceTM, CompWorldSpaceTM, BaseBoneTM);
 						Bodies[BodyIndex]->InitWorldTransform(BodyTM);
-						if (OutputData.ParentBodyIndex != INDEX_NONE)
-						{
-							BodyAnimData[BodyIndex].RefPoseLength = (float)
-								BodyTM.GetRelativeTransform(Bodies[OutputData.ParentBodyIndex]->GetWorldTransform()).GetLocation().Size();
-						}
 					}
 				}
 				break;
@@ -719,8 +708,11 @@ void FAnimNode_RigidBodyWithControl::EvaluateSkeletalControl_AnyThread(FComponen
 		if (bNeedsSimulationTick)
 		{
 			// Update the pose data
-			PoseData.Update(
-				Output, OutputBoneData, SimulationSpace, BaseBoneRef, Output.AnimInstanceProxy->GetUpdateCounter());
+			{
+				SCOPE_CYCLE_COUNTER(STAT_RigidBodyNodeWithControl_PoseUpdate);
+				PoseData.Update(
+					Output, OutputBoneData, SimulationSpace, BaseBoneRef, Output.AnimInstanceProxy->GetUpdateCounter());
+			}
 
 			// Transfer bone velocities previously captured.
 			if (bTransferBoneVelocities && (CapturedBoneVelocityPose.GetPose().GetNumBones() > 0))
@@ -734,9 +726,7 @@ void FAnimNode_RigidBodyWithControl::EvaluateSkeletalControl_AnyThread(FComponen
 					{
 						ImmediatePhysics::FActorHandle* Body = Bodies[BodyIndex];
 						Body->SetLinearVelocity(BodyData.TransferedBoneLinearVelocity);
-
-						const FQuat AngularVelocity = BodyData.TransferedBoneAngularVelocity;
-						Body->SetAngularVelocity(AngularVelocity.GetRotationAxis() * AngularVelocity.GetAngle());
+						Body->SetAngularVelocity(BodyData.TransferedBoneAngularVelocity);
 					}
 				}
 
@@ -813,8 +803,7 @@ void FAnimNode_RigidBodyWithControl::EvaluateSkeletalControl_AnyThread(FComponen
 				if (!BodyAnimData[BodyIndex].bIsSimulated)
 				{
 					// TODO support explicit targets for kinematics, like we do for controls
-					FTransform TM = PoseData.GetTM(BodyIndex);
-					Bodies[BodyIndex]->SetKinematicTarget(TM);
+					Bodies[BodyIndex]->SetKinematicTarget(PoseData.GetTM(BodyIndex).ToTransform());
 				}
 			}
 
@@ -888,51 +877,7 @@ void FAnimNode_RigidBodyWithControl::EvaluateSkeletalControl_AnyThread(FComponen
 			const int32 BodyIndex = OutputData.BodyIndex;
 			// Note that we always read back, whether kinematic or simulated
 			FTransform BodyTM = Bodies[BodyIndex]->GetWorldTransform();
-
-			// if we clamp translation, we only do this when all linear translation are locked
-			// 
-			// @todo(ccaulfield): this shouldn't be required with Chaos - projection should be handling it...
-			if (bClampLinearTranslationLimitToRefPose
-				&&BodyAnimData[BodyIndex].LinearXMotion == ELinearConstraintMotion::LCM_Locked
-				&& BodyAnimData[BodyIndex].LinearYMotion == ELinearConstraintMotion::LCM_Locked
-				&& BodyAnimData[BodyIndex].LinearZMotion == ELinearConstraintMotion::LCM_Locked)
-			{
-				// grab local space of length from ref pose 
-				// we have linear limit value - see if that works
-				// calculate current local space from parent
-				// find parent transform
-				const int32 ParentBodyIndex = OutputData.ParentBodyIndex;
-				FTransform ParentTransform = FTransform::Identity;
-				if (ParentBodyIndex != INDEX_NONE)
-				{
-					ParentTransform = Bodies[ParentBodyIndex]->GetWorldTransform();
-				}
-
-				// get local transform
-				FTransform LocalTransform = BodyTM.GetRelativeTransform(ParentTransform);
-				const float CurrentLength = (float) LocalTransform.GetTranslation().Size();
-
-				// this is inconsistent with constraint. The actual linear limit is set by constraint
-				if (!FMath::IsNearlyEqual(CurrentLength, BodyAnimData[BodyIndex].RefPoseLength, KINDA_SMALL_NUMBER))
-				{
-					float RefPoseLength = BodyAnimData[BodyIndex].RefPoseLength;
-					if (CurrentLength > RefPoseLength)
-					{
-						float Scale = (CurrentLength > KINDA_SMALL_NUMBER) ? RefPoseLength / CurrentLength : 0.f;
-						// we don't use 1.f here because 1.f can create pops based on float issue. 
-						// so we only activate clamping when less than 90%
-						if (Scale < 0.9f)
-						{
-							LocalTransform.ScaleTranslation(Scale);
-							BodyTM = LocalTransform * ParentTransform;
-							Bodies[BodyIndex]->SetWorldTransform(BodyTM);
-						}
-					}
-				}
-			}
-
 			FTransform ComponentSpaceTM;
-
 			switch(SimulationSpace)
 			{
 				case ESimulationSpace::ComponentSpace: ComponentSpaceTM = BodyTM; break;
@@ -1235,17 +1180,6 @@ void FAnimNode_RigidBodyWithControl::InitPhysics(const UAnimInstance* InAnimInst
 						if (Bodies.Find(Body1Handle, BodyIndex))
 						{
 							Joints[BodyIndex] = NewJointHandle;
-
-							BodyAnimData[BodyIndex].LinearXMotion = CI->GetLinearXMotion();
-							BodyAnimData[BodyIndex].LinearYMotion = CI->GetLinearYMotion();
-							BodyAnimData[BodyIndex].LinearZMotion = CI->GetLinearZMotion();
-							BodyAnimData[BodyIndex].LinearLimit = CI->GetLinearLimit();
-
-							//set limit to ref pose 
-							FTransform Body1Transform = Body1Handle->GetWorldTransform();
-							FTransform Body2Transform = Body2Handle->GetWorldTransform();
-							BodyAnimData[BodyIndex].RefPoseLength = (float) 
-								Body1Transform.GetRelativeTransform(Body2Transform).GetLocation().Size();
 						}
 
 						if (CI->IsCollisionDisabled())
diff --git a/Engine/Plugins/Experimental/PhysicsControl/Source/PhysicsControl/Private/RigidBodyControlData.cpp b/Engine/Plugins/Experimental/PhysicsControl/Source/PhysicsControl/Private/RigidBodyControlData.cpp
--- a/Engine/Plugins/Experimental/PhysicsControl/Source/PhysicsControl/Private/RigidBodyControlData.cpp
+++ b/Engine/Plugins/Experimental/PhysicsControl/Source/PhysicsControl/Private/RigidBodyControlData.cpp
@@ -81,6 +81,8 @@ FRigidBodyControlRecord::FRigidBodyControlRecord(const FRigidBodyControl& InCont
 	: Control(InControl)
 	, JointHandle(InJointHandle)
 	, CurrentData(Control.ControlData)
+	, ChildBodyIndex(-1)
+	, ParentBodyIndex(-1)
 {
 }
 
diff --git a/Engine/Plugins/Experimental/PhysicsControl/Source/PhysicsControl/Private/RigidBodyPoseData.cpp b/Engine/Plugins/Experimental/PhysicsControl/Source/PhysicsControl/Private/RigidBodyPoseData.cpp
--- a/Engine/Plugins/Experimental/PhysicsControl/Source/PhysicsControl/Private/RigidBodyPoseData.cpp
+++ b/Engine/Plugins/Experimental/PhysicsControl/Source/PhysicsControl/Private/RigidBodyPoseData.cpp
@@ -9,11 +9,11 @@ namespace RigidBodyWithControl
 
 //======================================================================================================================
 void FRigidBodyPoseData::Update(
-	FComponentSpacePoseContext&   ComponentSpacePoseContext,
-	const TArray<FOutputBoneData> OutputBoneData,
-	const ESimulationSpace        SimulationSpace,
-	const FBoneReference&         BaseBoneRef,
-	const FGraphTraversalCounter& InUpdateCounter)
+	FComponentSpacePoseContext&    ComponentSpacePoseContext,
+	const TArray<FOutputBoneData>& OutputBoneData,
+	const ESimulationSpace         SimulationSpace,
+	const FBoneReference&          BaseBoneRef,
+	const FGraphTraversalCounter&  InUpdateCounter)
 {
 	ExpectedUpdateCounter = UpdateCounter;
 	ExpectedUpdateCounter.Increment();
@@ -24,25 +24,25 @@ void FRigidBodyPoseData::Update(
 	const FTransform BaseBoneTM = ComponentSpacePoseContext.Pose.GetComponentSpaceTransform(
 		BaseBoneRef.GetCompactPoseIndex(BoneContainer));
 
-	if (BoneData.Num() == OutputBoneData.Num())
+	if (BoneTMs.Num() == OutputBoneData.Num())
 	{
 		for (const FOutputBoneData& OutputData : OutputBoneData)
 		{
 			const int32 BodyIndex = OutputData.BodyIndex;
 			const FTransform& ComponentSpaceTM = ComponentSpacePoseContext.Pose.GetComponentSpaceTransform(OutputData.CompactPoseBoneIndex);
 			const FTransform BodyTM = ConvertCSTransformToSimSpace(SimulationSpace, ComponentSpaceTM, CompWorldSpaceTM, BaseBoneTM);
-			BoneData[BodyIndex].Set(BodyTM.GetLocation(), BodyTM.GetRotation());
+			BoneTMs[BodyIndex] = BodyTM;
 		}
 	}
 	else
 	{
-		BoneData.Empty(OutputBoneData.Num());
+		BoneTMs.Empty(OutputBoneData.Num());
 		for (const FOutputBoneData& OutputData : OutputBoneData)
 		{
 			const int32 BodyIndex = OutputData.BodyIndex;
 			const FTransform& ComponentSpaceTM = ComponentSpacePoseContext.Pose.GetComponentSpaceTransform(OutputData.CompactPoseBoneIndex);
 			const FTransform BodyTM = ConvertCSTransformToSimSpace(SimulationSpace, ComponentSpaceTM, CompWorldSpaceTM, BaseBoneTM);
-			BoneData.Emplace(BodyTM.GetLocation(), BodyTM.GetRotation());
+			BoneTMs.Emplace(BodyTM);
 		}
 	}
 }
diff --git a/Engine/Plugins/Experimental/PhysicsControl/Source/PhysicsControl/Private/RigidBodyWithControl.cpp b/Engine/Plugins/Experimental/PhysicsControl/Source/PhysicsControl/Private/RigidBodyWithControl.cpp
--- a/Engine/Plugins/Experimental/PhysicsControl/Source/PhysicsControl/Private/RigidBodyWithControl.cpp
+++ b/Engine/Plugins/Experimental/PhysicsControl/Source/PhysicsControl/Private/RigidBodyWithControl.cpp
@@ -203,6 +203,14 @@ ImmediatePhysics::FActorHandle* FAnimNode_RigidBodyWithControl::FindBodyFromBone
 	return BodyHandle;
 }
 
+//======================================================================================================================
+void FAnimNode_RigidBodyWithControl::UpdateBodyIndicesInControlRecord(FRigidBodyControlRecord& ControlRecord)
+{
+	ControlRecord.ChildBodyIndex = FindBodyIndexFromBoneName(ControlRecord.Control.ChildBoneName);
+	ControlRecord.ParentBodyIndex = ControlRecord.Control.ParentBoneName.IsNone() ?
+		-1 : FindBodyIndexFromBoneName(ControlRecord.Control.ParentBoneName);
+}
+
 //======================================================================================================================
 FName FAnimNode_RigidBodyWithControl::CreateControl(
 	const FName ParentBoneName, const FName ChildBoneName, const FPhysicsControlData& ControlData)
@@ -315,6 +323,12 @@ void FAnimNode_RigidBodyWithControl::InitControlsAndBodyModifiers(const FReferen
 
 	CreateOperatorsForNode(this, AllLimbBones, RefSkeleton, GetPhysicsAsset(), NameRecords);
 
+	for (TMap<FName, FRigidBodyControlRecord>::ElementType& NameRecordPair : ControlRecords)
+	{
+		FRigidBodyControlRecord& ControlRecord = NameRecordPair.Value;
+		UpdateBodyIndicesInControlRecord(ControlRecord);
+	}
+
 	// Create any additional sets that have been requested
 	CreateAdditionalSets(AdditionalSets, ModifierRecords, ControlRecords, NameRecords);
 
@@ -416,9 +430,11 @@ static void ConvertStrengthToSpringParams(
 }
 
 //======================================================================================================================
-void FAnimNode_RigidBodyWithControl::UpdateDriveSpringDamperSettings(
-	Chaos::FPBDJointSettings&  Settings, 
-	const FPhysicsControlData& ControlData)
+// Adjusts the constraint spring drive settings to reflect the control data
+static void UpdateDriveSpringDamperSettings(
+	Chaos::FPBDJointConstraintHandle* Constraint,
+	const Chaos::FPBDJointSettings&   Settings, 
+	const FPhysicsControlData&        ControlData)
 {
 	float LinearSpring;
 	float LinearDamping;
@@ -430,28 +446,26 @@ void FAnimNode_RigidBodyWithControl::UpdateDriveSpringDamperSettings(
 	ConvertStrengthToSpringParams(AngularSpring, AngularDamping, 
 		ControlData.AngularStrength, ControlData.AngularDampingRatio, ControlData.AngularExtraDamping);
 
-	Settings.LinearDriveStiffness = { LinearSpring, LinearSpring, LinearSpring };
-	Settings.LinearDriveDamping = { LinearDamping, LinearDamping, LinearDamping };
-	Settings.LinearDriveMaxForce = { 0, 0, 0};
-
-	Settings.AngularDriveStiffness = { AngularSpring, AngularSpring, AngularSpring };
-	Settings.AngularDriveDamping = { AngularDamping, AngularDamping, AngularDamping };
-	Settings.AngularDriveMaxTorque = { 0, 0, 0 };
+	Constraint->SetDriveProperties(
+		Chaos::FVec3(LinearSpring), Chaos::FVec3(LinearDamping), Chaos::FVec3(0),
+		Chaos::FVec3(AngularSpring), Chaos::FVec3(AngularDamping), Chaos::FVec3(0));
 }
 
 //======================================================================================================================
-static FTransform CalculateTargetTM(
+static RigidBodyWithControl::FPosQuat CalculateTargetTM(
 	const Chaos::FPBDJointSettings&                 JointSettings, 
 	const RigidBodyWithControl::FRigidBodyPoseData& PoseData,
 	const int32                                     ParentBodyIndex, 
 	const int32                                     ChildBodyIndex)
 {
-	FTransform ChildTargetTM = 
-		JointSettings.ConnectorTransforms[ConstraintChildIndex] * PoseData.GetTM(ChildBodyIndex);
+	const RigidBodyWithControl::FPosQuat ChildTargetTM =
+		RigidBodyWithControl::FPosQuat(JointSettings.ConnectorTransforms[ConstraintChildIndex]) * 
+		PoseData.GetTM(ChildBodyIndex);
 	if (ParentBodyIndex >= 0)
 	{
-		const FTransform ParentTargetTM = 
-			JointSettings.ConnectorTransforms[ConstraintParentIndex] * PoseData.GetTM(ParentBodyIndex);
+		const RigidBodyWithControl::FPosQuat ParentTargetTM =
+			RigidBodyWithControl::FPosQuat(JointSettings.ConnectorTransforms[ConstraintParentIndex]) * 
+			PoseData.GetTM(ParentBodyIndex);
 		return ChildTargetTM * ParentTargetTM.Inverse();
 	}
 	return ChildTargetTM;
@@ -465,7 +479,7 @@ void FAnimNode_RigidBodyWithControl::ApplyControl(FRigidBodyControlRecord& Contr
 
 	if (JointHandle)
 	{
-		Chaos::FPBDJointConstraintHandle* const Constraint = JointHandle->GetConstraint();
+		Chaos::FPBDJointConstraintHandle* Constraint = JointHandle->GetConstraint();
 		if (Constraint)
 		{
 			if (ControlRecord.ExpectedUpdateCounter.Get() != PoseData.UpdateCounter.Get())
@@ -478,7 +492,7 @@ void FAnimNode_RigidBodyWithControl::ApplyControl(FRigidBodyControlRecord& Contr
 
 			if (ChildActorHandle && ParentActorHandle)
 			{
-				Chaos::FPBDJointSettings JointSettings = Constraint->GetSettings();
+				const Chaos::FPBDJointSettings& JointSettings = Constraint->GetSettings();
 
 				// TODO
 				// - cache settings / previous input parameters to avoid unnecessary repeating
@@ -487,60 +501,61 @@ void FAnimNode_RigidBodyWithControl::ApplyControl(FRigidBodyControlRecord& Contr
 				// Update the target point on the child
 				if (ControlRecord.ControlTarget.bUseTargetPoint)
 				{
-					JointSettings.ConnectorTransforms[ConstraintChildIndex].SetLocation(
-						ControlRecord.ControlTarget.TargetPoint);
+					Constraint->SetChildConnectorLocation(ControlRecord.ControlTarget.TargetPoint);
 				}
 				else
 				{
-					FVector ChildCoMPositionOffset = ChildActorHandle->GetLocalCoMTransform().GetLocation();
-					JointSettings.ConnectorTransforms[ConstraintChildIndex].SetLocation(ChildCoMPositionOffset);
+					const FVector ChildCoMPositionOffset = ChildActorHandle->GetLocalCoMLocation();
+					Constraint->SetChildConnectorLocation(ChildCoMPositionOffset);
 				}
 
-				const int32 ChildBodyIndex = FindBodyIndexFromBoneName(ControlRecord.Control.ChildBoneName);
-				const int32 ParentBodyIndex = ControlRecord.Control.ParentBoneName.IsNone() ? 
-					-1 : FindBodyIndexFromBoneName(ControlRecord.Control.ParentBoneName);
+				checkSlow(FindBodyIndexFromBoneName(ControlRecord.Control.ChildBoneName) == ControlRecord.ChildBodyIndex);
+				checkSlow((ControlRecord.Control.ParentBoneName.IsNone() ? -1 : 
+					FindBodyIndexFromBoneName(ControlRecord.Control.ParentBoneName)) == ControlRecord.ParentBodyIndex);
 
-				FTransform TargetTM(
+				RigidBodyWithControl::FPosQuat TargetTM(
 					ControlRecord.ControlTarget.TargetOrientation, 
 					ControlRecord.ControlTarget.TargetPosition);
 
 				if (ControlRecord.ControlTarget.bUseSkeletalAnimation)
 				{
-					FTransform AnimTargetTM = CalculateTargetTM(
-						JointSettings, PoseData, ParentBodyIndex, ChildBodyIndex);
+					RigidBodyWithControl::FPosQuat AnimTargetTM = CalculateTargetTM(
+						JointSettings, PoseData, ControlRecord.ParentBodyIndex, ControlRecord.ChildBodyIndex);
 					TargetTM = TargetTM * AnimTargetTM;
 				}
 
-				JointSettings.LinearDrivePositionTarget = TargetTM.GetTranslation();
-				JointSettings.AngularDrivePositionTarget = TargetTM.GetRotation();
+				Constraint->SetLinearDrivePositionTarget(TargetTM.GetTranslation());
+				Constraint->SetAngularDrivePositionTarget(TargetTM.GetRotation());
 
-				if (DeltaTime != 0)
+				if ((DeltaTime * ControlRecord.CurrentData.LinearTargetVelocityMultiplier) != 0)
+				{
+					FVector Velocity = (TargetTM.GetTranslation() - ControlRecord.PrevTargetTM.GetTranslation()) / DeltaTime;
+					Constraint->SetLinearDriveVelocityTarget(
+						Velocity * ControlRecord.CurrentData.LinearTargetVelocityMultiplier);
+				}
+				else
 				{
-					FTransform PrevTargetTM = ControlRecord.PrevTargetTM;
+					Constraint->SetLinearDriveVelocityTarget(Chaos::FVec3(0));
+				}
 
-					FVector Velocity = (TargetTM.GetTranslation() - PrevTargetTM.GetTranslation()) / DeltaTime;
+				if ((DeltaTime * ControlRecord.CurrentData.AngularTargetVelocityMultiplier) != 0)
+				{
 					// Note that quats multiply in the opposite order to TMs, and must be in the same hemisphere.
 					const FQuat Q = TargetTM.GetRotation();
-					FQuat PrevQ = PrevTargetTM.GetRotation();
+					FQuat PrevQ = ControlRecord.PrevTargetTM.GetRotation();
 					PrevQ.EnforceShortestArcWith(Q);
-
 					const FQuat DeltaQ = Q * PrevQ.Inverse();
 					const FVector AngularVelocity = DeltaQ.ToRotationVector() / DeltaTime;
 
-					JointSettings.LinearDriveVelocityTarget = 
-						Velocity * ControlRecord.CurrentData.LinearTargetVelocityMultiplier;
-					JointSettings.AngularDriveVelocityTarget = 
-						AngularVelocity * ControlRecord.CurrentData.AngularTargetVelocityMultiplier;
+					Constraint->SetAngularDriveVelocityTarget( 
+						AngularVelocity * ControlRecord.CurrentData.AngularTargetVelocityMultiplier);
 				}
 				else
 				{
-					JointSettings.LinearDriveVelocityTarget = FVector::ZeroVector;
-					JointSettings.AngularDriveVelocityTarget = FVector::ZeroVector;
+					Constraint->SetAngularDriveVelocityTarget(Chaos::FVec3(0));
 				}
 
-				UpdateDriveSpringDamperSettings(JointSettings, ControlRecord.CurrentData);
-
-				Constraint->SetSettings(JointSettings);
+				UpdateDriveSpringDamperSettings(Constraint, JointSettings, ControlRecord.CurrentData);
 
 				ControlRecord.PrevTargetTM = TargetTM;
 				ControlRecord.ExpectedUpdateCounter = PoseData.UpdateCounter;
@@ -731,13 +746,12 @@ void FAnimNode_RigidBodyWithControl::ApplyKinematicTargets()
 				const int32 BodyIndex = FindBodyIndexFromBoneName(ModifierRecord->Modifier.BoneName);
 				if (ActorHandle->GetIsKinematic() && BodyIndex != INDEX_NONE)
 				{
-					FTransform TM(Target.TargetOrientation, Target.TargetPosition);
+					RigidBodyWithControl::FPosQuat TM(Target.TargetOrientation, Target.TargetPosition);
 					if (Target.bUseSkeletalAnimation)
 					{
-						FTransform PoseTM = PoseData.GetTM(BodyIndex);
-						TM = TM * PoseTM;
+						TM = TM * PoseData.GetTM(BodyIndex);
 					}
-					ActorHandle->SetKinematicTarget(TM);
+					ActorHandle->SetKinematicTarget(TM.ToTransform());
 				}
 			}
 		}
diff --git a/Engine/Plugins/Experimental/PhysicsControl/Source/PhysicsControl/Public/AnimNode_RigidBodyWithControl.h b/Engine/Plugins/Experimental/PhysicsControl/Source/PhysicsControl/Public/AnimNode_RigidBodyWithControl.h
--- a/Engine/Plugins/Experimental/PhysicsControl/Source/PhysicsControl/Public/AnimNode_RigidBodyWithControl.h
+++ b/Engine/Plugins/Experimental/PhysicsControl/Source/PhysicsControl/Public/AnimNode_RigidBodyWithControl.h
@@ -177,13 +177,6 @@ public:
 	UPROPERTY(EditAnywhere, Category = Settings)
 	uint8 bFreezeIncomingPoseOnStart : 1;
 
-	/**
-		Correct for linear tearing on bodies with all axes Locked.
-		This only works if all axes linear translation are locked
-	*/
-	UPROPERTY(EditAnywhere, Category = Settings)
-	uint8 bClampLinearTranslationLimitToRefPose : 1;
-
 	/**
 		Change the parent space transforms of constraints read from the physics asset to match the relative 
 		bone transforms in the in-coming skeleton.
@@ -331,6 +324,7 @@ private:
 	int32 AddBody(ImmediatePhysics::FActorHandle* const BodyHandle); 
 	int32 FindBodyIndexFromBoneName(const FName BoneName) const;
 	ImmediatePhysics::FActorHandle* FindBodyFromBoneName(const FName BoneName) const;
+	void UpdateBodyIndicesInControlRecord(FRigidBodyControlRecord& ControlRecord);
 
 	ImmediatePhysics::FJointHandle* CreateConstraint(
 		ImmediatePhysics::FActorHandle* const ChildBodyHandle, ImmediatePhysics::FActorHandle* const ParentBodyHandle);
@@ -347,9 +341,6 @@ private:
 	// This applies the desired constraint profile, if necessary
 	void ApplyCurrentConstraintProfile();
 
-	// Adjusts the spring drive settings to reflect the control data
-	void UpdateDriveSpringDamperSettings(Chaos::FPBDJointSettings& Settings, const FPhysicsControlData& ControlData);
-
 	TMap<FName, int32> BodyNameToIndexMap;
 	ImmediatePhysics::FActorHandle* WorldSpaceControlActorHandle;
 
@@ -465,28 +456,13 @@ private:
 		FBodyAnimData()
 			: TransferedBoneAngularVelocity(ForceInit)
 			, TransferedBoneLinearVelocity(ForceInitToZero)
-			, LinearXMotion(ELinearConstraintMotion::LCM_Locked)
-			, LinearYMotion(ELinearConstraintMotion::LCM_Locked)
-			, LinearZMotion(ELinearConstraintMotion::LCM_Locked)
-			, LinearLimit(0.0f)
-			, RefPoseLength (0.f)
 			, bIsSimulated(false)
 			, bBodyTransformInitialized(false)
 		{}
 
-		FQuat TransferedBoneAngularVelocity;
+		FVector TransferedBoneAngularVelocity;
 		FVector TransferedBoneLinearVelocity;
 
-		ELinearConstraintMotion LinearXMotion;
-		ELinearConstraintMotion LinearYMotion;
-		ELinearConstraintMotion LinearZMotion;
-		float LinearLimit;
-		// we don't use linear limit but use default length to limit the bodies
-		// linear limits are defined per constraint - it can be any two joints that can limit
-		// this is just default length of the local space from parent, and we use that info to limit
-		// the translation
-		float RefPoseLength;
-
 		bool bIsSimulated : 1;
 		bool bBodyTransformInitialized : 1;
 	};
diff --git a/Engine/Plugins/Experimental/PhysicsControl/Source/PhysicsControl/Public/RigidBodyControlData.h b/Engine/Plugins/Experimental/PhysicsControl/Source/PhysicsControl/Public/RigidBodyControlData.h
--- a/Engine/Plugins/Experimental/PhysicsControl/Source/PhysicsControl/Public/RigidBodyControlData.h
+++ b/Engine/Plugins/Experimental/PhysicsControl/Source/PhysicsControl/Public/RigidBodyControlData.h
@@ -5,6 +5,7 @@
 #include "PhysicsControlLimbData.h"
 #include "Physics/ImmediatePhysics/ImmediatePhysicsDeclares.h"
 #include "Animation/AnimTypes.h"
+#include "RigidBodyPoseData.h"
 #include "RigidBodyControlData.generated.h"
 
 /**
@@ -339,11 +340,17 @@ struct FRigidBodyControlRecord
 
 	// The previous control target. This will have been set at the end of a previous update (but
 	// only if the control was enabled etc), so to check if it is valid, check the update counter.
-	FTransform PrevTargetTM;
+	RigidBodyWithControl::FPosQuat PrevTargetTM;
 
 	// Update counter set when the control was last updated.
 	// TODO just store the count we're interested in rather than the whole structure
 	FGraphTraversalCounter ExpectedUpdateCounter;
+
+	// Cached child body index - needs to be updated whenever the bone name changes
+	int32 ChildBodyIndex;
+
+	// Cached parent body index - needs to be updated whenever the bone name changes
+	int32 ParentBodyIndex;
 };
 
 /**
diff --git a/Engine/Plugins/Experimental/PhysicsControl/Source/PhysicsControl/Public/RigidBodyPoseData.h b/Engine/Plugins/Experimental/PhysicsControl/Source/PhysicsControl/Public/RigidBodyPoseData.h
--- a/Engine/Plugins/Experimental/PhysicsControl/Source/PhysicsControl/Public/RigidBodyPoseData.h
+++ b/Engine/Plugins/Experimental/PhysicsControl/Source/PhysicsControl/Public/RigidBodyPoseData.h
@@ -28,46 +28,65 @@ struct FOutputBoneData
 };
 
 //======================================================================================================================
-struct FBoneData
+// Simple minimal implementation of a "FTransform without scale"
+struct FPosQuat
 {
-	FBoneData()
-		: Position(FVector::ZeroVector), Orientation(FQuat::Identity) 
-	{}
-	FBoneData(const FVector& InPosition, const FQuat& InOrientation)
-		: Position(InPosition)
-		, Orientation(InOrientation)
-	{}
-
-	/**
-	 * Sets position/orientation in both current and previous (i.e. implying zero velocity)
-	 */
-	void Set(const FVector& InPosition, const FQuat& InOrientation) { 
-		Position = InPosition; Orientation = InOrientation; }
-
-	FTransform GetTM() const { return FTransform(Orientation, Position); }
-
-	FVector Position;
-	FQuat   Orientation;
+	FPosQuat(const FVector& Pos, const FQuat& Quat) : Translation(Pos), Rotation(Quat) {}
+	FPosQuat(const FQuat& Quat, const FVector& Pos) : Translation(Pos), Rotation(Quat) {}
+	FPosQuat(const FRotator& Rotator, const FVector& Pos) : Translation(Pos), Rotation(Rotator) {}
+	FPosQuat() : Translation(EForceInit::ForceInitToZero), Rotation(EForceInit::ForceInit) {}
+	FPosQuat(const FTransform& TM) : Translation(TM.GetTranslation()), Rotation(TM.GetRotation()) {}
+	FPosQuat(ENoInit) {}
+
+	FORCEINLINE FVector GetTranslation() const { return Translation; }
+	FORCEINLINE FQuat GetRotation() const { return Rotation; }
+
+	FORCEINLINE FTransform ToTransform() const
+	{
+		return FTransform(Rotation, Translation);
+	}
+
+	// Note that multiplication operates in the same sense as FTransform (i.e. "backwards")
+	FORCEINLINE FPosQuat operator*(const FPosQuat& Other) const
+	{
+		FQuat OutRotation = Other.Rotation * Rotation;
+		FVector OutTranslation = Other.Rotation * (Translation) + Other.Translation;
+		return FPosQuat(OutTranslation, OutRotation);
+	}
+
+	FORCEINLINE FVector operator*(const FVector& Position) const
+	{
+		return Translation + Rotation * Position;
+	}
+
+	FORCEINLINE FPosQuat Inverse() const
+	{
+		const FQuat OutRotation = Rotation.Inverse();
+		return FPosQuat(OutRotation * -Translation, OutRotation);
+	}
+
+	FVector Translation;
+	FQuat Rotation;
 };
 
 //======================================================================================================================
 struct FRigidBodyPoseData
 {
 	void Update(
-		FComponentSpacePoseContext&   ComponentSpacePoseContext,
-		const TArray<FOutputBoneData> OutputBoneData,
-		const ESimulationSpace        SimulationSpace,
-		const FBoneReference&         BaseBoneRef,
-		const FGraphTraversalCounter& InUpdateCounter);
+		FComponentSpacePoseContext&    ComponentSpacePoseContext,
+		const TArray<FOutputBoneData>& OutputBoneData,
+		const ESimulationSpace         SimulationSpace,
+		const FBoneReference&          BaseBoneRef,
+		const FGraphTraversalCounter&  InUpdateCounter);
 
-	FTransform GetTM(int32 Index) const { return BoneData[Index].GetTM(); }
-	bool IsValidIndex(const int32 Index) const { return BoneData.IsValidIndex(Index); }
-	bool IsEmpty() const { return BoneData.IsEmpty(); }
+	FPosQuat GetTM(int32 Index) const { return BoneTMs[Index]; }
+	bool IsValidIndex(const int32 Index) const { return BoneTMs.IsValidIndex(Index); }
+	bool IsEmpty() const { return BoneTMs.IsEmpty(); }
 
 	/**
 	 * The cached skeletal data, updated at the start of each tick
 	 */
-	TArray<FBoneData> BoneData;
+	TArray<FPosQuat> BoneTMs;
 
 	/**
 	 * The origin (in world space). BoneData transforms are relative to this.
diff --git a/Engine/Source/Runtime/Engine/Private/PhysicsEngine/ImmediatePhysics/ImmediatePhysicsChaos/ImmediatePhysicsActorHandle_Chaos.cpp b/Engine/Source/Runtime/Engine/Private/PhysicsEngine/ImmediatePhysics/ImmediatePhysicsChaos/ImmediatePhysicsActorHandle_Chaos.cpp
--- a/Engine/Source/Runtime/Engine/Private/PhysicsEngine/ImmediatePhysics/ImmediatePhysicsChaos/ImmediatePhysicsActorHandle_Chaos.cpp
+++ b/Engine/Source/Runtime/Engine/Private/PhysicsEngine/ImmediatePhysics/ImmediatePhysicsChaos/ImmediatePhysicsActorHandle_Chaos.cpp
@@ -732,6 +732,11 @@ namespace ImmediatePhysics_Chaos
 		return FTransform(Handle()->RotationOfMass(), Handle()->CenterOfMass());
 	}
 
+	FVector FActorHandle::GetLocalCoMLocation() const
+	{
+		return Handle()->CenterOfMass();
+	}
+
 	int32 FActorHandle::GetLevel() const
 	{
 		return Level;
diff --git a/Engine/Source/Runtime/Engine/Public/Physics/ImmediatePhysics/ImmediatePhysicsChaos/ImmediatePhysicsActorHandle_Chaos.h b/Engine/Source/Runtime/Engine/Public/Physics/ImmediatePhysics/ImmediatePhysicsChaos/ImmediatePhysicsActorHandle_Chaos.h
--- a/Engine/Source/Runtime/Engine/Public/Physics/ImmediatePhysics/ImmediatePhysicsChaos/ImmediatePhysicsActorHandle_Chaos.h
+++ b/Engine/Source/Runtime/Engine/Public/Physics/ImmediatePhysics/ImmediatePhysicsChaos/ImmediatePhysicsActorHandle_Chaos.h
@@ -136,6 +136,9 @@ namespace ImmediatePhysics_Chaos
 		/** Get the actor-space centre of mass offset */
 		ENGINE_API FTransform GetLocalCoMTransform() const;
 
+		/** Get the actor-space centre of mass offset (location only) */
+		ENGINE_API FVector GetLocalCoMLocation() const;
+
 		ENGINE_API Chaos::FGeometryParticleHandle* GetParticle();
 		ENGINE_API const Chaos::FGeometryParticleHandle* GetParticle() const;
 
diff --git a/Engine/Source/Runtime/Experimental/Chaos/Private/Chaos/PBDJointConstraints.cpp b/Engine/Source/Runtime/Experimental/Chaos/Private/Chaos/PBDJointConstraints.cpp
--- a/Engine/Source/Runtime/Experimental/Chaos/Private/Chaos/PBDJointConstraints.cpp
+++ b/Engine/Source/Runtime/Experimental/Chaos/Private/Chaos/PBDJointConstraints.cpp
@@ -132,6 +132,97 @@ namespace Chaos
 		return ConcreteContainer()->GetConstraintResimType(ConstraintIndex);
 	}
 
+	constexpr int32 ConstraintChildIndex = 0;
+	constexpr int32 ConstraintParentIndex = 1;
+
+	FPBDJointSettings& FPBDJointConstraintHandle::GetConstraintSettingsInternal() 
+	{ 
+		return ConcreteContainer()->ConstraintSettings[ConstraintIndex]; 
+	}
+
+	void FPBDJointConstraintHandle::SetParentConnectorLocation(const FVec3 Location)
+	{
+		GetConstraintSettingsInternal().ConnectorTransforms[ConstraintParentIndex].SetLocation(Location);
+	}
+
+	void FPBDJointConstraintHandle::SetParentConnectorRotation(const FQuat Rotation)
+	{
+		GetConstraintSettingsInternal().ConnectorTransforms[ConstraintParentIndex].SetRotation(Rotation);
+	}
+
+	void FPBDJointConstraintHandle::SetChildConnectorLocation(const FVec3 Location)
+	{
+		GetConstraintSettingsInternal().ConnectorTransforms[ConstraintChildIndex].SetLocation(Location);
+	}
+
+	void FPBDJointConstraintHandle::SetChildConnectorRotation(const FQuat Rotation)
+	{
+		GetConstraintSettingsInternal().ConnectorTransforms[ConstraintChildIndex].SetRotation(Rotation);
+	}
+
+	void FPBDJointConstraintHandle::SetLinearDrivePositionTarget(const FVec3 Target)
+	{
+		GetConstraintSettingsInternal().LinearDrivePositionTarget = Target;
+	}
+
+	void FPBDJointConstraintHandle::SetAngularDrivePositionTarget(const FQuat Target)
+	{
+		GetConstraintSettingsInternal().AngularDrivePositionTarget = Target;
+	}
+
+	void FPBDJointConstraintHandle::SetLinearDriveVelocityTarget(const FVec3 Target)
+	{
+		GetConstraintSettingsInternal().LinearDriveVelocityTarget = Target;
+	}
+
+	void FPBDJointConstraintHandle::SetAngularDriveVelocityTarget(const FVec3 Target)
+	{
+		GetConstraintSettingsInternal().AngularDriveVelocityTarget = Target;
+	}
+
+	void FPBDJointConstraintHandle::SetLinearDriveStiffness(const FVec3 Stiffness)
+	{
+		GetConstraintSettingsInternal().LinearDriveStiffness = Stiffness;
+	}
+
+	void FPBDJointConstraintHandle::SetLinearDriveDamping(const FVec3 Damping)
+	{
+		GetConstraintSettingsInternal().LinearDriveDamping = Damping;
+	}
+
+	void FPBDJointConstraintHandle::SetLinearDriveMaxForce(const FVec3 MaxForce)
+	{
+		GetConstraintSettingsInternal().LinearDriveMaxForce = MaxForce;
+	}
+
+	void FPBDJointConstraintHandle::SetAngularDriveStiffness(const FVec3 Stiffness)
+	{
+		GetConstraintSettingsInternal().AngularDriveStiffness = Stiffness;
+	}
+
+	void FPBDJointConstraintHandle::SetAngularDriveDamping(const FVec3 Damping)
+	{
+		GetConstraintSettingsInternal().AngularDriveDamping = Damping;
+	}
+
+	void FPBDJointConstraintHandle::SetAngularDriveMaxTorque(const FVec3 MaxTorque)
+	{
+		GetConstraintSettingsInternal().AngularDriveMaxTorque = MaxTorque;
+	}
+
+	void FPBDJointConstraintHandle::SetDriveProperties(
+		const FVec3 LinearStiffness, const FVec3 LinearDamping, const FVec3 MaxForce,
+		const FVec3 AngularStiffness, const FVec3 AngularDamping, const FVec3 MaxTorque)
+	{
+		FPBDJointSettings& Settings = GetConstraintSettingsInternal();
+		Settings.LinearDriveStiffness = LinearStiffness;
+		Settings.LinearDriveDamping = LinearDamping;
+		Settings.LinearDriveMaxForce = MaxForce;
+		Settings.AngularDriveStiffness = AngularStiffness;
+		Settings.AngularDriveDamping = AngularDamping;
+		Settings.AngularDriveMaxTorque = MaxTorque;
+	}
+
 	const FPBDJointSettings& FPBDJointConstraintHandle::GetSettings() const
 	{
 		return ConcreteContainer()->GetConstraintSettings(ConstraintIndex);
@@ -632,18 +723,21 @@ namespace Chaos
 
 	void FPBDJointConstraints::SetConstraintEnabled(int32 ConstraintIndex, bool bEnabled)
 	{
-		const FGenericParticleHandle Particle0 = FGenericParticleHandle(ConstraintParticles[ConstraintIndex][0]);
-		const FGenericParticleHandle Particle1 = FGenericParticleHandle(ConstraintParticles[ConstraintIndex][1]);
-
 		if (bEnabled)
 		{ 
-			// only enable constraint if the particles are valid and not disabled
-			// and if the constraint is not broken
-			if (Particle0->Handle() != nullptr && !Particle0->Disabled()
-				&& Particle1->Handle() != nullptr && !Particle1->Disabled()
-				&& !IsConstraintBroken(ConstraintIndex))
+			if (ConstraintStates[ConstraintIndex].bDisabled)
 			{
-				ConstraintStates[ConstraintIndex].bDisabled = false;
+				const FGenericParticleHandle Particle0 = FGenericParticleHandle(ConstraintParticles[ConstraintIndex][0]);
+				const FGenericParticleHandle Particle1 = FGenericParticleHandle(ConstraintParticles[ConstraintIndex][1]);
+
+				// only enable constraint if the particles are valid and not disabled
+				// and if the constraint is not broken
+				if (Particle0->Handle() != nullptr && !Particle0->Disabled()
+					&& Particle1->Handle() != nullptr && !Particle1->Disabled()
+					&& !IsConstraintBroken(ConstraintIndex))
+				{
+					ConstraintStates[ConstraintIndex].bDisabled = false;
+				}
 			}
 		}
 		else
diff --git a/Engine/Source/Runtime/Experimental/Chaos/Public/Chaos/PBDJointConstraints.h b/Engine/Source/Runtime/Experimental/Chaos/Public/Chaos/PBDJointConstraints.h
--- a/Engine/Source/Runtime/Experimental/Chaos/Public/Chaos/PBDJointConstraints.h
+++ b/Engine/Source/Runtime/Experimental/Chaos/Public/Chaos/PBDJointConstraints.h
@@ -47,9 +47,36 @@ namespace Chaos
 		CHAOS_API const FPBDJointSettings& GetSettings() const;
 		const FPBDJointSettings& GetJointSettings() const { return GetSettings(); }	//needed for property macros
 
-		// Note that setting drive targets etc can be done through the settings
+		// Note that this is the "regular" way to change any settings, but note that if it results
+		// in a change of state/mode there may be some overhead. For simple/numerical changes, it
+		// may be faster to call one of the specific "Set" functions.
 		CHAOS_API void SetSettings(const FPBDJointSettings& Settings);
 
+		// Individual properties can be set if the settings don't need to be Sanitized().
+
+		CHAOS_API void SetParentConnectorLocation(const FVec3 Location);
+		CHAOS_API void SetParentConnectorRotation(const FQuat Rotation);
+		CHAOS_API void SetChildConnectorLocation(const FVec3 Location);
+		CHAOS_API void SetChildConnectorRotation(const FQuat Rotation);
+
+		CHAOS_API void SetLinearDrivePositionTarget(const FVec3 Target);
+		CHAOS_API void SetAngularDrivePositionTarget(const FQuat Target);
+
+		CHAOS_API void SetLinearDriveVelocityTarget(const FVec3 Target);
+		CHAOS_API void SetAngularDriveVelocityTarget(const FVec3 Target);
+
+		CHAOS_API void SetLinearDriveStiffness(const FVec3 Stiffness);
+		CHAOS_API void SetLinearDriveDamping(const FVec3 Damping);
+		CHAOS_API void SetLinearDriveMaxForce(const FVec3 MaxForce);
+
+		CHAOS_API void SetAngularDriveStiffness(const FVec3 Stiffness);
+		CHAOS_API void SetAngularDriveDamping(const FVec3 Damping);
+		CHAOS_API void SetAngularDriveMaxTorque(const FVec3 MaxTorque);
+
+		CHAOS_API void SetDriveProperties(
+			const FVec3 LinearStiffness, const FVec3 LinearDamping, const FVec3 MaxForce,
+			const FVec3 AngularStiffness, const FVec3 AngularDamping, const FVec3 MaxTorque);
+
 		CHAOS_API TVec2<FGeometryParticleHandle*> GetConstrainedParticles() const override final;
 
 		static const FConstraintHandleTypeID& StaticType()
@@ -75,6 +102,9 @@ namespace Chaos
 		using Base::ConstraintIndex;
 		using Base::ConcreteContainer;
 	private:
+		// Our own direct access to the settings, to modify things that don't need to go through sanitize
+		FPBDJointSettings& GetConstraintSettingsInternal();
+
 		bool bLinearPlasticityInitialized;
 		bool bAngularPlasticityInitialized;
 	};
