Re-organise the PhysicsControlComponent to remove the private implementation and make it more amenable to extension.
Handle missing skeletal mesh component.
Fully disable the constraints when strength and damping are zero.
Make it possible to set the physics asset in RigidBodyWithControl whilst it is running.
Report attempts to use invalid controls etc as warnings
Move ControlSettings into the ControlData struct
Further progress towards aligning PCC and RBWC
Minor changes to physics: expose collision and parent inverse mass scaling in the PBDJointConstraintHandle. Check for collision in the BasicBroadPhase. Remove temporary/workaround code.
Use TWeakObjectPtr instead of TObjectPtr for storing references to objects owned elsewhere

diff --git a/Engine/Plugins/Experimental/PhysicsControl/Source/PhysicsControl/Private/AnimNode_RigidBodyWithControl.cpp b/Engine/Plugins/Experimental/PhysicsControl/Source/PhysicsControl/Private/AnimNode_RigidBodyWithControl.cpp
--- a/Engine/Plugins/Experimental/PhysicsControl/Source/PhysicsControl/Private/AnimNode_RigidBodyWithControl.cpp
+++ b/Engine/Plugins/Experimental/PhysicsControl/Source/PhysicsControl/Private/AnimNode_RigidBodyWithControl.cpp
@@ -39,7 +39,9 @@ LLM_DEFINE_TAG(Animation_RigidBodyWithControl);
 
 #define LOCTEXT_NAMESPACE "ImmediatePhysicsWithControl"
 
-//DEFINE_STAT(STAT_RigidBodyNodeWithControlInitTime);
+DECLARE_CYCLE_STAT(TEXT("RigidBodyNodeWithControl_InitPhysics"), STAT_RigidBodyWithControlInitPhysicsTime, STATGROUP_Anim);
+DECLARE_CYCLE_STAT(TEXT("RigidBodyNodeWithControl_SetupControls"), STAT_RigidBodyWithControlSetupControlsTime, STATGROUP_Anim);
+
 DECLARE_CYCLE_STAT(TEXT("RigidBodyNodeWithControl_Eval"), STAT_RigidBodyNodeWithControl_Eval, STATGROUP_Anim);
 DECLARE_CYCLE_STAT(TEXT("RigidBodyNodeWithControl_Simulation"), STAT_RigidBodyNodeWithControl_Simulation, STATGROUP_Anim);
 DECLARE_CYCLE_STAT(TEXT("RigidBodyNodeWithControl_SimulationWait"), STAT_RigidBodyNodeWithControl_SimulationWait, STATGROUP_Anim);
@@ -494,6 +496,8 @@ static FTransform CalculateJointTargetTransform(
 
 void FAnimNode_RigidBodyWithControl::SetupControls(USkeletalMeshComponent* const SkeletalMeshComponent)
 {
+	SCOPE_CYCLE_COUNTER(STAT_RigidBodyWithControlSetupControlsTime);
+
 	bool bSuccess = false;
 
 	if (SkeletalMeshComponent)
@@ -956,12 +960,30 @@ void ComputeBodyInsertionOrderWithControl(TArray<FBoneIndexType>& InsertionOrder
 	}
 }
 
+UPhysicsAsset* FAnimNode_RigidBodyWithControl::GetPhysicsAssetToBeUsed(const UAnimInstance* InAnimInstance) const
+{
+	if (OverridePhysicsAsset.IsValid())
+	{
+		return OverridePhysicsAsset.Get();
+	}
+
+	if (bDefaultToSkeletalMeshPhysicsAsset && ensure(InAnimInstance))
+	{
+		const USkeletalMeshComponent* SkeletalMeshComp = InAnimInstance->GetSkelMeshComponent();
+		if (SkeletalMeshComp)
+		{
+			return SkeletalMeshComp->GetPhysicsAsset();
+		}
+	}
+
+	return nullptr;
+}
+
 void FAnimNode_RigidBodyWithControl::InitPhysics(const UAnimInstance* InAnimInstance)
 {
-	LLM_SCOPE_BYNAME(TEXT("Animation/RigidBodyWithControl")); 
-	
-//	SCOPE_CYCLE_COUNTER(STAT_RigidBodyNodeInitTime);
+	SCOPE_CYCLE_COUNTER(STAT_RigidBodyWithControlInitPhysicsTime);
 
+	DestroyControlsAndBodyModifiers();
 	DestroyPhysicsSimulation();
 
 	const USkeletalMeshComponent* SkeletalMeshComp = InAnimInstance->GetSkelMeshComponent();
@@ -976,7 +998,7 @@ void FAnimNode_RigidBodyWithControl::InitPhysics(const UAnimInstance* InAnimInst
 	}
 
 	const FReferenceSkeleton& SkelMeshRefSkel = SkeletalMeshAsset->GetRefSkeleton();
-	PhysicsAssetToUse = OverridePhysicsAsset ? ToRawPtr(OverridePhysicsAsset) : InAnimInstance->GetSkelMeshComponent()->GetPhysicsAsset();
+	PhysicsAssetToUse = GetPhysicsAssetToBeUsed(InAnimInstance);
 
 	ensure(SkeletonAsset == SkeletalMeshAsset->GetSkeleton());
 
@@ -1040,7 +1062,8 @@ void FAnimNode_RigidBodyWithControl::InitPhysics(const UAnimInstance* InAnimInst
 			FPhysicsAggregateHandle(),
 			bCreateBodiesInRefPose);
 
-		if (bModifyConstraintTransformsToMatchSkeleton && (CVarEnableRigidBodyNodeWithControlMatchingConstraintsToSkeleton.GetValueOnAnyThread() > 0))
+		if (bModifyConstraintTransformsToMatchSkeleton && 
+			CVarEnableRigidBodyNodeWithControlMatchingConstraintsToSkeleton.GetValueOnAnyThread() > 0)
 		{
 			TransformConstraintsToMatchSkeletalMesh(SkeletalMeshAsset, HighLevelConstraintInstances);
 		}
@@ -1054,14 +1077,15 @@ void FAnimNode_RigidBodyWithControl::InitPhysics(const UAnimInstance* InAnimInst
 		TArray<FBoneIndexType> InsertionOrder;
 		ComputeBodyInsertionOrderWithControl(InsertionOrder, *SkeletalMeshComp);
 
-		// NOTE: NumBonesLOD0 may be less than NumBonesTotal, and it may be middle bones that are missing from LOD0.
-		// In this case, LOD0 bone indices may be >= NumBonesLOD0, but always < NumBonesTotal. Arrays indexed by
-		// bone index must be size NumBonesTotal.
+		// NOTE: NumBonesLOD0 may be less than NumBonesTotal, and it may be middle bones that are
+		// missing from LOD0. In this case, LOD0 bone indices may be >= NumBonesLOD0, but always <
+		// NumBonesTotal. Arrays indexed by bone index must be size NumBonesTotal.
 		const int32 NumBonesLOD0 = InsertionOrder.Num();
 		const int32 NumBonesTotal = SkelMeshRefSkel.GetNum();
 
-		// If our skeleton is not the one that was used to build the PhysicsAsset, some bodies may be missing, or rearranged.
-		// We need to map the original indices to the new bodies for use by the CollisionDisableTable.
+		// If our skeleton is not the one that was used to build the PhysicsAsset, some bodies may
+		// be missing, or rearranged. We need to map the original indices to the new bodies for use
+		// by the CollisionDisableTable.
 		// NOTE: This array is indexed by the original BodyInstance body index (BodyInstance->InstanceBodyIndex)
 		TArray<ImmediatePhysics::FActorHandle*> BodyIndexToActorHandle;
 		BodyIndexToActorHandle.AddZeroed(HighLevelBodyInstances.Num());
@@ -1432,8 +1456,25 @@ void FAnimNode_RigidBodyWithControl::ResetDynamics(ETeleportType InTeleportType)
 	ResetSimulatedTeleportType = ((InTeleportType > ResetSimulatedTeleportType) ? InTeleportType : ResetSimulatedTeleportType);
 }
 
+void FAnimNode_RigidBodyWithControl::SetOverridePhysicsAsset(UPhysicsAsset* PhysicsAsset)
+{
+	OverridePhysicsAsset = PhysicsAsset;
+}
+
 void FAnimNode_RigidBodyWithControl::PreUpdate(const UAnimInstance* InAnimInstance)
 {
+	// Detect changes in the physics asset to be used. This can happen when using the override
+	// physics asset feature.
+	UPhysicsAsset* PhysicsAssetToBeUsed = GetPhysicsAssetToBeUsed(InAnimInstance);
+	if (GetPhysicsAsset() != PhysicsAssetToBeUsed)
+	{
+		InitPhysics(InAnimInstance);
+
+		// Update the bone references after a change in the physics asset. This needs to happen
+		// after initializing physics as the Bodies set up in InitPhysics() need to be up to date.
+		InitializeBoneReferences(InAnimInstance->GetRequiredBones());
+	}
+
 	// Don't update geometry if RBN is disabled
 	if(!bEnabled)
 	{
diff --git a/Engine/Plugins/Experimental/PhysicsControl/Source/PhysicsControl/Private/PhysicsControlActor.cpp b/Engine/Plugins/Experimental/PhysicsControl/Source/PhysicsControl/Private/PhysicsControlActor.cpp
--- a/Engine/Plugins/Experimental/PhysicsControl/Source/PhysicsControl/Private/PhysicsControlActor.cpp
+++ b/Engine/Plugins/Experimental/PhysicsControl/Source/PhysicsControl/Private/PhysicsControlActor.cpp
@@ -18,7 +18,7 @@ void UPhysicsControlInitializerComponent::CreateInitialCharacterControls(UPhysic
 {
 	USkeletalMeshComponent* SkeletalMeshComponent = nullptr;
 
-	if (InitialCharacterControls.CharacterActor)
+	if (InitialCharacterControls.CharacterActor.IsValid())
 	{
 		if (InitialCharacterControls.SkeletalMeshComponentName != NAME_None)
 		{
@@ -53,10 +53,8 @@ void UPhysicsControlInitializerComponent::CreateInitialCharacterControls(UPhysic
 			SkeletalMeshComponent,
 			InitialCharacterControls.LimbSetupData,
 			InitialCharacterControls.WorldSpaceControlData,
-			InitialCharacterControls.WorldSpaceControlSettings,
 			InitialCharacterControls.ParentSpaceControlData,
-			InitialCharacterControls.ParentSpaceControlSettings,
-			InitialCharacterControls.PhysicsMovementType
+			InitialCharacterControls.BodyModifierData
 		);
 	}
 }
@@ -71,7 +69,7 @@ void UPhysicsControlInitializerComponent::CreateOrUpdateInitialControls(UPhysics
 		const FInitialPhysicsControl& InitialPhysicsControl = InitialPhysicsControlPair.Value;
 
 		UMeshComponent* ChildMeshComponent = nullptr;
-		if (InitialPhysicsControl.ChildActor)
+		if (InitialPhysicsControl.ChildActor.IsValid())
 		{
 			if (InitialPhysicsControl.ChildMeshComponentName != NAME_None)
 			{
@@ -96,7 +94,7 @@ void UPhysicsControlInitializerComponent::CreateOrUpdateInitialControls(UPhysics
 		}
 
 		UMeshComponent* ParentMeshComponent = nullptr;
-		if (InitialPhysicsControl.ParentActor)
+		if (InitialPhysicsControl.ParentActor.IsValid())
 		{
 			if (InitialPhysicsControl.ParentMeshComponentName != NAME_None)
 			{
@@ -117,7 +115,10 @@ void UPhysicsControlInitializerComponent::CreateOrUpdateInitialControls(UPhysics
 
 		// If we get here then we will be removing the original and making a new one in its place.
 		TArray<FName> Sets = ControlComponent->GetSetsContainingControl(Name);
-		ControlComponent->DestroyControl(Name);
+		if (ControlComponent->GetControlExists(Name))
+		{
+			ControlComponent->DestroyControl(Name);
+		}
 
 		if (ControlComponent->CreateNamedControl(
 			Name,
@@ -127,7 +128,6 @@ void UPhysicsControlInitializerComponent::CreateOrUpdateInitialControls(UPhysics
 			InitialPhysicsControl.ChildBoneName,
 			InitialPhysicsControl.ControlData,
 			InitialPhysicsControl.ControlTarget,
-			InitialPhysicsControl.ControlSettings,
 			"All"))
 		{
 			// If we were replacing a control then add it to the original sets
@@ -150,7 +150,7 @@ void UPhysicsControlInitializerComponent::CreateOrUpdateInitialBodyModifiers(UPh
 		const FInitialBodyModifier& InitialBodyModifier = InitialBodyModifierPair.Value;
 
 		UMeshComponent* MeshComponent = nullptr;
-		if (InitialBodyModifier.Actor)
+		if (InitialBodyModifier.Actor.IsValid())
 		{
 			if (InitialBodyModifier.MeshComponentName != NAME_None)
 			{
@@ -176,17 +176,13 @@ void UPhysicsControlInitializerComponent::CreateOrUpdateInitialBodyModifiers(UPh
 
 		// If we get here then we will be removing the original and making a new one in its place.
 		TArray<FName> Sets = ControlComponent->GetSetsContainingBodyModifier(Name);
-		ControlComponent->DestroyBodyModifier(Name);
+		if (ControlComponent->GetBodyModifierExists(Name))
+		{
+			ControlComponent->DestroyBodyModifier(Name);
+		}
 
 		if (ControlComponent->CreateNamedBodyModifier(
-			Name,
-			MeshComponent,
-			InitialBodyModifier.BoneName,
-			"All",
-			InitialBodyModifier.MovementType,
-			InitialBodyModifier.CollisionType,
-			InitialBodyModifier.GravityMultiplier,
-			InitialBodyModifier.bUseSkeletalAnimation))
+			Name, MeshComponent, InitialBodyModifier.BoneName, "All", InitialBodyModifier.BodyModifierData))
 		{
 			// If we were replacing a control then add it to the original sets
 			FPhysicsControlNames Names;
diff --git a/Engine/Plugins/Experimental/PhysicsControl/Source/PhysicsControl/Private/PhysicsControlBPLibrary.cpp b/Engine/Plugins/Experimental/PhysicsControl/Source/PhysicsControl/Private/PhysicsControlBPLibrary.cpp
--- a/Engine/Plugins/Experimental/PhysicsControl/Source/PhysicsControl/Private/PhysicsControlBPLibrary.cpp
+++ b/Engine/Plugins/Experimental/PhysicsControl/Source/PhysicsControl/Private/PhysicsControlBPLibrary.cpp
@@ -11,21 +11,21 @@
 
 //======================================================================================================================
 template<typename TParameterType> const TArray<FName>* FindNamesInSet(
-	const FRigidBodyNameRecords& NameRecords, const FName SetName)
+	const FPhysicsControlNameRecords& NameRecords, const FName SetName)
 {
 	return nullptr;
 }
 
 //======================================================================================================================
 template<> const TArray<FName>* FindNamesInSet<FRigidBodyControl>(
-	const FRigidBodyNameRecords& NameRecords, const FName SetName)
+	const FPhysicsControlNameRecords& NameRecords, const FName SetName)
 {
 	return NameRecords.ControlSets.Find(SetName);
 }
 
 //======================================================================================================================
 template<> const TArray<FName>* FindNamesInSet<FRigidBodyModifier>(
-	const FRigidBodyNameRecords& NameRecords, const FName SetName)
+	const FPhysicsControlNameRecords& NameRecords, const FName SetName)
 {
 	return NameRecords.BodyModifierSets.Find(SetName);
 }
@@ -194,6 +194,20 @@ FRigidBodyWithControlReference UPhysicsControlBPLibrary::ConvertToRigidBodyWithC
 	return FAnimNodeReference::ConvertToType<FRigidBodyWithControlReference>(Node, Result);
 }
 
+//======================================================================================================================
+FRigidBodyWithControlReference UPhysicsControlBPLibrary::SetOverridePhysicsAsset(
+	const FRigidBodyWithControlReference& Node, UPhysicsAsset* PhysicsAsset)
+{
+	Node.CallAnimNodeFunction<FAnimNode_RigidBodyWithControl>(
+		TEXT("SetOverridePhysicsAsset"),
+		[PhysicsAsset](FAnimNode_RigidBodyWithControl& Node)
+		{
+			Node.SetOverridePhysicsAsset(PhysicsAsset);
+		});
+
+	return Node;
+}
+
 //======================================================================================================================
 UFUNCTION(BlueprintPure, Category = "Animation|PhysicsControl", meta = (BlueprintThreadSafe))
 TArray<FName> UPhysicsControlBPLibrary::GetControlNamesInSet(
diff --git a/Engine/Plugins/Experimental/PhysicsControl/Source/PhysicsControl/Private/PhysicsControlComponent.cpp b/Engine/Plugins/Experimental/PhysicsControl/Source/PhysicsControl/Private/PhysicsControlComponent.cpp
--- a/Engine/Plugins/Experimental/PhysicsControl/Source/PhysicsControl/Private/PhysicsControlComponent.cpp
+++ b/Engine/Plugins/Experimental/PhysicsControl/Source/PhysicsControl/Private/PhysicsControlComponent.cpp
@@ -4,14 +4,11 @@
 #include "PhysicsControlLog.h"
 #include "PhysicsControlRecord.h"
 #include "PhysicsControlComponentHelpers.h"
-#include "PhysicsControlComponentImpl.h"
 
-#include "PhysicsEngine/BodySetup.h"
 #include "PhysicsEngine/ConstraintInstance.h"
 #include "PhysicsEngine/BodyInstance.h"
 #include "PhysicsEngine/PhysicsAsset.h"
 
-#include "Physics/Experimental/PhysScene_Chaos.h"
 #include "Physics/PhysicsInterfaceCore.h"
 
 #include "Components/StaticMeshComponent.h"
@@ -19,7 +16,6 @@
 #include "Components/BillboardComponent.h"
 
 #include "Engine/Engine.h"
-#include "Engine/World.h"
 #include "Engine/Texture2D.h"
 
 #include "SceneManagement.h"
@@ -27,6 +23,12 @@
 DECLARE_CYCLE_STAT(TEXT("PhysicsControl UpdateTargetCaches"), STAT_PhysicsControl_UpdateTargetCaches, STATGROUP_Anim);
 DECLARE_CYCLE_STAT(TEXT("PhysicsControl UpdateControls"), STAT_PhysicsControl_UpdateControls, STATGROUP_Anim);
 
+//UE_DISABLE_OPTIMIZATION;
+
+//======================================================================================================================
+// This file contains the public member functions of UPhysicsControlComponent
+//======================================================================================================================
+
 //======================================================================================================================
 // This is used, rather than UEnum::GetValueAsString, so that we have more control over the string returned, which 
 // gets used as a prefix for the automatically named sets etc
@@ -47,8 +49,6 @@ static FName GetControlTypeName(EPhysicsControlType ControlType)
 UPhysicsControlComponent::UPhysicsControlComponent(const FObjectInitializer& ObjectInitializer)
 	: Super(ObjectInitializer)
 {
-	Implementation = MakePimpl<FPhysicsControlComponentImpl>(this);
-
 	// ActorComponent setup
 	bWantsInitializeComponent = true;
 	PrimaryComponentTick.bCanEverTick = true;
@@ -64,28 +64,26 @@ UPhysicsControlComponent::UPhysicsControlComponent(const FObjectInitializer& Obj
 void UPhysicsControlComponent::InitializeComponent()
 {
 	Super::InitializeComponent();
-	Implementation->ResetControls(false);
+	ResetControls(false);
 }
 
 //======================================================================================================================
 void UPhysicsControlComponent::BeginDestroy()
 {
-	if (Implementation)
+	for (TPair<FName, FPhysicsControlRecord>& PhysicsControlRecordPair : PhysicsControlRecords)
 	{
-		for (TPair<FName, FPhysicsControlRecord>& PhysicsControlRecordPair : Implementation->PhysicsControlRecords)
-		{
-			Implementation->DestroyControl(
-				PhysicsControlRecordPair.Key, FPhysicsControlComponentImpl::EDestroyBehavior::KeepRecord);
-		}
-		Implementation->PhysicsControlRecords.Empty();
+		DestroyControl(
+			PhysicsControlRecordPair.Key, EDestroyBehavior::KeepRecord);
+	}
+	PhysicsControlRecords.Empty();
 
-		for (TPair<FName, FPhysicsBodyModifier>& PhysicsBodyModifierPair : Implementation->PhysicsBodyModifiers)
-		{
-			Implementation->DestroyBodyModifier(
-				PhysicsBodyModifierPair.Key, FPhysicsControlComponentImpl::EDestroyBehavior::KeepRecord);
-		}
-		Implementation->PhysicsBodyModifiers.Empty();
+	for (TPair<FName, FPhysicsBodyModifier>& PhysicsBodyModifierPair : PhysicsBodyModifiers)
+	{
+		DestroyBodyModifier(
+			PhysicsBodyModifierPair.Key, EDestroyBehavior::KeepRecord);
 	}
+	PhysicsBodyModifiers.Empty();
+
 	Super::BeginDestroy();
 }
 
@@ -96,7 +94,7 @@ void UPhysicsControlComponent::UpdateTargetCaches(float DeltaTime)
 	TRACE_CPUPROFILER_EVENT_SCOPE(UPhysicsControlComponent::UpdateTargetCaches);
 
 	// Update the skeletal mesh caches
-	Implementation->UpdateCachedSkeletalBoneData(DeltaTime);
+	UpdateCachedSkeletalBoneData(DeltaTime);
 }
 
 //======================================================================================================================
@@ -105,95 +103,26 @@ void UPhysicsControlComponent::UpdateControls(float DeltaTime)
 	SCOPE_CYCLE_COUNTER(STAT_PhysicsControl_UpdateControls);
 	TRACE_CPUPROFILER_EVENT_SCOPE(UPhysicsControlComponent::UpdateControls);
 
-	// Tidy up
-	Implementation->PhysicsControlRecords.Compact();
-	Implementation->PhysicsBodyModifiers.Compact();
+	PhysicsControlRecords.Compact();
+	PhysicsBodyModifiers.Compact();
 
-	for (TPair<FName, FPhysicsControlRecord>& RecordPair : Implementation->PhysicsControlRecords)
+	for (TPair<FName, FPhysicsControlRecord>& RecordPair : PhysicsControlRecords)
 	{
 		// New constraint requested when one doesn't exist
 		FName ControlName = RecordPair.Key;
 		FPhysicsControlRecord& Record = RecordPair.Value;
-		FConstraintInstance* ConstraintInstance = Record.PhysicsControlState.ConstraintInstance.Get();
-		if (!ConstraintInstance)
+		if (!Record.ConstraintInstance)
 		{
-			ConstraintInstance = Record.CreateConstraint(this, ControlName);
-		}
-
-		if (ConstraintInstance)
-		{
-			// Always control collision, because otherwise maintaining it is very difficult, since
-			// constraint-controlled collision doesn't interact nicely when there are multiple constraints.
-			ConstraintInstance->SetDisableCollision(Record.PhysicsControl.ControlSettings.bDisableCollision);
-
-			// Constraint is not enabled
-			if (!Record.Enabled())
-			{
-				ConstraintInstance->SetAngularDriveParams(0.0f, 0.0f, 0.0f);
-				ConstraintInstance->SetLinearDriveParams(0.0f, 0.0f, 0.0f);
-			}
-			else
-			{
-				Implementation->ApplyControl(Record);
-			}
+			Record.InitConstraint(this, ControlName);
 		}
+		ApplyControl(Record);
 	}
 
 	// Handle body modifiers
-	for (TPair<FName, FPhysicsBodyModifier>& BodyModifierPair : Implementation->PhysicsBodyModifiers)
+	for (TPair<FName, FPhysicsBodyModifier>& BodyModifierPair : PhysicsBodyModifiers)
 	{
 		FPhysicsBodyModifier& BodyModifier = BodyModifierPair.Value;
-		FBodyInstance* BodyInstance = UE::PhysicsControlComponent::GetBodyInstance(
-			BodyModifier.MeshComponent, BodyModifier.BoneName);
-
-		switch (BodyModifier.MovementType)
-		{
-		case EPhysicsMovementType::Static:
-			BodyInstance->SetInstanceSimulatePhysics(false, true);
-			break;
-		case EPhysicsMovementType::Kinematic:
-			BodyInstance->SetInstanceSimulatePhysics(false, true);
-			Implementation->ApplyKinematicTarget(BodyModifier);
-			break;
-		case EPhysicsMovementType::Simulated:
-			BodyInstance->SetInstanceSimulatePhysics(true, true, true);
-			break;
-		default:
-			UE_LOG(LogPhysicsControlComponent, Warning, TEXT("Invalid movement type %d"), int(BodyModifier.MovementType));
-			break;
-		}
-		// We always overwrite the physics blend weight, since the functions above can still modify
-		// it (even though they all use the "maintain physics blending" option), since there is an
-		// expectation that zero blend weight means to disable physics.
-		BodyInstance->PhysicsBlendWeight = BodyModifier.PhysicsBlendWeight;
-
-		BodyInstance->SetUpdateKinematicFromSimulation(BodyModifier.bUpdateKinematicFromSimulation);
-
-		if (BodyModifier.bResetToCachedTarget)
-		{
-			BodyModifier.bResetToCachedTarget = false;
-			Implementation->ResetToCachedTarget(BodyModifier);
-		}
-
-
-		UBodySetup* BodySetup = BodyInstance->GetBodySetup();
-		if (BodySetup)
-		{
-			int32 NumShapes = BodySetup->AggGeom.GetElementCount();
-			for (int32 ShapeIndex = 0 ; ShapeIndex != NumShapes ; ++ShapeIndex)
-			{
-				BodyInstance->SetShapeCollisionEnabled(ShapeIndex, BodyModifier.CollisionType);
-			}
-		}
-
-		if (BodyInstance->IsInstanceSimulatingPhysics())
-		{
-			float GravityZ = BodyInstance->GetPhysicsScene()->GetOwningWorld()->GetGravityZ();
-			float AppliedGravityZ = BodyInstance->bEnableGravity ? GravityZ : 0.0f;
-			float DesiredGravityZ = GravityZ * BodyModifier.GravityMultiplier;
-			float GravityZToApply = DesiredGravityZ - AppliedGravityZ;
-			BodyInstance->AddForce(FVector(0, 0, GravityZToApply), true, true);
-		}
+		ApplyBodyModifier(BodyModifier);
 	}
 }
 
@@ -287,14 +216,12 @@ FName UPhysicsControlComponent::CreateControl(
 	const FName                   ChildBoneName,
 	const FPhysicsControlData     ControlData, 
 	const FPhysicsControlTarget   ControlTarget, 
-	const FPhysicsControlSettings ControlSettings,
 	const FName                   Set,
 	const FString                 NamePrefix)
 {
-	const FName Name = Implementation->GetUniqueControlName(ParentBoneName, ChildBoneName, NamePrefix);
+	const FName Name = GetUniqueControlName(ParentBoneName, ChildBoneName, NamePrefix);
 	if (CreateNamedControl(
-		Name, ParentMeshComponent, ParentBoneName, ChildMeshComponent, ChildBoneName, 
-		ControlData, ControlTarget, ControlSettings, Set))
+		Name, ParentMeshComponent, ParentBoneName, ChildMeshComponent, ChildBoneName, ControlData, ControlTarget, Set))
 	{
 		return Name;
 	}
@@ -311,11 +238,12 @@ bool UPhysicsControlComponent::CreateNamedControl(
 	const FName                   ChildBoneName,
 	const FPhysicsControlData     ControlData, 
 	const FPhysicsControlTarget   ControlTarget, 
-	const FPhysicsControlSettings ControlSettings,
 	const FName                   Set)
 {
-	if (Implementation->FindControlRecord(Name))
+	if (FindControlRecord(Name))
 	{
+		UE_LOG(LogPhysicsControlComponent, Warning,
+			TEXT("Unable to make a Control as one with the desired name already exists"), *Name.ToString());
 		return false;
 	}
 
@@ -328,20 +256,19 @@ bool UPhysicsControlComponent::CreateNamedControl(
 
 	if (USkeletalMeshComponent* SkeletalMeshComponent = Cast<USkeletalMeshComponent>(ParentMeshComponent))
 	{
-		Implementation->AddSkeletalMeshReferenceForCaching(SkeletalMeshComponent);
+		AddSkeletalMeshReferenceForCaching(SkeletalMeshComponent);
 	}
 	if (USkeletalMeshComponent* SkeletalMeshComponent = Cast<USkeletalMeshComponent>(ChildMeshComponent))
 	{
-		Implementation->AddSkeletalMeshReferenceForCaching(SkeletalMeshComponent);
+		AddSkeletalMeshReferenceForCaching(SkeletalMeshComponent);
 	}
 
-	FPhysicsControlRecord& NewRecord = Implementation->PhysicsControlRecords.Add(
+	FPhysicsControlRecord& NewRecord = PhysicsControlRecords.Add(
 		Name, FPhysicsControl(
-			ParentMeshComponent, ParentBoneName, ChildMeshComponent, ChildBoneName,
-			ControlData, ControlTarget, ControlSettings));
+			ParentMeshComponent, ParentBoneName, ChildMeshComponent, ChildBoneName, ControlData, ControlTarget));
 	NewRecord.ResetControlPoint();
 
-	Implementation->NameRecords.AddControl(Name, Set);
+	NameRecords.AddControl(Name, Set);
 
 	return true;
 }
@@ -354,7 +281,6 @@ TArray<FName> UPhysicsControlComponent::CreateControlsFromSkeletalMeshBelow(
 	const bool                    bIncludeSelf,
 	const EPhysicsControlType     ControlType,
 	const FPhysicsControlData     ControlData,
-	const FPhysicsControlSettings ControlSettings,
 	const FName                   Set)
 {
 	TArray<FName> Result;
@@ -373,7 +299,7 @@ TArray<FName> UPhysicsControlComponent::CreateControlsFromSkeletalMeshBelow(
 		BoneName, bIncludeSelf, /*bSkipCustomType=*/false, 
 		[
 			this, PhysicsAsset, ParentMeshComponent, SkeletalMeshComponent, 
-			ControlType, &ControlData, &ControlSettings, Set, &Result
+			ControlType, &ControlData, Set, &Result
 		](const FBodyInstance* BI)
 		{
 			if (USkeletalBodySetup* BodySetup = Cast<USkeletalBodySetup>(BI->BodySetup.Get()))
@@ -392,12 +318,12 @@ TArray<FName> UPhysicsControlComponent::CreateControlsFromSkeletalMeshBelow(
 				}
 				const FName ControlName = CreateControl(
 					ParentMeshComponent, ParentBoneName, SkeletalMeshComponent, ChildBoneName,
-					ControlData, FPhysicsControlTarget(), ControlSettings, 
+					ControlData, FPhysicsControlTarget(), 
 					FName(GetControlTypeName(ControlType).ToString().Append("_").Append(Set.ToString())));
 				if (!ControlName.IsNone())
 				{
 					Result.Add(ControlName);
-					Implementation->NameRecords.AddControl(ControlName, GetControlTypeName(ControlType));
+					NameRecords.AddControl(ControlName, GetControlTypeName(ControlType));
 				}
 				else
 				{
@@ -428,16 +354,11 @@ TArray<FName> UPhysicsControlComponent::CreateControlsFromSkeletalMeshAndConstra
 		return Result;
 	}
 
-	FPhysicsControlSettings ControlSettings;
-	// This is to match the skeletal mesh component velocity drive, which does not use the
-	// target animation velocity.
-	ControlSettings.SkeletalAnimationVelocityMultiplier = 0;
-
 	SkeletalMeshComponent->ForEachBodyBelow(
 		BoneName, bIncludeSelf, /*bSkipCustomType=*/false,
 		[
 			this, PhysicsAsset, SkeletalMeshComponent,
-			ConstraintProfile, Set, &ControlSettings, &Result, bEnabled
+			ConstraintProfile, Set, &Result, bEnabled
 		](const FBodyInstance* BI)
 		{
 			if (USkeletalBodySetup* BodySetup = Cast<USkeletalBodySetup>(BI->BodySetup.Get()))
@@ -453,6 +374,9 @@ TArray<FName> UPhysicsControlComponent::CreateControlsFromSkeletalMeshAndConstra
 				}
 
 				FPhysicsControlData ControlData;
+				// This is to match the skeletal mesh component velocity drive, which does not use the
+				// target animation velocity.
+				ControlData.SkeletalAnimationVelocityMultiplier = 0;
 				FConstraintProfileProperties ProfileProperties;
 				if (!SkeletalMeshComponent->GetConstraintProfilePropertiesOrDefault(
 					ProfileProperties, ChildBoneName, ConstraintProfile))
@@ -467,12 +391,12 @@ TArray<FName> UPhysicsControlComponent::CreateControlsFromSkeletalMeshAndConstra
 
 				const FName ControlName = CreateControl(
 					SkeletalMeshComponent, ParentBoneName, SkeletalMeshComponent, ChildBoneName,
-					ControlData, FPhysicsControlTarget(), ControlSettings, 
+					ControlData, FPhysicsControlTarget(), 
 					FName(GetControlTypeName(EPhysicsControlType::ParentSpace).ToString().Append("_").Append(Set.ToString())));
 				if (!ControlName.IsNone())
 				{
 					Result.Add(ControlName);
-					Implementation->NameRecords.AddControl(
+					NameRecords.AddControl(
 						ControlName, GetControlTypeName(EPhysicsControlType::ParentSpace));
 				}
 				else
@@ -493,7 +417,6 @@ TArray<FName> UPhysicsControlComponent::CreateControlsFromSkeletalMesh(
 	const TArray<FName>&          BoneNames,
 	const EPhysicsControlType     ControlType,
 	const FPhysicsControlData     ControlData,
-	const FPhysicsControlSettings ControlSettings,
 	const FName                   Set)
 {
 	TArray<FName> Result;
@@ -509,8 +432,6 @@ TArray<FName> UPhysicsControlComponent::CreateControlsFromSkeletalMesh(
 
 	for (FName ChildBoneName : BoneNames)
 	{
-		FBodyInstance* ChildBone = SkeletalMeshComponent->GetBodyInstance(ChildBoneName);
-
 		FName ParentBoneName;
 		if (ParentMeshComponent)
 		{
@@ -523,12 +444,12 @@ TArray<FName> UPhysicsControlComponent::CreateControlsFromSkeletalMesh(
 		}
 		const FName ControlName = CreateControl(
 			ParentMeshComponent, ParentBoneName, SkeletalMeshComponent, ChildBoneName,
-			ControlData, FPhysicsControlTarget(), ControlSettings, 
+			ControlData, FPhysicsControlTarget(), 
 			FName(GetControlTypeName(ControlType).ToString().Append("_").Append(Set.ToString())));
 		if (!ControlName.IsNone())
 		{
 			Result.Add(ControlName);
-			Implementation->NameRecords.AddControl(ControlName, GetControlTypeName(ControlType));
+			NameRecords.AddControl(ControlName, GetControlTypeName(ControlType));
 		}
 		else
 		{
@@ -556,15 +477,8 @@ TArray<FName> UPhysicsControlComponent::CreateControlsFromSkeletalMeshAndConstra
 		return Result;
 	}
 
-	FPhysicsControlSettings ControlSettings;
-	// This is to match the skeletal mesh component velocity drive, which does not use the
-	// target animation velocity.
-	ControlSettings.SkeletalAnimationVelocityMultiplier = 0;
-
 	for (FName ChildBoneName : BoneNames)
 	{
-		FBodyInstance* ChildBone = SkeletalMeshComponent->GetBodyInstance(ChildBoneName);
-
 		const FName ParentBoneName = 
 			UE::PhysicsControlComponent::GetPhysicalParentBone(SkeletalMeshComponent, ChildBoneName);
 		if (ParentBoneName.IsNone())
@@ -573,6 +487,9 @@ TArray<FName> UPhysicsControlComponent::CreateControlsFromSkeletalMeshAndConstra
 		}
 
 		FPhysicsControlData ControlData;
+		// This is to match the skeletal mesh component velocity drive, which does not use the
+		// target animation velocity.
+		ControlData.SkeletalAnimationVelocityMultiplier = 0;
 		FConstraintProfileProperties ProfileProperties;
 		if (!SkeletalMeshComponent->GetConstraintProfilePropertiesOrDefault(
 			ProfileProperties, ChildBoneName, ConstraintProfile))
@@ -587,12 +504,12 @@ TArray<FName> UPhysicsControlComponent::CreateControlsFromSkeletalMeshAndConstra
 
 		const FName ControlName = CreateControl(
 			SkeletalMeshComponent, ParentBoneName, SkeletalMeshComponent, ChildBoneName,
-			ControlData, FPhysicsControlTarget(), ControlSettings, 
+			ControlData, FPhysicsControlTarget(), 
 			FName(GetControlTypeName(EPhysicsControlType::ParentSpace).ToString().Append("_").Append(Set.ToString())));
 		if (!ControlName.IsNone())
 		{
 			Result.Add(ControlName);
-			Implementation->NameRecords.AddControl(ControlName, GetControlTypeName(EPhysicsControlType::ParentSpace));
+			NameRecords.AddControl(ControlName, GetControlTypeName(EPhysicsControlType::ParentSpace));
 		}
 		else
 		{
@@ -610,7 +527,6 @@ TMap<FName, FPhysicsControlNames> UPhysicsControlComponent::CreateControlsFromLi
 	const TMap<FName, FPhysicsControlLimbBones>& LimbBones,
 	const EPhysicsControlType                    ControlType,
 	const FPhysicsControlData                    ControlData,
-	const FPhysicsControlSettings                ControlSettings,
 	UMeshComponent*                              WorldComponent,
 	FName                                        WorldBoneName,
 	FString                                      NamePrefix)
@@ -623,7 +539,7 @@ TMap<FName, FPhysicsControlNames> UPhysicsControlComponent::CreateControlsFromLi
 		const FName LimbName = LimbBoneEntry.Key;
 		const FPhysicsControlLimbBones& BonesInLimb = LimbBoneEntry.Value;
 
-		if (!BonesInLimb.SkeletalMeshComponent)
+		if (!BonesInLimb.SkeletalMeshComponent.IsValid())
 		{
 			UE_LOG(LogPhysicsControlComponent, Warning, TEXT("No Skeletal mesh in limb %s"), *LimbName.ToString());
 			continue;
@@ -636,7 +552,7 @@ TMap<FName, FPhysicsControlNames> UPhysicsControlComponent::CreateControlsFromLi
 		}
 
 		USkeletalMeshComponent* ParentSkeletalMeshComponent =
-			(ControlType == EPhysicsControlType::ParentSpace) ? BonesInLimb.SkeletalMeshComponent : nullptr;
+			(ControlType == EPhysicsControlType::ParentSpace) ? BonesInLimb.SkeletalMeshComponent.Get() : nullptr;
 
 		const int32 NumBonesInLimb = BonesInLimb.BoneNames.Num();
 
@@ -656,7 +572,6 @@ TMap<FName, FPhysicsControlNames> UPhysicsControlComponent::CreateControlsFromLi
 			}
 
 			FName ChildBoneName = BonesInLimb.BoneNames[BoneIndex];
-			FBodyInstance* ChildBone = BonesInLimb.SkeletalMeshComponent->GetBodyInstance(ChildBoneName);
 
 			FName ParentBoneName;
 			if (ParentSkeletalMeshComponent)
@@ -677,15 +592,14 @@ TMap<FName, FPhysicsControlNames> UPhysicsControlComponent::CreateControlsFromLi
 			}
 
 			const FName ControlName = CreateControl(
-				ParentMeshComponent, ParentBoneName, BonesInLimb.SkeletalMeshComponent, ChildBoneName,
-				ControlData, FPhysicsControlTarget(), ControlSettings, 
-				FName(SetName));
+				ParentMeshComponent, ParentBoneName, BonesInLimb.SkeletalMeshComponent.Get(), ChildBoneName,
+				ControlData, FPhysicsControlTarget(), FName(SetName));
 
 			if (!ControlName.IsNone())
 			{
 				LimbResult.Names.Add(ControlName);
 				AllControls.Names.Add(ControlName);
-				Implementation->NameRecords.AddControl(ControlName, GetControlTypeName(ControlType));
+				NameRecords.AddControl(ControlName, GetControlTypeName(ControlType));
 			}
 			else
 			{
@@ -711,7 +625,7 @@ TMap<FName, FPhysicsControlNames> UPhysicsControlComponent::CreateControlsFromLi
 		const FName LimbName = LimbBoneEntry.Key;
 		const FPhysicsControlLimbBones& BonesInLimb = LimbBoneEntry.Value;
 
-		USkeletalMeshComponent* SkeletalMeshComponent = BonesInLimb.SkeletalMeshComponent;
+		USkeletalMeshComponent* SkeletalMeshComponent = BonesInLimb.SkeletalMeshComponent.Get();
 		if (!SkeletalMeshComponent)
 		{
 			UE_LOG(LogPhysicsControlComponent, Warning, TEXT("No Skeletal mesh in limb %s"), *LimbName.ToString());
@@ -730,11 +644,6 @@ TMap<FName, FPhysicsControlNames> UPhysicsControlComponent::CreateControlsFromLi
 		LimbResult.Names.Reserve(NumBonesInLimb);
 		AllControls.Names.Reserve(AllControls.Names.Num() + NumBonesInLimb);
 
-		FPhysicsControlSettings ControlSettings;
-		// This is to match the skeletal mesh component velocity drive, which does not use the
-		// target animation velocity.
-		ControlSettings.SkeletalAnimationVelocityMultiplier = 0;
-
 		for (int32 BoneIndex = 0; BoneIndex != NumBonesInLimb; ++BoneIndex)
 		{
 			// Don't create the parent space control if it's the first bone in a limb that had bIncludeParentBone
@@ -744,7 +653,6 @@ TMap<FName, FPhysicsControlNames> UPhysicsControlComponent::CreateControlsFromLi
 			}
 
 			const FName ChildBoneName = BonesInLimb.BoneNames[BoneIndex];
-			FBodyInstance* ChildBone = SkeletalMeshComponent->GetBodyInstance(ChildBoneName);
 			const FName ParentBoneName = 
 				UE::PhysicsControlComponent::GetPhysicalParentBone(SkeletalMeshComponent, ChildBoneName);
 			if (ParentBoneName.IsNone())
@@ -753,6 +661,10 @@ TMap<FName, FPhysicsControlNames> UPhysicsControlComponent::CreateControlsFromLi
 			}
 
 			FPhysicsControlData ControlData;
+			// This is to match the skeletal mesh component velocity drive, which does not use the
+			// target animation velocity.
+			ControlData.SkeletalAnimationVelocityMultiplier = 0;
+
 			FConstraintProfileProperties ProfileProperties;
 			if (!SkeletalMeshComponent->GetConstraintProfilePropertiesOrDefault(
 				ProfileProperties, ChildBoneName, ConstraintProfile))
@@ -768,13 +680,13 @@ TMap<FName, FPhysicsControlNames> UPhysicsControlComponent::CreateControlsFromLi
 
 			const FName ControlName = CreateControl(
 				SkeletalMeshComponent, ParentBoneName, SkeletalMeshComponent, ChildBoneName,
-				ControlData, FPhysicsControlTarget(), ControlSettings, 
+				ControlData, FPhysicsControlTarget(), 
 				FName(GetControlTypeName(EPhysicsControlType::ParentSpace).ToString().Append("_").Append(LimbName.ToString())));
 			if (!ControlName.IsNone())
 			{
 				LimbResult.Names.Add(ControlName);
 				AllControls.Names.Add(ControlName);
-				Implementation->NameRecords.AddControl(ControlName, GetControlTypeName(EPhysicsControlType::ParentSpace));
+				NameRecords.AddControl(ControlName, GetControlTypeName(EPhysicsControlType::ParentSpace));
 			}
 			else
 			{
@@ -789,7 +701,7 @@ TMap<FName, FPhysicsControlNames> UPhysicsControlComponent::CreateControlsFromLi
 //======================================================================================================================
 bool UPhysicsControlComponent::DestroyControl(const FName Name)
 {
-	return Implementation->DestroyControl(Name);
+	return DestroyControl(Name, EDestroyBehavior::RemoveRecord);
 }
 
 //======================================================================================================================
@@ -797,7 +709,7 @@ void UPhysicsControlComponent::DestroyControls(const TArray<FName>& Names)
 {
 	for (FName Name : Names)
 	{
-		Implementation->DestroyControl(Name);
+		DestroyControl(Name, EDestroyBehavior::RemoveRecord);
 	}
 }
 
@@ -812,12 +724,17 @@ void UPhysicsControlComponent::DestroyControlsInSet(const FName SetName)
 //======================================================================================================================
 bool UPhysicsControlComponent::SetControlEnabled(const FName Name, const bool bEnable)
 {
-	FPhysicsControlRecord* Record = Implementation->FindControlRecord(Name);
+	FPhysicsControlRecord* Record = FindControlRecord(Name);
 	if (Record)
 	{
 		Record->PhysicsControl.ControlData.bEnabled = bEnable;
 		return true;
 	}
+	if (bWarnAboutInvalidNames)
+	{
+		UE_LOG(LogPhysicsControlComponent, Warning,
+			TEXT("SetControlEnabled - invalid name %s"), *Name.ToString());
+	}
 	return false;
 }
 
@@ -842,24 +759,29 @@ bool UPhysicsControlComponent::SetControlParent(
 	UMeshComponent* ParentMeshComponent,
 	const FName     ParentBoneName)
 {
-	FPhysicsControlRecord* Record = Implementation->FindControlRecord(Name);
+	FPhysicsControlRecord* Record = FindControlRecord(Name);
 	if (Record)
 	{
 		if (USkeletalMeshComponent* SkeletalMeshComponent = 
 			Cast<USkeletalMeshComponent>(Record->PhysicsControl.ParentMeshComponent))
 		{
-			Implementation->RemoveSkeletalMeshReferenceForCaching(SkeletalMeshComponent);
+			RemoveSkeletalMeshReferenceForCaching(SkeletalMeshComponent);
 		}
 
 		if (USkeletalMeshComponent* SkeletalMeshComponent = Cast<USkeletalMeshComponent>(ParentMeshComponent))
 		{
-			Implementation->AddSkeletalMeshReferenceForCaching(SkeletalMeshComponent);
+			AddSkeletalMeshReferenceForCaching(SkeletalMeshComponent);
 		}
 
 		Record->PhysicsControl.ParentMeshComponent = ParentMeshComponent;
 		Record->PhysicsControl.ParentBoneName = ParentBoneName;
 		return Record->InitConstraint(this, Name);
 	}
+	if (bWarnAboutInvalidNames)
+	{
+		UE_LOG(LogPhysicsControlComponent, Warning,
+			TEXT("SetControlParent - invalid name %s"), *Name.ToString());
+	}
 	return false;
 }
 
@@ -884,27 +806,30 @@ void UPhysicsControlComponent::SetControlParentsInSet(
 	SetControlParents(GetControlNamesInSet(SetName), ParentMeshComponent, ParentBoneName);
 }
 
-
-
 //======================================================================================================================
 // Note - some params passed by value to allow them to be unconnected and enable inline BP editing
 bool UPhysicsControlComponent::SetControlData(
 	const FName               Name, 
 	const FPhysicsControlData ControlData)
 {
-	FPhysicsControlRecord* Record = Implementation->FindControlRecord(Name);
+	FPhysicsControlRecord* Record = FindControlRecord(Name);
 	if (Record)
 	{
 		Record->PhysicsControl.ControlData = ControlData;
 		return true;
 	}
+	if (bWarnAboutInvalidNames)
+	{
+		UE_LOG(LogPhysicsControlComponent, Warning,
+			TEXT("SetControlData - invalid name %s"), *Name.ToString());
+	}
 	return false;
 }
 
 //======================================================================================================================
 void UPhysicsControlComponent::SetControlDatas(
-	const TArray<FName>&       Names, 
-	const FPhysicsControlData  ControlData)
+	const TArray<FName>&      Names, 
+	const FPhysicsControlData ControlData)
 {
 	for (FName Name : Names)
 	{
@@ -914,12 +839,51 @@ void UPhysicsControlComponent::SetControlDatas(
 
 //======================================================================================================================
 void UPhysicsControlComponent::SetControlDatasInSet(
-	const FName                SetName,
-	const FPhysicsControlData  ControlData)
+	const FName               SetName,
+	const FPhysicsControlData ControlData)
 {
 	SetControlDatas(GetControlNamesInSet(SetName), ControlData);
 }
 
+//======================================================================================================================
+// Note - some params passed by value to allow them to be unconnected and enable inline BP editing
+bool UPhysicsControlComponent::SetControlSparseData(
+	const FName                     Name, 
+	const FPhysicsControlSparseData ControlData)
+{
+	FPhysicsControlRecord* Record = FindControlRecord(Name);
+	if (Record)
+	{
+		Record->PhysicsControl.ControlData.UpdateFromSparseData(ControlData);
+		return true;
+	}
+	if (bWarnAboutInvalidNames)
+	{
+		UE_LOG(LogPhysicsControlComponent, Warning,
+			TEXT("SetControlData - invalid name %s"), *Name.ToString());
+	}
+	return false;
+}
+
+//======================================================================================================================
+void UPhysicsControlComponent::SetControlSparseDatas(
+	const TArray<FName>&            Names, 
+	const FPhysicsControlSparseData ControlData)
+{
+	for (FName Name : Names)
+	{
+		SetControlSparseData(Name, ControlData);
+	}
+}
+
+//======================================================================================================================
+void UPhysicsControlComponent::SetControlSparseDatasInSet(
+	const FName                     SetName,
+	const FPhysicsControlSparseData ControlData)
+{
+	SetControlSparseDatas(GetControlNamesInSet(SetName), ControlData);
+}
+
 //======================================================================================================================
 // Note - some params passed by value to allow them to be unconnected and enable inline BP editing
 bool UPhysicsControlComponent::SetControlMultiplier(
@@ -927,7 +891,7 @@ bool UPhysicsControlComponent::SetControlMultiplier(
 	const FPhysicsControlMultiplier  ControlMultiplier, 
 	const bool                       bEnableControl)
 {
-	FPhysicsControlRecord* Record = Implementation->FindControlRecord(Name);
+	FPhysicsControlRecord* Record = FindControlRecord(Name);
 	if (Record)
 	{
 		Record->PhysicsControl.ControlMultiplier = ControlMultiplier;
@@ -937,6 +901,11 @@ bool UPhysicsControlComponent::SetControlMultiplier(
 		}
 		return true;
 	}
+	if (bWarnAboutInvalidNames)
+	{
+		UE_LOG(LogPhysicsControlComponent, Warning,
+			TEXT("SetControlMultiplier - invalid name %s"), *Name.ToString());
+	}
 	return false;
 }
 
@@ -961,6 +930,52 @@ void UPhysicsControlComponent::SetControlMultipliersInSet(
 	SetControlMultipliers(GetControlNamesInSet(SetName), ControlMultiplier, bEnableControl);
 }
 
+//======================================================================================================================
+// Note - some params passed by value to allow them to be unconnected and enable inline BP editing
+bool UPhysicsControlComponent::SetControlSparseMultiplier(
+	const FName                            Name,
+	const FPhysicsControlSparseMultiplier  ControlMultiplier,
+	const bool                             bEnableControl)
+{
+	FPhysicsControlRecord* Record = FindControlRecord(Name);
+	if (Record)
+	{
+		Record->PhysicsControl.ControlMultiplier.UpdateFromSparseData(ControlMultiplier);
+		if (bEnableControl)
+		{
+			Record->PhysicsControl.ControlData.bEnabled = true;
+		}
+		return true;
+	}
+	if (bWarnAboutInvalidNames)
+	{
+		UE_LOG(LogPhysicsControlComponent, Warning,
+			TEXT("SetControlSparseMultiplier - invalid name %s"), *Name.ToString());
+	}
+	return false;
+}
+
+//======================================================================================================================
+void UPhysicsControlComponent::SetControlSparseMultipliers(
+	const TArray<FName>&                  Names,
+	const FPhysicsControlSparseMultiplier ControlMultiplier,
+	const bool                            bEnableControl)
+{
+	for (FName Name : Names)
+	{
+		SetControlSparseMultiplier(Name, ControlMultiplier, bEnableControl);
+	}
+}
+
+//======================================================================================================================
+void UPhysicsControlComponent::SetControlSparseMultipliersInSet(
+	const FName                           SetName,
+	const FPhysicsControlSparseMultiplier ControlMultiplier,
+	const bool                            bEnableControl)
+{
+	SetControlSparseMultipliers(GetControlNamesInSet(SetName), ControlMultiplier, bEnableControl);
+}
+
 //======================================================================================================================
 bool UPhysicsControlComponent::SetControlLinearData(
 	const FName Name, 
@@ -970,7 +985,7 @@ bool UPhysicsControlComponent::SetControlLinearData(
 	const float MaxForce, 
 	const bool  bEnableControl)
 {
-	FPhysicsControlRecord* Record = Implementation->FindControlRecord(Name);
+	FPhysicsControlRecord* Record = FindControlRecord(Name);
 	if (Record)
 	{
 		Record->PhysicsControl.ControlData.LinearStrength = Strength;
@@ -983,6 +998,11 @@ bool UPhysicsControlComponent::SetControlLinearData(
 		}
 		return true;
 	}
+	if (bWarnAboutInvalidNames)
+	{
+		UE_LOG(LogPhysicsControlComponent, Warning,
+			TEXT("SetControlLinearData - invalid name %s"), *Name.ToString());
+	}
 	return false;
 }
 
@@ -995,7 +1015,7 @@ bool UPhysicsControlComponent::SetControlAngularData(
 	const float MaxTorque, 
 	const bool  bEnableControl)
 {
-	FPhysicsControlRecord* Record = Implementation->FindControlRecord(Name);
+	FPhysicsControlRecord* Record = FindControlRecord(Name);
 	if (Record)
 	{
 		Record->PhysicsControl.ControlData.AngularStrength = Strength;
@@ -1008,31 +1028,47 @@ bool UPhysicsControlComponent::SetControlAngularData(
 		}
 		return true;
 	}
+	if (bWarnAboutInvalidNames)
+	{
+		UE_LOG(LogPhysicsControlComponent, Warning,
+			TEXT("SetControlAngularData - invalid name %s"), *Name.ToString());
+	}
 	return false;
 }
 
 //======================================================================================================================
 bool UPhysicsControlComponent::SetControlPoint(const FName Name, const FVector Position)
 {
-	FPhysicsControlRecord* Record = Implementation->FindControlRecord(Name);
+	FPhysicsControlRecord* Record = FindControlRecord(Name);
 	if (Record)
 	{
-		Record->PhysicsControl.ControlSettings.ControlPoint = Position;
+		Record->PhysicsControl.ControlData.bUseCustomControlPoint = true;
+		Record->PhysicsControl.ControlData.CustomControlPoint = Position;
 		Record->UpdateConstraintControlPoint();
 		return true;
 	}
+	if (bWarnAboutInvalidNames)
+	{
+		UE_LOG(LogPhysicsControlComponent, Warning,
+			TEXT("SetControlPoint - invalid name %s"), *Name.ToString());
+	}
 	return false;
 }
 
 //======================================================================================================================
 bool UPhysicsControlComponent::ResetControlPoint(const FName Name)
 {
-	FPhysicsControlRecord* Record = Implementation->FindControlRecord(Name);
+	FPhysicsControlRecord* Record = FindControlRecord(Name);
 	if (Record)
 	{
 		Record->ResetControlPoint();
 		return true;
 	}
+	if (bWarnAboutInvalidNames)
+	{
+		UE_LOG(LogPhysicsControlComponent, Warning,
+			TEXT("ResetControlPoint - invalid name %s"), *Name.ToString());
+	}
 	return false;
 }
 
@@ -1042,7 +1078,7 @@ bool UPhysicsControlComponent::SetControlTarget(
 	const FPhysicsControlTarget ControlTarget, 
 	const bool                  bEnableControl)
 {
-	FPhysicsControlRecord* Record = Implementation->FindControlRecord(Name);
+	FPhysicsControlRecord* Record = FindControlRecord(Name);
 	if (Record)
 	{
 		Record->PhysicsControl.ControlTarget = ControlTarget;
@@ -1052,6 +1088,11 @@ bool UPhysicsControlComponent::SetControlTarget(
 		}
 		return true;
 	}
+	if (bWarnAboutInvalidNames)
+	{
+		UE_LOG(LogPhysicsControlComponent, Warning,
+			TEXT("SetControlTarget - invalid name %s"), *Name.ToString());
+	}
 	return false;
 }
 
@@ -1085,13 +1126,19 @@ bool UPhysicsControlComponent::SetControlTargetPositionAndOrientation(
 	const bool     bEnableControl, 
 	const bool     bApplyControlPointToTarget)
 {
-	FPhysicsControlRecord* Record = Implementation->FindControlRecord(Name);
+	FPhysicsControlRecord* Record = FindControlRecord(Name);
 	if (Record)
 	{
 		SetControlTargetPosition(
 			Name, Position, VelocityDeltaTime, bEnableControl, bApplyControlPointToTarget);
 		SetControlTargetOrientation(
 			Name, Orientation, VelocityDeltaTime, bEnableControl, bApplyControlPointToTarget);
+		return true;
+	}
+	if (bWarnAboutInvalidNames)
+	{
+		UE_LOG(LogPhysicsControlComponent, Warning,
+			TEXT("SetControlTargetPositionAndOrientation - invalid name %s"), *Name.ToString());
 	}
 	return false;
 }
@@ -1133,7 +1180,7 @@ bool UPhysicsControlComponent::SetControlTargetPosition(
 	const bool    bEnableControl, 
 	const bool    bApplyControlPointToTarget)
 {
-	FPhysicsControlRecord* Record = Implementation->FindControlRecord(Name);
+	FPhysicsControlRecord* Record = FindControlRecord(Name);
 	if (Record)
 	{
 		if (VelocityDeltaTime != 0)
@@ -1153,6 +1200,11 @@ bool UPhysicsControlComponent::SetControlTargetPosition(
 		}
 		return true;
 	}
+	if (bWarnAboutInvalidNames)
+	{
+		UE_LOG(LogPhysicsControlComponent, Warning,
+			TEXT("SetControlTargetPosition - invalid name %s"), *Name.ToString());
+	}
 	return false;
 }
 
@@ -1191,7 +1243,7 @@ bool UPhysicsControlComponent::SetControlTargetOrientation(
 	const bool     bEnableControl, 
 	const bool     bApplyControlPointToTarget)
 {
-	FPhysicsControlRecord* Record = Implementation->FindControlRecord(Name);
+	FPhysicsControlRecord* Record = FindControlRecord(Name);
 	if (Record)
 	{
 		if (AngularVelocityDeltaTime != 0)
@@ -1216,6 +1268,11 @@ bool UPhysicsControlComponent::SetControlTargetOrientation(
 		}
 		return true;
 	}
+	if (bWarnAboutInvalidNames)
+	{
+		UE_LOG(LogPhysicsControlComponent, Warning,
+			TEXT("SetControlTargetOrientation - invalid name %s"), *Name.ToString());
+	}
 	return false;
 }
 
@@ -1258,6 +1315,8 @@ bool UPhysicsControlComponent::SetControlTargetPositionsFromArray(
 	int32 NumPositions = Positions.Num();
 	if (NumControlNames != NumPositions)
 	{
+		UE_LOG(LogPhysicsControlComponent, Warning,
+			TEXT("SetControlTargetPositionsFromArray - names and positions arrays sizes do not match"));
 		return false;
 	}
 	for (int32 Index = 0 ; Index != NumControlNames ; ++Index)
@@ -1280,6 +1339,8 @@ bool UPhysicsControlComponent::SetControlTargetOrientationsFromArray(
 	int32 NumOrientations = Orientations.Num();
 	if (NumControlNames != NumOrientations)
 	{
+		UE_LOG(LogPhysicsControlComponent, Warning,
+			TEXT("SetControlTargetOrientationsFromArray - names and orientations arrays sizes do not match"));
 		return false;
 	}
 	for (int32 Index = 0 ; Index != NumControlNames ; ++Index)
@@ -1304,6 +1365,8 @@ bool UPhysicsControlComponent::SetControlTargetPositionsAndOrientationsFromArray
 	int32 NumOrientations = Orientations.Num();
 	if (NumControlNames != NumPositions || NumControlNames != NumOrientations)
 	{
+		UE_LOG(LogPhysicsControlComponent, Warning,
+			TEXT("SetControlTargetPositionsAndOrientationsFromArray - names and positions/orientation arrays sizes do not match"));
 		return false;
 	}
 	for (int32 Index = 0; Index != NumControlNames; ++Index)
@@ -1324,7 +1387,7 @@ bool UPhysicsControlComponent::SetControlTargetPoses(
 	const float    VelocityDeltaTime, 
 	const bool     bEnableControl)
 {
-	FPhysicsControlRecord* Record = Implementation->FindControlRecord(Name);
+	FPhysicsControlRecord* Record = FindControlRecord(Name);
 	if (Record)
 	{
 		FTransform ParentTM(ParentOrientation, ParentPosition, FVector::One());
@@ -1360,6 +1423,11 @@ bool UPhysicsControlComponent::SetControlTargetPoses(
 		}
 		return true;
 	}
+	if (bWarnAboutInvalidNames)
+	{
+		UE_LOG(LogPhysicsControlComponent, Warning,
+			TEXT("SetControlTargetPoses - invalid name %s"), *Name.ToString());
+	}
 	return false;
 }
 
@@ -1369,13 +1437,18 @@ bool UPhysicsControlComponent::SetControlUseSkeletalAnimation(
 	const bool  bUseSkeletalAnimation,
 	const float SkeletalAnimationVelocityMultiplier)
 {
-	FPhysicsControl* PhysicsControl = Implementation->FindControl(Name);
+	FPhysicsControl* PhysicsControl = FindControl(Name);
 	if (PhysicsControl)
 	{
-		PhysicsControl->ControlSettings.bUseSkeletalAnimation = bUseSkeletalAnimation;
-		PhysicsControl->ControlSettings.SkeletalAnimationVelocityMultiplier = SkeletalAnimationVelocityMultiplier;
+		PhysicsControl->ControlData.bUseSkeletalAnimation = bUseSkeletalAnimation;
+		PhysicsControl->ControlData.SkeletalAnimationVelocityMultiplier = SkeletalAnimationVelocityMultiplier;
 		return true;
 	}
+	if (bWarnAboutInvalidNames)
+	{
+		UE_LOG(LogPhysicsControlComponent, Warning,
+			TEXT("SetControlUseSkeletalAnimation - invalid name %s"), *Name.ToString());
+	}
 	return false;
 }
 
@@ -1404,12 +1477,17 @@ void UPhysicsControlComponent::SetControlsInSetUseSkeletalAnimation(
 //======================================================================================================================
 bool UPhysicsControlComponent::SetControlDisableCollision(const FName Name, const bool bDisableCollision)
 {
-	FPhysicsControl* PhysicsControl = Implementation->FindControl(Name);
+	FPhysicsControl* PhysicsControl = FindControl(Name);
 	if (PhysicsControl)
 	{
-		PhysicsControl->ControlSettings.bDisableCollision = bDisableCollision;
+		PhysicsControl->ControlData.bDisableCollision = bDisableCollision;
 		return true;
 	}
+	if (bWarnAboutInvalidNames)
+	{
+		UE_LOG(LogPhysicsControlComponent, Warning,
+			TEXT("SetControlDisableCollision - invalid name %s"), *Name.ToString());
+	}
 	return false;
 }
 
@@ -1431,66 +1509,79 @@ void UPhysicsControlComponent::SetControlsInSetDisableCollision(const FName SetN
 //======================================================================================================================
 bool UPhysicsControlComponent::GetControlData(const FName Name, FPhysicsControlData& ControlData) const
 {
-	FPhysicsControl* PhysicsControl = Implementation->FindControl(Name);
+	const FPhysicsControl* PhysicsControl = FindControl(Name);
 	if (PhysicsControl)
 	{
 		ControlData = PhysicsControl->ControlData;
 		return true;
 	}
+	if (bWarnAboutInvalidNames)
+	{
+		UE_LOG(LogPhysicsControlComponent, Warning,
+			TEXT("GetControlData - invalid name %s"), *Name.ToString());
+	}
 	return false;
 }
 
 //======================================================================================================================
 bool UPhysicsControlComponent::GetControlMultiplier(const FName Name, FPhysicsControlMultiplier& ControlMultiplier) const
 {
-	FPhysicsControl* PhysicsControl = Implementation->FindControl(Name);
+	const FPhysicsControl* PhysicsControl = FindControl(Name);
 	if (PhysicsControl)
 	{
 		ControlMultiplier = PhysicsControl->ControlMultiplier;
 		return true;
 	}
+	if (bWarnAboutInvalidNames)
+	{
+		UE_LOG(LogPhysicsControlComponent, Warning,
+			TEXT("GetControlMultiplier - invalid name %s"), *Name.ToString());
+	}
 	return false;
 }
 
 //======================================================================================================================
 bool UPhysicsControlComponent::GetControlTarget(const FName Name, FPhysicsControlTarget& ControlTarget) const
 {
-	FPhysicsControl* PhysicsControl = Implementation->FindControl(Name);
+	const FPhysicsControl* PhysicsControl = FindControl(Name);
 	if (PhysicsControl)
 	{
 		ControlTarget = PhysicsControl->ControlTarget;
 		return true;
 	}
+	if (bWarnAboutInvalidNames)
+	{
+		UE_LOG(LogPhysicsControlComponent, Warning,
+			TEXT("GetControlTarget - invalid name %s"), *Name.ToString());
+	}
 	return false;
 }
 
 //======================================================================================================================
 bool UPhysicsControlComponent::GetControlEnabled(const FName Name) const
 {
-	FPhysicsControlRecord* Record = Implementation->FindControlRecord(Name);
+	const FPhysicsControlRecord* Record = FindControlRecord(Name);
 	if (Record)
 	{
 		return Record->Enabled();
 	}
+	if (bWarnAboutInvalidNames)
+	{
+		UE_LOG(LogPhysicsControlComponent, Warning,
+			TEXT("GetControlEnabled - invalid name %s"), *Name.ToString());
+	}
 	return false;
 }
 
 //======================================================================================================================
 FName UPhysicsControlComponent::CreateBodyModifier(
-	UMeshComponent*               MeshComponent,
-	const FName                   BoneName,
-	const FName                   Set,
-	const EPhysicsMovementType    MovementType,
-	const ECollisionEnabled::Type CollisionType,
-	const float                   GravityMultiplier,
-	const float                   PhysicsBlendWeight,
-	const bool                    bUseSkeletalAnimation,
-	const bool                    bUpdateKinematicFromSimulation)
-{
-	const FName Name = Implementation->GetUniqueBodyModifierName(BoneName);
-	if (CreateNamedBodyModifier(
-		Name, MeshComponent, BoneName, Set, MovementType, CollisionType, 
-		GravityMultiplier, PhysicsBlendWeight, bUseSkeletalAnimation, bUpdateKinematicFromSimulation))
+	UMeshComponent*                   MeshComponent,
+	const FName                       BoneName,
+	const FName                       Set,
+	const FPhysicsControlModifierData BodyModifierData)
+{
+	const FName Name = GetUniqueBodyModifierName(BoneName);
+	if (CreateNamedBodyModifier(Name, MeshComponent, BoneName, Set, BodyModifierData))
 	{
 		return Name;
 	}
@@ -1499,19 +1590,16 @@ FName UPhysicsControlComponent::CreateBodyModifier(
 
 //======================================================================================================================
 bool UPhysicsControlComponent::CreateNamedBodyModifier(
-	const FName                   Name,
-	UMeshComponent*               MeshComponent,
-	const FName                   BoneName,
-	const FName                   Set,
-	const EPhysicsMovementType    MovementType,
-	const ECollisionEnabled::Type CollisionType,
-	const float                   GravityMultiplier,
-	const float                   PhysicsBlendWeight,
-	const bool                    bUseSkeletalAnimation,
-	const bool                    bUpdateKinematicFromSimulation)
+	const FName                       Name,
+	UMeshComponent*                   MeshComponent,
+	const FName                       BoneName,
+	const FName                       Set,
+	const FPhysicsControlModifierData BodyModifierData)
 {
-	if (Implementation->FindBodyModifier(Name))
+	if (FindBodyModifier(Name))
 	{
+		UE_LOG(LogPhysicsControlComponent, Warning,
+			TEXT("CreateNamedBodyModifier - modifier with name %s already exists"), *Name.ToString());
 		return false;
 	}
 
@@ -1522,56 +1610,47 @@ bool UPhysicsControlComponent::CreateNamedBodyModifier(
 		return false;
 	}
 
-	FPhysicsBodyModifier& Modifier = Implementation->PhysicsBodyModifiers.Add(
-		Name, FPhysicsBodyModifier(
-			MeshComponent, BoneName, MovementType, CollisionType, 
-			GravityMultiplier, PhysicsBlendWeight, bUseSkeletalAnimation, bUpdateKinematicFromSimulation));
+	FPhysicsBodyModifier& Modifier = PhysicsBodyModifiers.Add(
+		Name, FPhysicsBodyModifier(MeshComponent, BoneName, BodyModifierData));
 
 	USkeletalMeshComponent* SkeletalMeshComponent = Cast<USkeletalMeshComponent>(MeshComponent);
 	if (SkeletalMeshComponent)
 	{
-		Implementation->AddSkeletalMeshReferenceForCaching(SkeletalMeshComponent);
-		Implementation->AddSkeletalMeshReferenceForModifier(SkeletalMeshComponent);
+		AddSkeletalMeshReferenceForCaching(SkeletalMeshComponent);
+		AddSkeletalMeshReferenceForModifier(SkeletalMeshComponent);
 	}
 
-	Implementation->NameRecords.AddBodyModifier(Name, Set);
+	NameRecords.AddBodyModifier(Name, Set);
 
 	return true;
 }
 
 //======================================================================================================================
 TArray<FName> UPhysicsControlComponent::CreateBodyModifiersFromSkeletalMeshBelow(
-	USkeletalMeshComponent*       SkeletalMeshComponent,
-	const FName                   BoneName,
-	const bool                    bIncludeSelf,
-	const FName                   Set,
-	const EPhysicsMovementType    MovementType,
-	const ECollisionEnabled::Type CollisionType,
-	const float                   GravityMultiplier,
-	const float                   PhysicsBlendWeight,
-	const bool                    bUseSkeletalAnimation,
-	const bool                    bUpdateKinematicFromSimulation)
+	USkeletalMeshComponent*           SkeletalMeshComponent,
+	const FName                       BoneName,
+	const bool                        bIncludeSelf,
+	const FName                       Set,
+	const FPhysicsControlModifierData BodyModifierData)
 {
 	TArray<FName> Result;
 	UPhysicsAsset* PhysicsAsset = SkeletalMeshComponent ? SkeletalMeshComponent->GetPhysicsAsset() : nullptr;
 	if (!PhysicsAsset)
 	{
+		UE_LOG(LogPhysicsControlComponent, Warning,
+			TEXT("CreateBodyModifiersFromSkeletalMeshBelow - No physics asset available"));
 		return Result;
 	}
 
 	SkeletalMeshComponent->ForEachBodyBelow(
 		BoneName, bIncludeSelf, /*bSkipCustomType=*/false,
-		[
-			this, PhysicsAsset, SkeletalMeshComponent, Set, MovementType, CollisionType,
-			GravityMultiplier, PhysicsBlendWeight, bUseSkeletalAnimation, bUpdateKinematicFromSimulation, &Result
-		](const FBodyInstance* BI)
+		[this, PhysicsAsset, SkeletalMeshComponent, Set, BodyModifierData, &Result](const FBodyInstance* BI)
 		{
 			if (USkeletalBodySetup* BodySetup = Cast<USkeletalBodySetup>(BI->BodySetup.Get()))
 			{
 				const FName BoneName = PhysicsAsset->SkeletalBodySetups[BI->InstanceBodyIndex]->BoneName;
 				const FName BodyModifierName = CreateBodyModifier(
-					SkeletalMeshComponent, BoneName, Set, MovementType, CollisionType, 
-					GravityMultiplier, PhysicsBlendWeight, bUseSkeletalAnimation, bUpdateKinematicFromSimulation);
+					SkeletalMeshComponent, BoneName, Set, BodyModifierData);
 				Result.Add(BodyModifierName);
 			}
 		});
@@ -1583,12 +1662,7 @@ TArray<FName> UPhysicsControlComponent::CreateBodyModifiersFromSkeletalMeshBelow
 TMap<FName, FPhysicsControlNames> UPhysicsControlComponent::CreateBodyModifiersFromLimbBones(
 	FPhysicsControlNames&                        AllBodyModifiers,
 	const TMap<FName, FPhysicsControlLimbBones>& LimbBones,
-	const EPhysicsMovementType                   MovementType,
-	const ECollisionEnabled::Type                CollisionType,
-	const float                                  GravityMultiplier,
-	const float                                  PhysicsBlendWeight,
-	const bool                                   bUseSkeletalAnimation,
-	const bool                                   bUpdateKinematicFromSimulation)
+	const FPhysicsControlModifierData            BodyModifierData)
 {
 	TMap<FName, FPhysicsControlNames> Result;
 	Result.Reserve(LimbBones.Num());
@@ -1598,7 +1672,7 @@ TMap<FName, FPhysicsControlNames> UPhysicsControlComponent::CreateBodyModifiersF
 		const FName LimbName = LimbBoneEntry.Key;
 		const FPhysicsControlLimbBones& BonesInLimb = LimbBoneEntry.Value;
 
-		if (!BonesInLimb.SkeletalMeshComponent)
+		if (!BonesInLimb.SkeletalMeshComponent.Get())
 		{
 			UE_LOG(LogPhysicsControlComponent, Warning, TEXT("No Skeletal mesh in limb %s"), *LimbName.ToString());
 			continue;
@@ -1613,8 +1687,7 @@ TMap<FName, FPhysicsControlNames> UPhysicsControlComponent::CreateBodyModifiersF
 		for (const FName BoneName : BonesInLimb.BoneNames)
 		{
 			const FName BodyModifierName = CreateBodyModifier(
-				BonesInLimb.SkeletalMeshComponent, BoneName, LimbName, MovementType, CollisionType, 
-				GravityMultiplier, PhysicsBlendWeight, bUseSkeletalAnimation, bUpdateKinematicFromSimulation);
+				BonesInLimb.SkeletalMeshComponent.Get(), BoneName, LimbName, BodyModifierData);
 			if (!BodyModifierName.IsNone())
 			{
 				LimbResult.Names.Add(BodyModifierName);
@@ -1633,7 +1706,7 @@ TMap<FName, FPhysicsControlNames> UPhysicsControlComponent::CreateBodyModifiersF
 //======================================================================================================================
 bool UPhysicsControlComponent::DestroyBodyModifier(const FName Name)
 {
-	return Implementation->DestroyBodyModifier(Name);
+	return DestroyBodyModifier(Name, EDestroyBehavior::RemoveRecord);
 }
 
 //======================================================================================================================
@@ -1641,7 +1714,7 @@ void UPhysicsControlComponent::DestroyBodyModifiers(const TArray<FName>& Names)
 {
 	for (FName Name : Names)
 	{
-		Implementation->DestroyBodyModifier(Name);
+		DestroyBodyModifier(Name, EDestroyBehavior::RemoveRecord);
 	}
 }
 
@@ -1652,6 +1725,85 @@ void UPhysicsControlComponent::DestroyBodyModifiersInSet(const FName SetName)
 	DestroyBodyModifiers(Names);
 }
 
+//======================================================================================================================
+// Note - some params passed by value to allow them to be unconnected and enable inline BP editing
+bool UPhysicsControlComponent::SetBodyModifierData(
+	const FName                       Name,
+	const FPhysicsControlModifierData ModifierData)
+{
+	FPhysicsBodyModifier* Modifier = FindBodyModifier(Name);
+	if (Modifier)
+	{
+		Modifier->BodyModifierData = ModifierData;
+		return true;
+	}
+	if (bWarnAboutInvalidNames)
+	{
+		UE_LOG(LogPhysicsControlComponent, Warning,
+			TEXT("SetBodyModifierData - invalid name %s"), *Name.ToString());
+	}
+	return false;
+}
+
+//======================================================================================================================
+void UPhysicsControlComponent::SetBodyModifierDatas(
+	const TArray<FName>&              Names,
+	const FPhysicsControlModifierData ModifierData)
+{
+	for (FName Name : Names)
+	{
+		SetBodyModifierData(Name, ModifierData);
+	}
+}
+
+//======================================================================================================================
+void UPhysicsControlComponent::SetBodyModifierDatasInSet(
+	const FName                       SetName,
+	const FPhysicsControlModifierData ModifierData)
+{
+	SetBodyModifierDatas(GetControlNamesInSet(SetName), ModifierData);
+}
+
+//======================================================================================================================
+// Note - some params passed by value to allow them to be unconnected and enable inline BP editing
+bool UPhysicsControlComponent::SetBodyModifierSparseData(
+	const FName                             Name,
+	const FPhysicsControlModifierSparseData ModifierData)
+{
+	FPhysicsBodyModifier* Modifier = FindBodyModifier(Name);
+	if (Modifier)
+	{
+		Modifier->BodyModifierData.UpdateFromSparseData(ModifierData);
+		return true;
+	}
+	if (bWarnAboutInvalidNames)
+	{
+		UE_LOG(LogPhysicsControlComponent, Warning,
+			TEXT("SetBodyModifierData - invalid name %s"), *Name.ToString());
+	}
+	return false;
+}
+
+//======================================================================================================================
+void UPhysicsControlComponent::SetBodyModifierSparseDatas(
+	const TArray<FName>&                    Names,
+	const FPhysicsControlModifierSparseData ModifierData)
+{
+	for (FName Name : Names)
+	{
+		SetBodyModifierSparseData(Name, ModifierData);
+	}
+}
+
+//======================================================================================================================
+void UPhysicsControlComponent::SetBodyModifierSparseDatasInSet(
+	const FName                             SetName,
+	const FPhysicsControlModifierSparseData ModifierData)
+{
+	SetBodyModifierSparseDatas(GetControlNamesInSet(SetName), ModifierData);
+}
+
+
 //======================================================================================================================
 bool UPhysicsControlComponent::SetBodyModifierKinematicTarget(
 	const FName    Name,
@@ -1659,17 +1811,22 @@ bool UPhysicsControlComponent::SetBodyModifierKinematicTarget(
 	const FRotator KinematicTargetOrienation,
 	const bool     bMakeKinematic)
 {
-	FPhysicsBodyModifier* PhysicsBodyModifier = Implementation->FindBodyModifier(Name);
+	FPhysicsBodyModifier* PhysicsBodyModifier = FindBodyModifier(Name);
 	if (PhysicsBodyModifier)
 	{
 		PhysicsBodyModifier->KinematicTargetPosition = KinematicTargetPosition;
 		PhysicsBodyModifier->KinematicTargetOrientation = KinematicTargetOrienation.Quaternion();
 		if (bMakeKinematic)
 		{
-			PhysicsBodyModifier->MovementType = EPhysicsMovementType::Kinematic;
+			PhysicsBodyModifier->BodyModifierData.MovementType = EPhysicsMovementType::Kinematic;
 		}
 		return true;
 	}
+	if (bWarnAboutInvalidNames)
+	{
+		UE_LOG(LogPhysicsControlComponent, Warning,
+			TEXT("SetBodyModifierKinematicTarget - invalid name %s"), *Name.ToString());
+	}
 	return false;
 }
 
@@ -1678,12 +1835,17 @@ bool UPhysicsControlComponent::SetBodyModifierMovementType(
 	const FName                Name,
 	const EPhysicsMovementType MovementType)
 {
-	FPhysicsBodyModifier* PhysicsBodyModifier = Implementation->FindBodyModifier(Name);
+	FPhysicsBodyModifier* PhysicsBodyModifier = FindBodyModifier(Name);
 	if (PhysicsBodyModifier)
 	{
-		PhysicsBodyModifier->MovementType = MovementType;
+		PhysicsBodyModifier->BodyModifierData.MovementType = MovementType;
 		return true;
 	}
+	if (bWarnAboutInvalidNames)
+	{
+		UE_LOG(LogPhysicsControlComponent, Warning,
+			TEXT("SetBodyModifierMovementType - invalid name %s"), *Name.ToString());
+	}
 	return false;
 }
 
@@ -1711,12 +1873,17 @@ bool UPhysicsControlComponent::SetBodyModifierCollisionType(
 	const FName                   Name,
 	const ECollisionEnabled::Type CollisionType)
 {
-	FPhysicsBodyModifier* PhysicsBodyModifier = Implementation->FindBodyModifier(Name);
+	FPhysicsBodyModifier* PhysicsBodyModifier = FindBodyModifier(Name);
 	if (PhysicsBodyModifier)
 	{
-		PhysicsBodyModifier->CollisionType = CollisionType;
+		PhysicsBodyModifier->BodyModifierData.CollisionType = CollisionType;
 		return true;
 	}
+	if (bWarnAboutInvalidNames)
+	{
+		UE_LOG(LogPhysicsControlComponent, Warning,
+			TEXT("SetBodyModifierCollisionType - invalid name %s"), *Name.ToString());
+	}
 	return false;
 }
 
@@ -1744,12 +1911,17 @@ bool UPhysicsControlComponent::SetBodyModifierGravityMultiplier(
 	const FName Name,
 	const float GravityMultiplier)
 {
-	FPhysicsBodyModifier* PhysicsBodyModifier = Implementation->FindBodyModifier(Name);
+	FPhysicsBodyModifier* PhysicsBodyModifier = FindBodyModifier(Name);
 	if (PhysicsBodyModifier)
 	{
-		PhysicsBodyModifier->GravityMultiplier = GravityMultiplier;
+		PhysicsBodyModifier->BodyModifierData.GravityMultiplier = GravityMultiplier;
 		return true;
 	}
+	if (bWarnAboutInvalidNames)
+	{
+		UE_LOG(LogPhysicsControlComponent, Warning,
+			TEXT("SetBodyModifierGravityMultiplier - invalid name %s"), *Name.ToString());
+	}
 	return false;
 }
 
@@ -1777,12 +1949,17 @@ bool UPhysicsControlComponent::SetBodyModifierPhysicsBlendWeight(
 	const FName Name,
 	const float PhysicsBlendWeight)
 {
-	FPhysicsBodyModifier* PhysicsBodyModifier = Implementation->FindBodyModifier(Name);
+	FPhysicsBodyModifier* PhysicsBodyModifier = FindBodyModifier(Name);
 	if (PhysicsBodyModifier)
 	{
-		PhysicsBodyModifier->PhysicsBlendWeight = PhysicsBlendWeight;
+		PhysicsBodyModifier->BodyModifierData.PhysicsBlendWeight = PhysicsBlendWeight;
 		return true;
 	}
+	if (bWarnAboutInvalidNames)
+	{
+		UE_LOG(LogPhysicsControlComponent, Warning,
+			TEXT("SetBodyModifierPhysicsBlendWeight - invalid name %s"), *Name.ToString());
+	}
 	return false;
 }
 
@@ -1810,12 +1987,17 @@ bool UPhysicsControlComponent::SetBodyModifierUseSkeletalAnimation(
 	const FName Name,
 	const bool  bUseSkeletalAnimation)
 {
-	FPhysicsBodyModifier* PhysicsBodyModifier = Implementation->FindBodyModifier(Name);
+	FPhysicsBodyModifier* PhysicsBodyModifier = FindBodyModifier(Name);
 	if (PhysicsBodyModifier)
 	{
-		PhysicsBodyModifier->bUseSkeletalAnimation = bUseSkeletalAnimation;
+		PhysicsBodyModifier->BodyModifierData.bUseSkeletalAnimation = bUseSkeletalAnimation;
 		return true;
 	}
+	if (bWarnAboutInvalidNames)
+	{
+		UE_LOG(LogPhysicsControlComponent, Warning,
+			TEXT("SetBodyModifierUseSkeletalAnimation - invalid name %s"), *Name.ToString());
+	}
 	return false;
 }
 
@@ -1843,12 +2025,17 @@ bool UPhysicsControlComponent::SetBodyModifierUpdateKinematicFromSimulation(
 	const FName Name,
 	const bool  bUpdateKinematicFromSimulation)
 {
-	FPhysicsBodyModifier* PhysicsBodyModifier = Implementation->FindBodyModifier(Name);
+	FPhysicsBodyModifier* PhysicsBodyModifier = FindBodyModifier(Name);
 	if (PhysicsBodyModifier)
 	{
-		PhysicsBodyModifier->bUpdateKinematicFromSimulation = bUpdateKinematicFromSimulation;
+		PhysicsBodyModifier->BodyModifierData.bUpdateKinematicFromSimulation = bUpdateKinematicFromSimulation;
 		return true;
 	}
+	if (bWarnAboutInvalidNames)
+	{
+		UE_LOG(LogPhysicsControlComponent, Warning,
+			TEXT("SetBodyModifierUpdateKinematicFromSimulation - invalid name %s"), *Name.ToString());
+	}
 	return false;
 }
 
@@ -1888,12 +2075,8 @@ void UPhysicsControlComponent::CreateControlsAndBodyModifiersFromLimbBones(
 	USkeletalMeshComponent*                     SkeletalMeshComponent,
 	const TArray<FPhysicsControlLimbSetupData>& LimbSetupData,
 	const FPhysicsControlData                   WorldSpaceControlData,
-	const FPhysicsControlSettings               WorldSpaceControlSettings,
 	const FPhysicsControlData                   ParentSpaceControlData,
-	const FPhysicsControlSettings               ParentSpaceControlSettings,
-	const EPhysicsMovementType                  PhysicsMovementType,
-	const float                                 GravityMultiplier,
-	const float                                 PhysicsBlendWeight,
+	const FPhysicsControlModifierData           BodyModifierData,
 	UMeshComponent*                             WorldComponent,
 	FName                                       WorldBoneName)
 {
@@ -1902,15 +2085,13 @@ void UPhysicsControlComponent::CreateControlsAndBodyModifiersFromLimbBones(
 
 	LimbWorldSpaceControls = CreateControlsFromLimbBones(
 		AllWorldSpaceControls, LimbBones, EPhysicsControlType::WorldSpace, 
-		WorldSpaceControlData, WorldSpaceControlSettings, WorldComponent, WorldBoneName);
+		WorldSpaceControlData, WorldComponent, WorldBoneName);
 
 	LimbParentSpaceControls = CreateControlsFromLimbBones(
 		AllParentSpaceControls, LimbBones, EPhysicsControlType::ParentSpace,
-		ParentSpaceControlData, ParentSpaceControlSettings);
+		ParentSpaceControlData);
 
-	LimbBodyModifiers = CreateBodyModifiersFromLimbBones(
-		AllBodyModifiers, LimbBones, PhysicsMovementType, ECollisionEnabled::QueryAndPhysics, 
-		GravityMultiplier, PhysicsBlendWeight, true);
+	LimbBodyModifiers = CreateBodyModifiersFromLimbBones(AllBodyModifiers, LimbBones, BodyModifierData);
 }
 
 //======================================================================================================================
@@ -1919,7 +2100,7 @@ void UPhysicsControlComponent::AddControlToSet(
 	const FName           Control, 
 	const FName           SetName)
 {
-	Implementation->NameRecords.AddControl(Control, SetName);
+	NameRecords.AddControl(Control, SetName);
 	NewSet.Names = GetControlNamesInSet(SetName);
 }
 
@@ -1931,7 +2112,7 @@ void UPhysicsControlComponent::AddControlsToSet(
 {
 	for (FName Control : Controls)
 	{
-		Implementation->NameRecords.AddControl(Control, SetName);
+		NameRecords.AddControl(Control, SetName);
 	}
 	NewSet.Names = GetControlNamesInSet(SetName);
 }
@@ -1939,7 +2120,7 @@ void UPhysicsControlComponent::AddControlsToSet(
 //======================================================================================================================
 const TArray<FName>& UPhysicsControlComponent::GetControlNamesInSet(const FName SetName) const
 {
-	return Implementation->NameRecords.GetControlNamesInSet(SetName);
+	return NameRecords.GetControlNamesInSet(SetName);
 }
 
 //======================================================================================================================
@@ -1951,7 +2132,7 @@ const TArray<FName>& UPhysicsControlComponent::GetAllBodyModifierNames() const
 //======================================================================================================================
 const TArray<FName>& UPhysicsControlComponent::GetBodyModifierNamesInSet(const FName SetName) const
 {
-	return Implementation->NameRecords.GetBodyModifierNamesInSet(SetName);
+	return NameRecords.GetBodyModifierNamesInSet(SetName);
 }
 
 //======================================================================================================================
@@ -1960,7 +2141,7 @@ void UPhysicsControlComponent::AddBodyModifierToSet(
 	const FName           BodyModifier, 
 	const FName           SetName)
 {
-	Implementation->NameRecords.AddBodyModifier(BodyModifier, SetName);
+	NameRecords.AddBodyModifier(BodyModifier, SetName);
 	NewSet.Names = GetBodyModifierNamesInSet(SetName);
 }
 
@@ -1973,7 +2154,7 @@ void UPhysicsControlComponent::AddBodyModifiersToSet(
 {
 	for (FName BodyModifier : BodyModifiers)
 	{
-		Implementation->NameRecords.AddBodyModifier(BodyModifier, SetName);
+		NameRecords.AddBodyModifier(BodyModifier, SetName);
 	}
 	NewSet.Names = GetBodyModifierNamesInSet(SetName);
 }
@@ -1982,7 +2163,7 @@ void UPhysicsControlComponent::AddBodyModifiersToSet(
 TArray<FName> UPhysicsControlComponent::GetSetsContainingControl(const FName Control) const
 {
 	TArray<FName> Result;
-	for (const TPair<FName, TArray<FName>>& ControlSetPair : Implementation->NameRecords.ControlSets)
+	for (const TPair<FName, TArray<FName>>& ControlSetPair : NameRecords.ControlSets)
 	{
 		for (const FName& ControlName : ControlSetPair.Value)
 		{
@@ -1999,7 +2180,7 @@ TArray<FName> UPhysicsControlComponent::GetSetsContainingControl(const FName Con
 TArray<FName> UPhysicsControlComponent::GetSetsContainingBodyModifier(const FName BodyModifier) const
 {
 	TArray<FName> Result;
-	for (const TPair<FName, TArray<FName>>& BodyModifierSetPair : Implementation->NameRecords.BodyModifierSets)
+	for (const TPair<FName, TArray<FName>>& BodyModifierSetPair : NameRecords.BodyModifierSets)
 	{
 		for (const FName& BodyModifierName : BodyModifierSetPair.Value)
 		{
@@ -2023,7 +2204,7 @@ TArray<FTransform> UPhysicsControlComponent::GetCachedBoneTransforms(
 
 	for (const FName& BoneName : BoneNames)
 	{
-		if (Implementation->GetBoneData(BoneData, SkeletalMeshComponent, BoneName))
+		if (GetBoneData(BoneData, SkeletalMeshComponent, BoneName))
 		{
 			FTransform BoneTransform(BoneData.Orientation, BoneData.Position);
 			Result.Add(BoneTransform);
@@ -2047,7 +2228,7 @@ TArray<FVector> UPhysicsControlComponent::GetCachedBonePositions(
 
 	for (const FName& BoneName : BoneNames)
 	{
-		if (Implementation->GetBoneData(BoneData, SkeletalMeshComponent, BoneName))
+		if (GetBoneData(BoneData, SkeletalMeshComponent, BoneName))
 		{
 			Result.Add(BoneData.Position);
 		}
@@ -2070,7 +2251,7 @@ TArray<FRotator> UPhysicsControlComponent::GetCachedBoneOrientations(
 
 	for (const FName& BoneName : BoneNames)
 	{
-		if (Implementation->GetBoneData(BoneData, SkeletalMeshComponent, BoneName))
+		if (GetBoneData(BoneData, SkeletalMeshComponent, BoneName))
 		{
 			Result.Add(BoneData.Orientation.Rotator());
 		}
@@ -2093,7 +2274,7 @@ TArray<FVector> UPhysicsControlComponent::GetCachedBoneVelocities(
 
 	for (const FName& BoneName : BoneNames)
 	{
-		if (Implementation->GetBoneData(BoneData, SkeletalMeshComponent, BoneName))
+		if (GetBoneData(BoneData, SkeletalMeshComponent, BoneName))
 		{
 			Result.Add(BoneData.Velocity);
 		}
@@ -2117,7 +2298,7 @@ TArray<FVector> UPhysicsControlComponent::GetCachedBoneAngularVelocities(
 
 	for (const FName& BoneName : BoneNames)
 	{
-		if (Implementation->GetBoneData(BoneData, SkeletalMeshComponent, BoneName))
+		if (GetBoneData(BoneData, SkeletalMeshComponent, BoneName))
 		{
 			Result.Add(BoneData.AngularVelocity);
 		}
@@ -2133,78 +2314,103 @@ TArray<FVector> UPhysicsControlComponent::GetCachedBoneAngularVelocities(
 //======================================================================================================================
 FTransform UPhysicsControlComponent::GetCachedBoneTransform(
 	const USkeletalMeshComponent* SkeletalMeshComponent,
-	const FName                   BoneName)
+	const FName                   Name)
 {
 	FCachedSkeletalMeshData::FBoneData BoneData;
-	if (Implementation->GetBoneData(BoneData, SkeletalMeshComponent, BoneName))
+	if (GetBoneData(BoneData, SkeletalMeshComponent, Name))
 	{
 		return FTransform(BoneData.Orientation, BoneData.Position);
 	}
+	if (bWarnAboutInvalidNames)
+	{
+		UE_LOG(LogPhysicsControlComponent, Warning,
+			TEXT("GetCachedBoneTransform - invalid name %s"), *Name.ToString());
+	}
 	return FTransform();
 }
 
 //======================================================================================================================
 FVector UPhysicsControlComponent::GetCachedBonePosition(
 	const USkeletalMeshComponent* SkeletalMeshComponent,
-	const FName                   BoneName)
+	const FName                   Name)
 {
 	FCachedSkeletalMeshData::FBoneData BoneData;
-	if (Implementation->GetBoneData(BoneData, SkeletalMeshComponent, BoneName))
+	if (GetBoneData(BoneData, SkeletalMeshComponent, Name))
 	{
 		return BoneData.Position;
 	}
+	if (bWarnAboutInvalidNames)
+	{
+		UE_LOG(LogPhysicsControlComponent, Warning,
+			TEXT("GetCachedBonePosition - invalid name %s"), *Name.ToString());
+	}
 	return FVector::ZeroVector;
 }
 
 //======================================================================================================================
 FRotator UPhysicsControlComponent::GetCachedBoneOrientation(
 	const USkeletalMeshComponent* SkeletalMeshComponent,
-	const FName                   BoneName)
+	const FName                   Name)
 {
 	FCachedSkeletalMeshData::FBoneData BoneData;
-	if (Implementation->GetBoneData(BoneData, SkeletalMeshComponent, BoneName))
+	if (GetBoneData(BoneData, SkeletalMeshComponent, Name))
 	{
 		return BoneData.Orientation.Rotator();
 	}
+	if (bWarnAboutInvalidNames)
+	{
+		UE_LOG(LogPhysicsControlComponent, Warning,
+			TEXT("GetCachedBoneOrientation - invalid name %s"), *Name.ToString());
+	}
 	return FRotator::ZeroRotator;
 }
 
 //======================================================================================================================
 FVector UPhysicsControlComponent::GetCachedBoneVelocity(
 	const USkeletalMeshComponent* SkeletalMeshComponent,
-	const FName                   BoneName)
+	const FName                   Name)
 {
 	FCachedSkeletalMeshData::FBoneData BoneData;
-	if (Implementation->GetBoneData(BoneData, SkeletalMeshComponent, BoneName))
+	if (GetBoneData(BoneData, SkeletalMeshComponent, Name))
 	{
 		return BoneData.Velocity;
 	}
+	if (bWarnAboutInvalidNames)
+	{
+		UE_LOG(LogPhysicsControlComponent, Warning,
+			TEXT("GetCachedBoneVelocity - invalid name %s"), *Name.ToString());
+	}
 	return FVector::Zero();
 }
 
 //======================================================================================================================
 FVector UPhysicsControlComponent::GetCachedBoneAngularVelocity(
 	const USkeletalMeshComponent* SkeletalMeshComponent,
-	const FName                   BoneName)
+	const FName                   Name)
 {
 	FCachedSkeletalMeshData::FBoneData BoneData;
-	if (Implementation->GetBoneData(BoneData, SkeletalMeshComponent, BoneName))
+	if (GetBoneData(BoneData, SkeletalMeshComponent, Name))
 	{
 		return BoneData.AngularVelocity;
 	}
+	if (bWarnAboutInvalidNames)
+	{
+		UE_LOG(LogPhysicsControlComponent, Warning,
+			TEXT("GetCachedBoneAngularVelocity - invalid name %s"), *Name.ToString());
+	}
 	return FVector::Zero();
 }
 
 //======================================================================================================================
 bool UPhysicsControlComponent::SetCachedBoneData(
 	const USkeletalMeshComponent* SkeletalMeshComponent,
-	const FName                   BoneName,
+	const FName                   Name,
 	const FTransform&             TM,
 	const FVector                 Velocity,
 	const FVector                 AngularVelocity)
 {
 	FCachedSkeletalMeshData::FBoneData* BoneData;
-	if (Implementation->GetModifiableBoneData(BoneData, SkeletalMeshComponent, BoneName))
+	if (GetModifiableBoneData(BoneData, SkeletalMeshComponent, Name))
 	{
 		BoneData->Position = TM.GetLocation();
 		BoneData->Orientation = TM.GetRotation();
@@ -2212,6 +2418,11 @@ bool UPhysicsControlComponent::SetCachedBoneData(
 		BoneData->AngularVelocity = AngularVelocity;
 		return true;
 	}
+	if (bWarnAboutInvalidNames)
+	{
+		UE_LOG(LogPhysicsControlComponent, Warning,
+			TEXT("SetCachedBoneData - invalid name %s"), *Name.ToString());
+	}
 	return false;
 }
 
@@ -2220,12 +2431,12 @@ bool UPhysicsControlComponent::ResetBodyModifierToCachedBoneTransform(
 	const FName                        Name,
 	const EResetToCachedTargetBehavior Behavior)
 {
-	FPhysicsBodyModifier* PhysicsBodyModifier = Implementation->FindBodyModifier(Name);
+	FPhysicsBodyModifier* PhysicsBodyModifier = FindBodyModifier(Name);
 	if (PhysicsBodyModifier)
 	{
 		if (Behavior == EResetToCachedTargetBehavior::ResetImmediately)
 		{
-			Implementation->ResetToCachedTarget(*PhysicsBodyModifier);
+			ResetToCachedTarget(*PhysicsBodyModifier);
 		}
 		else
 		{
@@ -2233,8 +2444,12 @@ bool UPhysicsControlComponent::ResetBodyModifierToCachedBoneTransform(
 		}
 		return true;
 	}
+	if (bWarnAboutInvalidNames)
+	{
+		UE_LOG(LogPhysicsControlComponent, Warning,
+			TEXT("ResetBodyModifierToCachedBoneTransform - invalid name %s"), *Name.ToString());
+	}
 	return false;
-
 }
 
 //======================================================================================================================
@@ -2256,6 +2471,17 @@ void UPhysicsControlComponent::ResetBodyModifiersInSetToCachedBoneTransforms(
 	ResetBodyModifiersToCachedBoneTransforms(GetBodyModifierNamesInSet(SetName), Behavior);
 }
 
+//======================================================================================================================
+bool UPhysicsControlComponent::GetControlExists(const FName Name) const
+{
+	return FindControlRecord(Name) != nullptr;
+}
+
+//======================================================================================================================
+bool UPhysicsControlComponent::GetBodyModifierExists(const FName Name) const
+{
+	return FindBodyModifier(Name) != nullptr;
+}
 
 #if WITH_EDITOR
 
@@ -2266,21 +2492,12 @@ void UPhysicsControlComponent::OnRegister()
 
 	if (SpriteComponent)
 	{
-		UpdateSpriteTexture();
+		SpriteComponent->SetSprite(LoadObject<UTexture2D>(NULL, TEXT("/Engine/EditorResources/S_KBSJoint.S_KBSJoint")));
 		SpriteComponent->SpriteInfo.Category = TEXT("Physics");
 		SpriteComponent->SpriteInfo.DisplayName = NSLOCTEXT("SpriteCategory", "Physics", "Physics");
 	}
 }
 
-//======================================================================================================================
-void UPhysicsControlComponent::UpdateSpriteTexture()
-{
-	if (SpriteComponent)
-	{
-		SpriteComponent->SetSprite(LoadObject<UTexture2D>(NULL, TEXT("/Engine/EditorResources/S_KBSJoint.S_KBSJoint")));
-	}
-}
-
 //======================================================================================================================
 void UPhysicsControlComponent::DebugDraw(FPrimitiveDrawInterface* PDI) const
 {
@@ -2288,7 +2505,7 @@ void UPhysicsControlComponent::DebugDraw(FPrimitiveDrawInterface* PDI) const
 	if (bShowDebugVisualization && VisualizationSizeScale > 0)
 	{
 		for (const TPair<FName, FPhysicsControlRecord>& PhysicsControlRecordPair : 
-			Implementation->PhysicsControlRecords)
+			PhysicsControlRecords)
 		{
 			const FName Name = PhysicsControlRecordPair.Key;
 			const FPhysicsControlRecord& Record = PhysicsControlRecordPair.Value;
@@ -2300,7 +2517,7 @@ void UPhysicsControlComponent::DebugDraw(FPrimitiveDrawInterface* PDI) const
 	if (!DebugControlDetailFilter.IsEmpty())
 	{
 		for (const TPair<FName, FPhysicsControlRecord>& PhysicsControlRecordPair :
-			Implementation->PhysicsControlRecords)
+			PhysicsControlRecords)
 		{
 			const FName Name = PhysicsControlRecordPair.Key;
 
@@ -2308,9 +2525,9 @@ void UPhysicsControlComponent::DebugDraw(FPrimitiveDrawInterface* PDI) const
 			{
 				const FPhysicsControlRecord& Record = PhysicsControlRecordPair.Value;
 
-				FString ParentComponentName = Record.PhysicsControl.ParentMeshComponent ?
+				FString ParentComponentName = Record.PhysicsControl.ParentMeshComponent.IsValid() ?
 					Record.PhysicsControl.ParentMeshComponent->GetName() : TEXT("NoParent");
-				FString ChildComponentName = Record.PhysicsControl.ChildMeshComponent ?
+				FString ChildComponentName = Record.PhysicsControl.ChildMeshComponent.IsValid() ?
 					Record.PhysicsControl.ChildMeshComponent->GetName() : TEXT("NoChild");
 
 				FString Text = FString::Printf(
@@ -2336,7 +2553,7 @@ void UPhysicsControlComponent::DebugDraw(FPrimitiveDrawInterface* PDI) const
 		FString AllNames;
 
 		for (const TPair<FName, FPhysicsControlRecord>& PhysicsControlRecordPair :
-			Implementation->PhysicsControlRecords)
+			PhysicsControlRecords)
 		{
 			const FName Name = PhysicsControlRecordPair.Key;
 			AllNames += Name.ToString() + TEXT(" ");
@@ -2348,7 +2565,7 @@ void UPhysicsControlComponent::DebugDraw(FPrimitiveDrawInterface* PDI) const
 		}
 		GEngine->AddOnScreenDebugMessage(
 			-1, 0.0f, FColor::White,
-			FString::Printf(TEXT("%d Controls: %s"), Implementation->PhysicsControlRecords.Num(), *AllNames));
+			FString::Printf(TEXT("%d Controls: %s"), PhysicsControlRecords.Num(), *AllNames));
 
 	}
 
@@ -2356,7 +2573,7 @@ void UPhysicsControlComponent::DebugDraw(FPrimitiveDrawInterface* PDI) const
 	if (!DebugBodyModifierDetailFilter.IsEmpty())
 	{
 		for (const TPair<FName, FPhysicsBodyModifier>& PhysicsBodyModifierPair :
-			Implementation->PhysicsBodyModifiers)
+			PhysicsBodyModifiers)
 		{
 			const FName Name = PhysicsBodyModifierPair.Key;
 
@@ -2364,19 +2581,20 @@ void UPhysicsControlComponent::DebugDraw(FPrimitiveDrawInterface* PDI) const
 			{
 				const FPhysicsBodyModifier& Record = PhysicsBodyModifierPair.Value;
 
-				FString ComponentName = Record.MeshComponent ? Record.MeshComponent->GetName() : TEXT("None");
+				FString ComponentName = Record.MeshComponent.IsValid() ? Record.MeshComponent->GetName() : TEXT("None");
 
 				FString Text = FString::Printf(
 					TEXT("%s: %s: %s %s GravityMultiplier %f BlendWeight %f"),
 					*Name.ToString(),
 					*ComponentName,
-					*UEnum::GetValueAsString(Record.MovementType),
-					*UEnum::GetValueAsString(Record.CollisionType),
-					Record.GravityMultiplier,
-					Record.PhysicsBlendWeight);
+					*UEnum::GetValueAsString(Record.BodyModifierData.MovementType),
+					*UEnum::GetValueAsString(Record.BodyModifierData.CollisionType),
+					Record.BodyModifierData.GravityMultiplier,
+					Record.BodyModifierData.PhysicsBlendWeight);
 
 				GEngine->AddOnScreenDebugMessage(-1, 0.0f, 
-					Record.MovementType == EPhysicsMovementType::Simulated ? FColor::Green : FColor::Red, Text);
+					Record.BodyModifierData.MovementType == EPhysicsMovementType::Simulated ? 
+					FColor::Green : FColor::Red, Text);
 			}
 		}
 	}
@@ -2386,7 +2604,7 @@ void UPhysicsControlComponent::DebugDraw(FPrimitiveDrawInterface* PDI) const
 	{
 		FString AllNames;
 
-		for (const TPair<FName, FPhysicsBodyModifier>& PhysicsBodyModifierPair : Implementation->PhysicsBodyModifiers)
+		for (const TPair<FName, FPhysicsBodyModifier>& PhysicsBodyModifierPair : PhysicsBodyModifiers)
 		{
 			const FName Name = PhysicsBodyModifierPair.Key;
 			AllNames += Name.ToString() + TEXT(" ");
@@ -2398,7 +2616,7 @@ void UPhysicsControlComponent::DebugDraw(FPrimitiveDrawInterface* PDI) const
 		}
 		GEngine->AddOnScreenDebugMessage(
 			-1, 0.0f, FColor::White,
-			FString::Printf(TEXT("%d Body modifiers: %s"), Implementation->PhysicsBodyModifiers.Num(), *AllNames));
+			FString::Printf(TEXT("%d Body modifiers: %s"), PhysicsBodyModifiers.Num(), *AllNames));
 
 	}
 }
@@ -2412,7 +2630,7 @@ void UPhysicsControlComponent::DebugDrawControl(
 	const FColor TargetColor(0, 255, 0);
 	const FColor CurrentColor(0, 0, 255);
 
-	const FConstraintInstance* ConstraintInstance = Record.PhysicsControlState.ConstraintInstance.Get();
+	const FConstraintInstance* ConstraintInstance = Record.ConstraintInstance.Get();
 
 	const bool bHaveLinear = Record.PhysicsControl.ControlData.LinearStrength > 0;
 	const bool bHaveAngular = Record.PhysicsControl.ControlData.AngularStrength > 0;
@@ -2420,7 +2638,7 @@ void UPhysicsControlComponent::DebugDrawControl(
 	if (Record.Enabled() && ConstraintInstance)
 	{
 		FBodyInstance* ChildBodyInstance = UE::PhysicsControlComponent::GetBodyInstance(
-			Record.PhysicsControl.ChildMeshComponent, Record.PhysicsControl.ChildBoneName);
+			Record.PhysicsControl.ChildMeshComponent.Get(), Record.PhysicsControl.ChildBoneName);
 		if (!ChildBodyInstance)
 		{
 			return;
@@ -2428,13 +2646,13 @@ void UPhysicsControlComponent::DebugDrawControl(
 		FTransform ChildBodyTM = ChildBodyInstance->GetUnrealWorldTransform();
 
 		FBodyInstance* ParentBodyInstance = UE::PhysicsControlComponent::GetBodyInstance(
-			Record.PhysicsControl.ParentMeshComponent, Record.PhysicsControl.ParentBoneName);
+			Record.PhysicsControl.ParentMeshComponent.Get(), Record.PhysicsControl.ParentBoneName);
 		const FTransform ParentBodyTM = ParentBodyInstance ? ParentBodyInstance->GetUnrealWorldTransform() : FTransform();
 
 		FTransform TargetTM;
 		FVector TargetVelocity;
 		FVector TargetAngularVelocity;
-		Implementation->CalculateControlTargetData(TargetTM, TargetVelocity, TargetAngularVelocity, Record, true);
+		CalculateControlTargetData(TargetTM, TargetVelocity, TargetAngularVelocity, Record, true);
 
 		// WorldChildFrameTM is the world-space transform of the child (driven) constraint frame
 		const FTransform WorldChildFrameTM = ConstraintInstance->GetRefFrame(EConstraintFrame::Frame1) * ChildBodyTM;
diff --git a/Engine/Plugins/Experimental/PhysicsControl/Source/PhysicsControl/Private/PhysicsControlComponentHelpers.cpp b/Engine/Plugins/Experimental/PhysicsControl/Source/PhysicsControl/Private/PhysicsControlComponentHelpers.cpp
--- a/Engine/Plugins/Experimental/PhysicsControl/Source/PhysicsControl/Private/PhysicsControlComponentHelpers.cpp
+++ b/Engine/Plugins/Experimental/PhysicsControl/Source/PhysicsControl/Private/PhysicsControlComponentHelpers.cpp
@@ -111,13 +111,11 @@ void ConvertConstraintProfileToControlData(
 //======================================================================================================================
 FBodyInstance* GetBodyInstance(UMeshComponent* MeshComponent, const FName BoneName)
 {
-	UStaticMeshComponent* StaticMeshComponent = Cast<UStaticMeshComponent>(MeshComponent);
-	if (StaticMeshComponent)
+	if (UStaticMeshComponent* StaticMeshComponent = Cast<UStaticMeshComponent>(MeshComponent))
 	{
 		return StaticMeshComponent->GetBodyInstance();
 	}
-	USkeletalMeshComponent* SkeletalMeshComponent = Cast<USkeletalMeshComponent>(MeshComponent);
-	if (SkeletalMeshComponent)
+	if (USkeletalMeshComponent* SkeletalMeshComponent = Cast<USkeletalMeshComponent>(MeshComponent))
 	{
 		return SkeletalMeshComponent->GetBodyInstance(BoneName);
 	}
@@ -134,8 +132,7 @@ FName GetPhysicalParentBone(USkeletalMeshComponent* SkeletalMeshComponent, FName
 		{
 			return FName();
 		}
-		const FBodyInstance* ParentBodyInstance = GetBodyInstance(SkeletalMeshComponent, ParentBoneName);
-		if (ParentBodyInstance)
+		if (const FBodyInstance* ParentBodyInstance = GetBodyInstance(SkeletalMeshComponent, ParentBoneName))
 		{
 			return ParentBoneName;
 		}
diff --git a/Engine/Plugins/Experimental/PhysicsControl/Source/PhysicsControl/Private/PhysicsControlComponentImpl.cpp b/Engine/Plugins/Experimental/PhysicsControl/Source/PhysicsControl/Private/PhysicsControlComponentImpl.cpp
--- a/Engine/Plugins/Experimental/PhysicsControl/Source/PhysicsControl/Private/PhysicsControlComponentImpl.cpp
+++ b/Engine/Plugins/Experimental/PhysicsControl/Source/PhysicsControl/Private/PhysicsControlComponentImpl.cpp
@@ -1,25 +1,32 @@
 // Copyright Epic Games, Inc. All Rights Reserved.
 
-#include "PhysicsControlComponentImpl.h"
-#include "Engine/SkeletalMesh.h"
-#include "PhysicsControlLog.h"
 #include "PhysicsControlComponent.h"
+#include "PhysicsControlLog.h"
 #include "PhysicsControlComponentHelpers.h"
 
-#include "Components/SkeletalMeshComponent.h"
+#include "PhysicsEngine/BodySetup.h"
+#include "PhysicsEngine/BodyInstance.h"
 #include "PhysicsEngine/ConstraintInstance.h"
 
+#include "Physics/Experimental/PhysScene_Chaos.h"
+
+#include "Components/SkeletalMeshComponent.h"
+
+#include "Engine/SkeletalMesh.h"
+#include "Engine/World.h"
+
+//UE_DISABLE_OPTIMIZATION;
+
+//======================================================================================================================
+// This file contains the non-public member functions of UPhysicsControlComponent
+//======================================================================================================================
+
 //======================================================================================================================
 // Helper to get a valid skeletal mesh component pointer from a record
 static USkeletalMeshComponent* GetValidSkeletalMeshComponentFromControlParent(
 	const FPhysicsControlRecord& Record)
 {
-	if (Record.PhysicsControl.ParentMeshComponent &&
-		Record.PhysicsControl.ParentMeshComponent->IsValidLowLevel())
-	{
-		return Cast<USkeletalMeshComponent>(Record.PhysicsControl.ParentMeshComponent.Get());
-	}
-	return nullptr;
+	return Cast<USkeletalMeshComponent>(Record.PhysicsControl.ParentMeshComponent.Get());
 }
 
 //======================================================================================================================
@@ -27,12 +34,7 @@ static USkeletalMeshComponent* GetValidSkeletalMeshComponentFromControlParent(
 static USkeletalMeshComponent* GetValidSkeletalMeshComponentFromControlChild(
 	const FPhysicsControlRecord& Record)
 {
-	if (Record.PhysicsControl.ChildMeshComponent &&
-		Record.PhysicsControl.ChildMeshComponent->IsValidLowLevel())
-	{
-		return Cast<USkeletalMeshComponent>(Record.PhysicsControl.ChildMeshComponent.Get());
-	}
-	return nullptr;
+	return Cast<USkeletalMeshComponent>(Record.PhysicsControl.ChildMeshComponent.Get());
 }
 
 //======================================================================================================================
@@ -40,16 +42,11 @@ static USkeletalMeshComponent* GetValidSkeletalMeshComponentFromControlChild(
 static USkeletalMeshComponent* GetValidSkeletalMeshComponentFromBodyModifier(
 	const FPhysicsBodyModifier& PhysicsBodyModifier)
 {
-	if (PhysicsBodyModifier.MeshComponent &&
-		PhysicsBodyModifier.MeshComponent->IsValidLowLevel())
-	{
-		return Cast<USkeletalMeshComponent>(PhysicsBodyModifier.MeshComponent.Get());
-	}
-	return nullptr;
+	return Cast<USkeletalMeshComponent>(PhysicsBodyModifier.MeshComponent.Get());
 }
 
 //======================================================================================================================
-bool FPhysicsControlComponentImpl::GetBoneData(
+bool UPhysicsControlComponent::GetBoneData(
 	FCachedSkeletalMeshData::FBoneData& OutBoneData,
 	const USkeletalMeshComponent*       InSkeletalMeshComponent,
 	const FName                         InBoneName) const
@@ -82,7 +79,7 @@ bool FPhysicsControlComponentImpl::GetBoneData(
 }
 
 //======================================================================================================================
-bool FPhysicsControlComponentImpl::GetModifiableBoneData(
+bool UPhysicsControlComponent::GetModifiableBoneData(
 	FCachedSkeletalMeshData::FBoneData*& OutBoneData,
 	const USkeletalMeshComponent*        InSkeletalMeshComponent,
 	const FName                          InBoneName)
@@ -116,7 +113,7 @@ bool FPhysicsControlComponentImpl::GetModifiableBoneData(
 
 
 //======================================================================================================================
-FPhysicsControlRecord* FPhysicsControlComponentImpl::FindControlRecord(const FName Name)
+FPhysicsControlRecord* UPhysicsControlComponent::FindControlRecord(const FName Name)
 {
 	if (FPhysicsControlRecord* Record = PhysicsControlRecords.Find(Name))
 	{
@@ -126,7 +123,17 @@ FPhysicsControlRecord* FPhysicsControlComponentImpl::FindControlRecord(const FNa
 }
 
 //======================================================================================================================
-FPhysicsControl* FPhysicsControlComponentImpl::FindControl(const FName Name)
+const FPhysicsControlRecord* UPhysicsControlComponent::FindControlRecord(const FName Name) const
+{
+	if (const FPhysicsControlRecord* Record = PhysicsControlRecords.Find(Name))
+	{
+		return Record;
+	}
+	return nullptr;
+}
+
+//======================================================================================================================
+FPhysicsControl* UPhysicsControlComponent::FindControl(const FName Name)
 {
 	if (FPhysicsControlRecord* Record = FindControlRecord(Name))
 	{
@@ -136,22 +143,32 @@ FPhysicsControl* FPhysicsControlComponentImpl::FindControl(const FName Name)
 }
 
 //======================================================================================================================
-bool FPhysicsControlComponentImpl::DetectTeleport(
+const FPhysicsControl* UPhysicsControlComponent::FindControl(const FName Name) const
+{
+	if (const FPhysicsControlRecord* Record = FindControlRecord(Name))
+	{
+		return &Record->PhysicsControl;
+	}
+	return nullptr;
+}
+
+//======================================================================================================================
+bool UPhysicsControlComponent::DetectTeleport(
 	const FVector& OldPosition, const FQuat& OldOrientation,
 	const FVector& NewPosition, const FQuat& NewOrientation) const
 {
-	if (Owner->TeleportDistanceThreshold > 0)
+	if (TeleportDistanceThreshold > 0)
 	{
-		double Distance = FVector::Distance(OldPosition, NewPosition);
-		if (Distance > Owner->TeleportDistanceThreshold)
+		const double Distance = FVector::Distance(OldPosition, NewPosition);
+		if (Distance > TeleportDistanceThreshold)
 		{
 			return true;
 		}
 	}
-	if (Owner->TeleportRotationThreshold > 0)
+	if (TeleportRotationThreshold > 0)
 	{
-		double Radians = OldOrientation.AngularDistance(NewOrientation);
-		if (FMath::RadiansToDegrees(Radians) > Owner->TeleportRotationThreshold)
+		const double Radians = OldOrientation.AngularDistance(NewOrientation);
+		if (FMath::RadiansToDegrees(Radians) > TeleportRotationThreshold)
 		{
 			return true;
 		}
@@ -160,15 +177,15 @@ bool FPhysicsControlComponentImpl::DetectTeleport(
 }
 
 //======================================================================================================================
-bool FPhysicsControlComponentImpl::DetectTeleport(const FTransform& OldTM, const FTransform& NewTM) const
+bool UPhysicsControlComponent::DetectTeleport(const FTransform& OldTM, const FTransform& NewTM) const
 {
 	return DetectTeleport(OldTM.GetTranslation(), OldTM.GetRotation(), NewTM.GetTranslation(), NewTM.GetRotation());
 }
 
 //======================================================================================================================
-void FPhysicsControlComponentImpl::UpdateCachedSkeletalBoneData(float Dt)
+void UPhysicsControlComponent::UpdateCachedSkeletalBoneData(float Dt)
 {
-	for (TPair<TObjectPtr<USkeletalMeshComponent>, FCachedSkeletalMeshData>& CachedSkeletalMeshDataPair :
+	for (TPair<TWeakObjectPtr<USkeletalMeshComponent>, FCachedSkeletalMeshData>& CachedSkeletalMeshDataPair :
 		CachedSkeletalMeshDatas)
 	{
 		FCachedSkeletalMeshData& CachedSkeletalMeshData = CachedSkeletalMeshDataPair.Value;
@@ -177,10 +194,10 @@ void FPhysicsControlComponentImpl::UpdateCachedSkeletalBoneData(float Dt)
 			continue;
 		}
 
-		TObjectPtr<USkeletalMeshComponent> SkeletalMeshComponent = CachedSkeletalMeshDataPair.Key;
+		TWeakObjectPtr<USkeletalMeshComponent> SkeletalMeshComponent = CachedSkeletalMeshDataPair.Key;
 		if (USkeletalMeshComponent* SkeletalMesh = SkeletalMeshComponent.Get())
 		{
-			FTransform ComponentTM = SkeletalMesh->GetComponentToWorld();
+			const FTransform ComponentTM = SkeletalMesh->GetComponentTransform();
 			const TArray<FTransform>& TMs = SkeletalMesh->GetEditableComponentSpaceTransforms();
 			const int32 NumTMs = TMs.Num();
 			if (NumTMs == CachedSkeletalMeshData.BoneData.Num() &&
@@ -223,12 +240,12 @@ void FPhysicsControlComponentImpl::UpdateCachedSkeletalBoneData(float Dt)
 }
 
 //======================================================================================================================
-void FPhysicsControlComponentImpl::ResetControls(bool bKeepControlRecords)
+void UPhysicsControlComponent::ResetControls(bool bKeepControlRecords)
 {
 	for (TPair<FName, FPhysicsControlRecord>& PhysicsControlRecordPair : PhysicsControlRecords)
 	{
 		FPhysicsControlRecord& Record = PhysicsControlRecordPair.Value;
-		Record.PhysicsControlState.Reset();
+		Record.ResetConstraint();
 	}
 
 	if (!bKeepControlRecords)
@@ -238,23 +255,23 @@ void FPhysicsControlComponentImpl::ResetControls(bool bKeepControlRecords)
 }
 
 //======================================================================================================================
-void FPhysicsControlComponentImpl::ApplyKinematicTarget(const FPhysicsBodyModifier& BodyModifier) const
+void UPhysicsControlComponent::ApplyKinematicTarget(const FPhysicsBodyModifier& BodyModifier) const
 {
 	// Seems like static and skeletal meshes need to be handled differently
 	if (USkeletalMeshComponent* SkeletalMeshComponent = GetValidSkeletalMeshComponentFromBodyModifier(BodyModifier))
 	{
 		FBodyInstance* BodyInstance = UE::PhysicsControlComponent::GetBodyInstance(
-			BodyModifier.MeshComponent, BodyModifier.BoneName);
+			BodyModifier.MeshComponent.Get(), BodyModifier.BoneName);
 		if (!BodyInstance)
 		{
 			return;
 		}
 
-		FTransform TM = BodyInstance->GetUnrealWorldTransform(); // Preserve scale
+		const FTransform TM = BodyInstance->GetUnrealWorldTransform(); // Preserve scale
 		FTransform KinematicTarget = TM;
 		KinematicTarget.SetRotation(BodyModifier.KinematicTargetOrientation);
 		KinematicTarget.SetTranslation(BodyModifier.KinematicTargetPosition);
-		if (BodyModifier.bUseSkeletalAnimation)
+		if (BodyModifier.BodyModifierData.bUseSkeletalAnimation)
 		{
 			FCachedSkeletalMeshData::FBoneData BoneData;
 			if (GetBoneData(BoneData, SkeletalMeshComponent, BodyModifier.BoneName))
@@ -268,8 +285,8 @@ void FPhysicsControlComponentImpl::ApplyKinematicTarget(const FPhysicsBodyModifi
 	}
 	else
 	{
-		FTransform TM = BodyModifier.MeshComponent->GetComponentToWorld();
-		ETeleportType TT = DetectTeleport(
+		const FTransform TM = BodyModifier.MeshComponent->GetComponentToWorld();
+		const ETeleportType TT = DetectTeleport(
 			TM.GetTranslation(), TM.GetRotation(), 
 			BodyModifier.KinematicTargetPosition, BodyModifier.KinematicTargetOrientation) 
 			? ETeleportType::ResetPhysics : ETeleportType::None;
@@ -280,10 +297,10 @@ void FPhysicsControlComponentImpl::ApplyKinematicTarget(const FPhysicsBodyModifi
 }
 
 //======================================================================================================================
-void FPhysicsControlComponentImpl::ResetToCachedTarget(const FPhysicsBodyModifier& BodyModifier) const
+void UPhysicsControlComponent::ResetToCachedTarget(const FPhysicsBodyModifier& BodyModifier) const
 {
 	FBodyInstance* BodyInstance = UE::PhysicsControlComponent::GetBodyInstance(
-		BodyModifier.MeshComponent, BodyModifier.BoneName);
+		BodyModifier.MeshComponent.Get(), BodyModifier.BoneName);
 	if (!BodyInstance)
 	{
 		return;
@@ -306,15 +323,15 @@ void FPhysicsControlComponentImpl::ResetToCachedTarget(const FPhysicsBodyModifie
 }
 
 //======================================================================================================================
-void FPhysicsControlComponentImpl::AddSkeletalMeshReferenceForCaching(
+void UPhysicsControlComponent::AddSkeletalMeshReferenceForCaching(
 	USkeletalMeshComponent* InSkeletalMeshComponent)
 {
 	check(InSkeletalMeshComponent);
-	for (TPair<TObjectPtr<USkeletalMeshComponent>, FCachedSkeletalMeshData>& CachedSkeletalMeshDataPair :
+	for (TPair<TWeakObjectPtr<USkeletalMeshComponent>, FCachedSkeletalMeshData>& CachedSkeletalMeshDataPair :
 		CachedSkeletalMeshDatas)
 	{
-		TObjectPtr<USkeletalMeshComponent> SkeletalMeshComponent = CachedSkeletalMeshDataPair.Key;
-		if (SkeletalMeshComponent == InSkeletalMeshComponent)
+		TWeakObjectPtr<USkeletalMeshComponent> SkeletalMeshComponent = CachedSkeletalMeshDataPair.Key;
+		if (SkeletalMeshComponent.Get() == InSkeletalMeshComponent)
 		{
 			FCachedSkeletalMeshData& CachedSkeletalMeshData = CachedSkeletalMeshDataPair.Value;
 			++CachedSkeletalMeshData.ReferenceCount;
@@ -323,11 +340,11 @@ void FPhysicsControlComponentImpl::AddSkeletalMeshReferenceForCaching(
 	}
 	FCachedSkeletalMeshData& Data = CachedSkeletalMeshDatas.Add(InSkeletalMeshComponent);
 	Data.ReferenceCount = 1;
-	Owner->PrimaryComponentTick.AddPrerequisite(InSkeletalMeshComponent, InSkeletalMeshComponent->PrimaryComponentTick);
+	PrimaryComponentTick.AddPrerequisite(InSkeletalMeshComponent, InSkeletalMeshComponent->PrimaryComponentTick);
 }
 
 //======================================================================================================================
-bool FPhysicsControlComponentImpl::RemoveSkeletalMeshReferenceForCaching(
+bool UPhysicsControlComponent::RemoveSkeletalMeshReferenceForCaching(
 	USkeletalMeshComponent* InSkeletalMeshComponent)
 {
 	check(InSkeletalMeshComponent);
@@ -339,13 +356,13 @@ bool FPhysicsControlComponentImpl::RemoveSkeletalMeshReferenceForCaching(
 
 	for (auto It = CachedSkeletalMeshDatas.CreateIterator(); It; ++It)
 	{
-		TObjectPtr<USkeletalMeshComponent> SkeletalMeshComponent = It.Key();
+		TWeakObjectPtr<USkeletalMeshComponent> SkeletalMeshComponent = It.Key();
 		FCachedSkeletalMeshData& Data = It.Value();
-		if (SkeletalMeshComponent == InSkeletalMeshComponent)
+		if (SkeletalMeshComponent.Get() == InSkeletalMeshComponent)
 		{
 			if (--Data.ReferenceCount == 0)
 			{
-				Owner->PrimaryComponentTick.RemovePrerequisite(
+				PrimaryComponentTick.RemovePrerequisite(
 					InSkeletalMeshComponent, InSkeletalMeshComponent->PrimaryComponentTick);
 				It.RemoveCurrent();
 				return true;
@@ -358,15 +375,15 @@ bool FPhysicsControlComponentImpl::RemoveSkeletalMeshReferenceForCaching(
 }
 
 //======================================================================================================================
-void FPhysicsControlComponentImpl::AddSkeletalMeshReferenceForModifier(
+void UPhysicsControlComponent::AddSkeletalMeshReferenceForModifier(
 	USkeletalMeshComponent* InSkeletalMeshComponent)
 {
 	check(InSkeletalMeshComponent);
-	for (TPair<TObjectPtr<USkeletalMeshComponent>, FModifiedSkeletalMeshData>& ModifiedSkeletalMeshDataPair :
+	for (TPair<TWeakObjectPtr<USkeletalMeshComponent>, FModifiedSkeletalMeshData>& ModifiedSkeletalMeshDataPair :
 		ModifiedSkeletalMeshDatas)
 	{
-		TObjectPtr<USkeletalMeshComponent> SkeletalMeshComponent = ModifiedSkeletalMeshDataPair.Key;
-		if (SkeletalMeshComponent == InSkeletalMeshComponent)
+		TWeakObjectPtr<USkeletalMeshComponent> SkeletalMeshComponent = ModifiedSkeletalMeshDataPair.Key;
+		if (SkeletalMeshComponent.Get() == InSkeletalMeshComponent)
 		{
 			FModifiedSkeletalMeshData& ModifiedSkeletalMeshData = ModifiedSkeletalMeshDataPair.Value;
 			++ModifiedSkeletalMeshData.ReferenceCount;
@@ -392,7 +409,7 @@ void FPhysicsControlComponentImpl::AddSkeletalMeshReferenceForModifier(
 }
 
 //======================================================================================================================
-bool FPhysicsControlComponentImpl::RemoveSkeletalMeshReferenceForModifier(
+bool UPhysicsControlComponent::RemoveSkeletalMeshReferenceForModifier(
 	USkeletalMeshComponent* InSkeletalMeshComponent)
 {
 	check(InSkeletalMeshComponent);
@@ -404,9 +421,9 @@ bool FPhysicsControlComponentImpl::RemoveSkeletalMeshReferenceForModifier(
 
 	for (auto It = ModifiedSkeletalMeshDatas.CreateIterator(); It; ++It)
 	{
-		TObjectPtr<USkeletalMeshComponent> SkeletalMeshComponent = It.Key();
+		TWeakObjectPtr<USkeletalMeshComponent> SkeletalMeshComponent = It.Key();
 		FModifiedSkeletalMeshData& Data = It.Value();
-		if (SkeletalMeshComponent == InSkeletalMeshComponent)
+		if (SkeletalMeshComponent.Get() == InSkeletalMeshComponent)
 		{
 			if (--Data.ReferenceCount == 0)
 			{
@@ -423,7 +440,7 @@ bool FPhysicsControlComponentImpl::RemoveSkeletalMeshReferenceForModifier(
 }
 
 //======================================================================================================================
-FName FPhysicsControlComponentImpl::GetUniqueControlName(
+FName UPhysicsControlComponent::GetUniqueControlName(
 	const FName ParentBoneName, const FName ChildBoneName, const FString& NamePrefix) const
 {
 	FString NameBase = NamePrefix;
@@ -441,7 +458,7 @@ FName FPhysicsControlComponentImpl::GetUniqueControlName(
 	// If the number gets too large, almost certainly we're in some nasty situation where this is
 	// getting called in a loop. Better to quit and fail, rather than allow the constraint set to
 	// increase without bound. 
-	for (int32 Index = 0; Index < Owner->MaxNumControlsOrModifiersPerName; ++Index)
+	for (int32 Index = 0; Index < MaxNumControlsOrModifiersPerName; ++Index)
 	{
 		FString NameStr = FString::Format(TEXT("{0}{1}"), { NameBase, Index });
 		FName Name(NameStr);
@@ -452,12 +469,12 @@ FName FPhysicsControlComponentImpl::GetUniqueControlName(
 	}
 	UE_LOG(LogPhysicsControlComponent, Warning,
 		TEXT("Unable to find a suitable Control name - the limit of MaxNumControlsOrModifiersPerName (%d) has been exceeded"),
-		Owner->MaxNumControlsOrModifiersPerName);
+		MaxNumControlsOrModifiersPerName);
 	return FName();
 }
 
 //======================================================================================================================
-FName FPhysicsControlComponentImpl::GetUniqueBodyModifierName(const FName BoneName) const
+FName UPhysicsControlComponent::GetUniqueBodyModifierName(const FName BoneName) const
 {
 	FString NameBase = TEXT("");
 	if (!BoneName.IsNone())
@@ -474,7 +491,7 @@ FName FPhysicsControlComponentImpl::GetUniqueBodyModifierName(const FName BoneNa
 	// If the number gets too large, almost certainly we're in some nasty situation where this is
 	// getting called in a loop. Better to quit and fail, rather than allow the modifier set to
 	// increase without bound. 
-	for (int32 Index = 0; Index != Owner->MaxNumControlsOrModifiersPerName; ++Index)
+	for (int32 Index = 0; Index != MaxNumControlsOrModifiersPerName; ++Index)
 	{
 		FString NameStr = FString::Format(TEXT("{0}{1}"), { NameBase, Index });
 		FName Name(NameStr);
@@ -485,12 +502,12 @@ FName FPhysicsControlComponentImpl::GetUniqueBodyModifierName(const FName BoneNa
 	}
 	UE_LOG(LogPhysicsControlComponent, Warning,
 		TEXT("Unable to find a suitable Body Modifier name - the limit of MaxNumControlsOrModifiersPerName (%d) has been exceeded"),
-		Owner->MaxNumControlsOrModifiersPerName);
+		MaxNumControlsOrModifiersPerName);
 	return FName();
 }
 
 //======================================================================================================================
-void FPhysicsControlComponentImpl::CalculateControlTargetData(
+void UPhysicsControlComponent::CalculateControlTargetData(
 	FTransform&                  OutTargetTM,
 	FVector&                     OutTargetVelocity,
 	FVector&                     OutTargetAngularVelocity,
@@ -505,7 +522,7 @@ void FPhysicsControlComponentImpl::CalculateControlTargetData(
 
 	// Set the target TM and velocities based on any skeletal action. Note that the targets from animation 
 	// should always account for the control point
-	if (Record.PhysicsControl.ControlSettings.bUseSkeletalAnimation)
+	if (Record.PhysicsControl.ControlData.bUseSkeletalAnimation)
 	{
 		FCachedSkeletalMeshData::FBoneData ChildBoneData, ParentBoneData;
 		bool bHaveChildBoneData = false;
@@ -522,9 +539,9 @@ void FPhysicsControlComponentImpl::CalculateControlTargetData(
 			bHaveParentBoneData = GetBoneData(
 				ParentBoneData, ParentSkeletalMeshComponent, Record.PhysicsControl.ParentBoneName);
 		}
-		else if (Record.PhysicsControl.ParentMeshComponent)
+		else if (Record.PhysicsControl.ParentMeshComponent.IsValid())
 		{
-			FTransform ParentTM = Record.PhysicsControl.ParentMeshComponent->GetComponentTransform();
+			const FTransform ParentTM = Record.PhysicsControl.ParentMeshComponent->GetComponentTransform();
 			ParentBoneData.Position = ParentTM.GetLocation();
 			ParentBoneData.Orientation = ParentTM.GetRotation();
 			ParentBoneData.Velocity = Record.PhysicsControl.ParentMeshComponent->GetPhysicsLinearVelocity();
@@ -542,50 +559,49 @@ void FPhysicsControlComponentImpl::CalculateControlTargetData(
 		if (bHaveChildBoneData)
 		{
 			bUsedSkeletalAnimation = true;
-			FTransform ChildBoneTM = ChildBoneData.GetTM();
+			const FTransform ChildBoneTM = ChildBoneData.GetTM();
 			if (bHaveParentBoneData)
 			{
-				FTransform ParentBoneTM = ParentBoneData.GetTM();
-				FTransform SkeletalDeltaTM = ChildBoneTM * ParentBoneTM.Inverse();
+				const FTransform ParentBoneTM = ParentBoneData.GetTM();
+				const FTransform SkeletalDeltaTM = ChildBoneTM * ParentBoneTM.Inverse();
 				// This puts TargetTM in the space of the ParentBone
 				OutTargetTM = SkeletalDeltaTM;
 
 				// Add on the control point offset
-				OutTargetTM.AddToTranslation(OutTargetTM.GetRotation() * Record.PhysicsControl.ControlSettings.ControlPoint);
+				OutTargetTM.AddToTranslation(OutTargetTM.GetRotation() * Record.GetControlPoint());
 
-				FQuat ParentBoneQ = ParentBoneTM.GetRotation();
-				FQuat ParentBoneInvQ = ParentBoneQ.Inverse();
+				const FQuat ParentBoneQ = ParentBoneTM.GetRotation();
+				const FQuat ParentBoneInvQ = ParentBoneQ.Inverse();
 
 				if (bCalculateVelocity)
 				{
-					if (Record.PhysicsControl.ControlSettings.SkeletalAnimationVelocityMultiplier != 0)
+					if (Record.PhysicsControl.ControlData.SkeletalAnimationVelocityMultiplier != 0)
 					{
 						// Offset of the control point from the target child bone TM, in world space.
-						FVector WorldControlPointOffset =
-							ChildBoneTM.GetRotation() * Record.PhysicsControl.ControlSettings.ControlPoint;
+						const FVector WorldControlPointOffset = ChildBoneTM.GetRotation() * Record.GetControlPoint();
 						// World space position of the target control point
-						FVector WorldChildControlPointPosition = ChildBoneTM.GetTranslation() + WorldControlPointOffset;
+						const FVector WorldChildControlPointPosition = ChildBoneTM.GetTranslation() + WorldControlPointOffset;
 
 						// World-space velocity of the control point due to the motion of the parent
 						// linear and angular velocity.
-						FVector ChildTargetVelocityDueToParent =
+						const FVector ChildTargetVelocityDueToParent =
 							ParentBoneData.Velocity + ParentBoneData.AngularVelocity.Cross(
 								WorldChildControlPointPosition - ParentBoneTM.GetTranslation());
 						// World-space velocity of the control point due to the motion of the child
 						// linear and angular velocity
-						FVector ChildTargetVelocity =
+						const FVector ChildTargetVelocity =
 							ChildBoneData.Velocity + ChildBoneData.AngularVelocity.Cross(WorldControlPointOffset);
 
 						// Pull out just the motion in the child that isn't due to the parent
-						FVector SkeletalTargetVelocity =
+						const FVector SkeletalTargetVelocity =
 							ParentBoneInvQ * (ChildTargetVelocity - ChildTargetVelocityDueToParent);
 						OutTargetVelocity += SkeletalTargetVelocity *
-							Record.PhysicsControl.ControlSettings.SkeletalAnimationVelocityMultiplier;
+							Record.PhysicsControl.ControlData.SkeletalAnimationVelocityMultiplier;
 
-						FVector SkeletalTargetAngularVelocity =
+						const FVector SkeletalTargetAngularVelocity =
 							ParentBoneInvQ * (ChildBoneData.AngularVelocity - ParentBoneData.AngularVelocity);
 						OutTargetAngularVelocity += SkeletalTargetAngularVelocity *
-							Record.PhysicsControl.ControlSettings.SkeletalAnimationVelocityMultiplier;
+							Record.PhysicsControl.ControlData.SkeletalAnimationVelocityMultiplier;
 					}
 				}
 			}
@@ -594,28 +610,27 @@ void FPhysicsControlComponentImpl::CalculateControlTargetData(
 				OutTargetTM = ChildBoneTM;
 
 				// Add on the control point offset
-				OutTargetTM.AddToTranslation(OutTargetTM.GetRotation()* Record.PhysicsControl.ControlSettings.ControlPoint);
+				OutTargetTM.AddToTranslation(OutTargetTM.GetRotation()* Record.GetControlPoint());
 
 				if (bCalculateVelocity)
 				{
 					OutTargetVelocity = ChildBoneTM.GetRotation() * OutTargetVelocity;
 					OutTargetAngularVelocity = ChildBoneTM.GetRotation() * OutTargetAngularVelocity;
 
-					if (Record.PhysicsControl.ControlSettings.SkeletalAnimationVelocityMultiplier != 0)
+					if (Record.PhysicsControl.ControlData.SkeletalAnimationVelocityMultiplier != 0)
 					{
-						FVector WorldControlPointOffset =
-							ChildBoneTM.GetRotation() * Record.PhysicsControl.ControlSettings.ControlPoint;
-						FVector WorldChildControlPointPosition = ChildBoneTM.GetTranslation() + WorldControlPointOffset;
+						const FVector WorldControlPointOffset = ChildBoneTM.GetRotation() * Record.GetControlPoint();
+						const FVector WorldChildControlPointPosition = ChildBoneTM.GetTranslation() + WorldControlPointOffset;
 
-						// World-space velocity of the control point due to the motion of the child
-						FVector ChildTargetVelocity =
+						// const World-space velocity of the control point due to the motion of the child
+						const FVector ChildTargetVelocity =
 							ChildBoneData.Velocity + ChildBoneData.AngularVelocity.Cross(WorldControlPointOffset);
 
 						OutTargetVelocity += ChildTargetVelocity *
-							Record.PhysicsControl.ControlSettings.SkeletalAnimationVelocityMultiplier;
+							Record.PhysicsControl.ControlData.SkeletalAnimationVelocityMultiplier;
 
 						OutTargetAngularVelocity += ChildBoneData.AngularVelocity *
-							Record.PhysicsControl.ControlSettings.SkeletalAnimationVelocityMultiplier;
+							Record.PhysicsControl.ControlData.SkeletalAnimationVelocityMultiplier;
 					}
 				}
 			}
@@ -629,7 +644,7 @@ void FPhysicsControlComponentImpl::CalculateControlTargetData(
 
 		// Calculate the authored target position/orientation - i.e. not using the skeletal animation
 		FQuat TargetOrientationQ = Target.TargetOrientation.Quaternion();
-		FVector TargetPosition = Target.TargetPosition;
+		const FVector TargetPosition = Target.TargetPosition;
 
 		FVector ExtraTargetPosition(0);
 		// Incorporate the offset from the control point. If we used animation, then we don't need
@@ -637,7 +652,7 @@ void FPhysicsControlComponentImpl::CalculateControlTargetData(
 		if (!bUsedSkeletalAnimation && 
 			Record.PhysicsControl.ControlTarget.bApplyControlPointToTarget)
 		{
-			ExtraTargetPosition = TargetOrientationQ * Record.PhysicsControl.ControlSettings.ControlPoint;
+			ExtraTargetPosition = TargetOrientationQ * Record.GetControlPoint();
 		}
 
 		// The record's target is specified in the space of the previously calculated/set OutTargetTM
@@ -646,16 +661,16 @@ void FPhysicsControlComponentImpl::CalculateControlTargetData(
 		if (bCalculateVelocity)
 		{
 			// Note that Target.TargetAngularVelocity is in revs per second (as it's user-facing)
-			FVector TargetAngularVelocity = Target.TargetAngularVelocity * UE_TWO_PI;
+			const FVector TargetAngularVelocity = Target.TargetAngularVelocity * UE_TWO_PI;
 			OutTargetAngularVelocity += TargetAngularVelocity;
-			FVector ExtraVelocity = TargetAngularVelocity.Cross(ExtraTargetPosition);
+			const FVector ExtraVelocity = TargetAngularVelocity.Cross(ExtraTargetPosition);
 			OutTargetVelocity += Target.TargetVelocity + ExtraVelocity;
 		}
 	}
 }
 
 //======================================================================================================================
-bool FPhysicsControlComponentImpl::ApplyControlStrengths(
+bool UPhysicsControlComponent::ApplyControlStrengths(
 	FPhysicsControlRecord& Record, FConstraintInstance* ConstraintInstance)
 {
 	const FPhysicsControlData& Data = Record.PhysicsControl.ControlData;
@@ -663,11 +678,11 @@ bool FPhysicsControlComponentImpl::ApplyControlStrengths(
 
 	float AngularSpring;
 	float AngularDamping;
-	float MaxTorque = Data.MaxTorque * Multiplier.MaxTorqueMultiplier;
+	const float MaxTorque = Data.MaxTorque * Multiplier.MaxTorqueMultiplier;
 
 	FVector LinearSpring;
 	FVector LinearDamping;
-	FVector MaxForce = Data.MaxForce * Multiplier.MaxForceMultiplier;
+	const FVector MaxForce = Data.MaxForce * Multiplier.MaxForceMultiplier;
 
 	UE::PhysicsControlComponent::ConvertStrengthToSpringParams(
 		AngularSpring, AngularDamping,
@@ -701,42 +716,52 @@ bool FPhysicsControlComponentImpl::ApplyControlStrengths(
 		LinearDamping.Z = 0;
 	}
 
-	ConstraintInstance->SetAngularDriveParams(AngularSpring, AngularDamping, MaxTorque);
-	ConstraintInstance->SetLinearDriveParams(LinearSpring, LinearDamping, MaxForce);
+	ConstraintInstance->SetDriveParams(
+		LinearSpring, LinearDamping, MaxForce,
+		FVector(0, 0, AngularSpring), FVector(0, 0, AngularDamping), FVector(0, 0, MaxTorque),
+		EAngularDriveMode::SLERP);
 
-	double TestAngular = (AngularSpring + AngularDamping) * FMath::Max(UE_SMALL_NUMBER, MaxTorque);
-	FVector TestLinear = (LinearSpring + LinearDamping) * FVector(
-		FMath::Max(UE_SMALL_NUMBER, MaxForce.X),
-		FMath::Max(UE_SMALL_NUMBER, MaxForce.Y),
-		FMath::Max(UE_SMALL_NUMBER, MaxForce.Z));
-	double Test = TestAngular + TestLinear.GetMax();
-	return Test > 0;
+	const bool bHaveAngular = (AngularSpring + AngularDamping) > 0;
+	const bool bHaveLinear = (LinearSpring + LinearDamping).GetMax() > 0;
+	return bHaveLinear || bHaveAngular;
 }
 
 //======================================================================================================================
-void FPhysicsControlComponentImpl::ApplyControl(FPhysicsControlRecord& Record)
+void UPhysicsControlComponent::ApplyControl(FPhysicsControlRecord& Record)
 {
-	FConstraintInstance* ConstraintInstance = Record.PhysicsControlState.ConstraintInstance.Get();
+	FConstraintInstance* ConstraintInstance = Record.ConstraintInstance.Get();
+
+	if (!ConstraintInstance)
+	{
+		return;
+	}
 
-	if (!ConstraintInstance || !Record.Enabled())
+	if (!Record.Enabled())
 	{
+		// Note that this will disable the constraint elements when strength/damping are zero
+		ConstraintInstance->SetDriveParams(
+			FVector::ZeroVector, FVector::ZeroVector, FVector::ZeroVector,
+			FVector::ZeroVector, FVector::ZeroVector, FVector::ZeroVector,
+			EAngularDriveMode::SLERP);
 		return;
 	}
 
+	// Always control collision, because otherwise maintaining it is very difficult, since
+	// constraint-controlled collision doesn't interact nicely when there are multiple constraints.
+	ConstraintInstance->SetDisableCollision(Record.PhysicsControl.ControlData.bDisableCollision);
+
 	FBodyInstance* ParentBodyInstance = UE::PhysicsControlComponent::GetBodyInstance(
-		Record.PhysicsControl.ParentMeshComponent, Record.PhysicsControl.ParentBoneName);
+		Record.PhysicsControl.ParentMeshComponent.Get(), Record.PhysicsControl.ParentBoneName);
 
 	FBodyInstance* ChildBodyInstance = UE::PhysicsControlComponent::GetBodyInstance(
-		Record.PhysicsControl.ChildMeshComponent, Record.PhysicsControl.ChildBoneName);
+		Record.PhysicsControl.ChildMeshComponent.Get(), Record.PhysicsControl.ChildBoneName);
 
 	if (!ParentBodyInstance && !ChildBodyInstance)
 	{
 		return;
 	}
 
-	// TODO UE-197351 disable constraint features fully when strengths/damping are zero
-
-	// Set strengths etc
+	// Set strengths etc and then targets (if there were strengths)
 	if (ApplyControlStrengths(Record, ConstraintInstance))
 	{
 		FTransform TargetTM;
@@ -748,7 +773,7 @@ void FPhysicsControlComponentImpl::ApplyControl(FPhysicsControlRecord& Record)
 		ConstraintInstance->SetAngularOrientationTarget(TargetTM.GetRotation());
 		ConstraintInstance->SetLinearVelocityTarget(TargetVelocity);
 		ConstraintInstance->SetAngularVelocityTarget(TargetAngularVelocity / UE_TWO_PI); // In rev/sec
-		ConstraintInstance->SetParentDominates(Record.PhysicsControl.ControlSettings.bOnlyControlChildObject);
+		ConstraintInstance->SetParentDominates(Record.PhysicsControl.ControlData.bOnlyControlChildObject);
 
 		if (ParentBodyInstance)
 		{
@@ -762,15 +787,78 @@ void FPhysicsControlComponentImpl::ApplyControl(FPhysicsControlRecord& Record)
 }
 
 //======================================================================================================================
-FPhysicsBodyModifier* FPhysicsControlComponentImpl::FindBodyModifier(const FName Name)
+void UPhysicsControlComponent::ApplyBodyModifier(FPhysicsBodyModifier& BodyModifier)
+{
+	USkeletalMeshComponent* SKM = GetValidSkeletalMeshComponentFromBodyModifier(BodyModifier);
+	FBodyInstance* BodyInstance = UE::PhysicsControlComponent::GetBodyInstance(
+		BodyModifier.MeshComponent.Get(), BodyModifier.BoneName);
+	if (BodyInstance)
+	{
+		switch (BodyModifier.BodyModifierData.MovementType)
+		{
+		case EPhysicsMovementType::Static:
+			BodyInstance->SetInstanceSimulatePhysics(false, false, true);
+			break;
+		case EPhysicsMovementType::Kinematic:
+			BodyInstance->SetInstanceSimulatePhysics(false, false, true);
+			ApplyKinematicTarget(BodyModifier);
+			break;
+		case EPhysicsMovementType::Simulated:
+			BodyInstance->SetInstanceSimulatePhysics(true, false, true);
+			break;
+		default:
+			UE_LOG(LogPhysicsControlComponent, Warning, TEXT("Invalid movement type %d"),
+				int(BodyModifier.BodyModifierData.MovementType));
+			break;
+		}
+
+		// We always overwrite the physics blend weight, since the functions above can still modify
+		// it (even though they all use the "maintain physics blending" option), since there is an
+		// expectation that zero blend weight means to disable physics.
+		BodyInstance->PhysicsBlendWeight = BodyModifier.BodyModifierData.PhysicsBlendWeight;
+		BodyInstance->SetUpdateKinematicFromSimulation(BodyModifier.BodyModifierData.bUpdateKinematicFromSimulation);
+
+		UBodySetup* BodySetup = BodyInstance->GetBodySetup();
+		if (BodySetup)
+		{
+			int32 NumShapes = BodySetup->AggGeom.GetElementCount();
+			for (int32 ShapeIndex = 0; ShapeIndex != NumShapes; ++ShapeIndex)
+			{
+				BodyInstance->SetShapeCollisionEnabled(ShapeIndex, BodyModifier.BodyModifierData.CollisionType);
+			}
+		}
+
+		if (BodyInstance->IsInstanceSimulatingPhysics())
+		{
+			const float GravityZ = BodyInstance->GetPhysicsScene()->GetOwningWorld()->GetGravityZ();
+			const float AppliedGravityZ = BodyInstance->bEnableGravity ? GravityZ : 0.0f;
+			const float DesiredGravityZ = GravityZ * BodyModifier.BodyModifierData.GravityMultiplier;
+			const float GravityZToApply = DesiredGravityZ - AppliedGravityZ;
+			BodyInstance->AddForce(FVector(0, 0, GravityZToApply), true, true);
+		}
+	}
+	if (BodyModifier.bResetToCachedTarget)
+	{
+		BodyModifier.bResetToCachedTarget = false;
+		ResetToCachedTarget(BodyModifier);
+	}
+}
+
+
+//======================================================================================================================
+FPhysicsBodyModifier* UPhysicsControlComponent::FindBodyModifier(const FName Name)
 {
 	return PhysicsBodyModifiers.Find(Name);
 }
 
 //======================================================================================================================
-bool FPhysicsControlComponentImpl::DestroyControl(
-	const FName            Name,
-	const EDestroyBehavior DestroyBehavior)
+const FPhysicsBodyModifier* UPhysicsControlComponent::FindBodyModifier(const FName Name) const
+{
+	return PhysicsBodyModifiers.Find(Name);
+}
+
+//======================================================================================================================
+bool UPhysicsControlComponent::DestroyControl(const FName Name, const EDestroyBehavior DestroyBehavior)
 {
 	FPhysicsControlRecord* Record = FindControlRecord(Name);
 	if (Record)
@@ -784,7 +872,7 @@ bool FPhysicsControlComponentImpl::DestroyControl(
 			RemoveSkeletalMeshReferenceForCaching(SkeletalMeshComponent);
 		}
 
-		Record->PhysicsControlState.Reset();
+		Record->ResetConstraint();
 		NameRecords.RemoveControl(Name);
 		if (DestroyBehavior == EDestroyBehavior::RemoveRecord)
 		{
@@ -792,13 +880,16 @@ bool FPhysicsControlComponentImpl::DestroyControl(
 		}
 		return true;
 	}
+	if (bWarnAboutInvalidNames)
+	{
+		UE_LOG(LogPhysicsControlComponent, Warning,
+			TEXT("DestroyControl - invalid name %s"), *Name.ToString());
+	}
 	return false;
 }
 
 //======================================================================================================================
-bool FPhysicsControlComponentImpl::DestroyBodyModifier(
-	const FName            Name,
-	const EDestroyBehavior DestroyBehavior)
+bool UPhysicsControlComponent::DestroyBodyModifier(const FName Name, const EDestroyBehavior DestroyBehavior)
 {
 	FPhysicsBodyModifier* BodyModifier = FindBodyModifier(Name);
 	if (BodyModifier)
@@ -815,6 +906,11 @@ bool FPhysicsControlComponentImpl::DestroyBodyModifier(
 		}
 		return true;
 	}
+	if (bWarnAboutInvalidNames)
+	{
+		UE_LOG(LogPhysicsControlComponent, Warning,
+			TEXT("DestroyBodyModifier - invalid name %s"), *Name.ToString());
+	}
 	return false;
 }
 
diff --git a/Engine/Plugins/Experimental/PhysicsControl/Source/PhysicsControl/Private/PhysicsControlComponentImpl.h b/Engine/Plugins/Experimental/PhysicsControl/Source/PhysicsControl/Private/PhysicsControlComponentImpl.h
deleted file mode 100644
diff --git a/Engine/Plugins/Experimental/PhysicsControl/Source/PhysicsControl/Private/PhysicsControlData.cpp b/Engine/Plugins/Experimental/PhysicsControl/Source/PhysicsControl/Private/PhysicsControlData.cpp
--- a/Engine/Plugins/Experimental/PhysicsControl/Source/PhysicsControl/Private/PhysicsControlData.cpp
+++ b/Engine/Plugins/Experimental/PhysicsControl/Source/PhysicsControl/Private/PhysicsControlData.cpp
@@ -4,6 +4,108 @@
 
 #define SET_SPARSE_DATA(NAME) NAME = SparseData.bEnable##NAME ? SparseData.NAME : NAME
 
+#define INTERPOLATE_PARAM(NAME) Output.NAME = FMath::Lerp(A.NAME, B.NAME, Weight)
+#define SELECT_PARAM(NAME) Output.NAME = (Weight < 0.5f) ? A.NAME : B.NAME
+#define SET_ENABLED_PARAM(NAME) Output.bEnable##NAME = A.bEnable##NAME && B.bEnable##NAME
+
+//======================================================================================================================
+FPhysicsControlData Interpolate(
+	const FPhysicsControlData& A, const FPhysicsControlData& B, const float Weight)
+{
+	FPhysicsControlData Output;
+	INTERPOLATE_PARAM(LinearStrength);
+	INTERPOLATE_PARAM(LinearDampingRatio);
+	INTERPOLATE_PARAM(LinearExtraDamping);
+	INTERPOLATE_PARAM(AngularStrength);
+	INTERPOLATE_PARAM(AngularDampingRatio);
+	INTERPOLATE_PARAM(AngularExtraDamping);
+	INTERPOLATE_PARAM(LinearTargetVelocityMultiplier);
+	INTERPOLATE_PARAM(AngularTargetVelocityMultiplier);
+	INTERPOLATE_PARAM(SkeletalAnimationVelocityMultiplier);
+	INTERPOLATE_PARAM(CustomControlPoint);
+	SELECT_PARAM(bEnabled);
+	SELECT_PARAM(bUseCustomControlPoint);
+	SELECT_PARAM(bUseSkeletalAnimation);
+	SELECT_PARAM(bDisableCollision);
+	SELECT_PARAM(bOnlyControlChildObject);
+	return Output;
+}
+
+//======================================================================================================================
+FPhysicsControlSparseData Interpolate(
+	const FPhysicsControlSparseData& A, const FPhysicsControlSparseData& B, const float Weight)
+{
+	FPhysicsControlSparseData Output;
+	INTERPOLATE_PARAM(LinearStrength);
+	INTERPOLATE_PARAM(LinearDampingRatio);
+	INTERPOLATE_PARAM(LinearExtraDamping);
+	INTERPOLATE_PARAM(AngularStrength);
+	INTERPOLATE_PARAM(AngularDampingRatio);
+	INTERPOLATE_PARAM(AngularExtraDamping);
+	INTERPOLATE_PARAM(LinearTargetVelocityMultiplier);
+	INTERPOLATE_PARAM(AngularTargetVelocityMultiplier);
+	INTERPOLATE_PARAM(SkeletalAnimationVelocityMultiplier);
+	INTERPOLATE_PARAM(CustomControlPoint);
+	SELECT_PARAM(bEnabled);
+	SELECT_PARAM(bUseCustomControlPoint);
+	SELECT_PARAM(bUseSkeletalAnimation);
+	SELECT_PARAM(bDisableCollision);
+	SELECT_PARAM(bOnlyControlChildObject);
+
+	SET_ENABLED_PARAM(LinearStrength);
+	SET_ENABLED_PARAM(LinearDampingRatio);
+	SET_ENABLED_PARAM(LinearExtraDamping);
+	SET_ENABLED_PARAM(AngularStrength);
+	SET_ENABLED_PARAM(AngularDampingRatio);
+	SET_ENABLED_PARAM(AngularExtraDamping);
+	SET_ENABLED_PARAM(LinearTargetVelocityMultiplier);
+	SET_ENABLED_PARAM(AngularTargetVelocityMultiplier);
+	SET_ENABLED_PARAM(SkeletalAnimationVelocityMultiplier);
+	SET_ENABLED_PARAM(CustomControlPoint);
+	SET_ENABLED_PARAM(bEnabled);
+	SET_ENABLED_PARAM(bUseCustomControlPoint);
+	SET_ENABLED_PARAM(bUseSkeletalAnimation);
+	SET_ENABLED_PARAM(bDisableCollision);
+	SET_ENABLED_PARAM(bOnlyControlChildObject);
+
+	return Output;
+}
+
+//======================================================================================================================
+FPhysicsControlModifierData Interpolate(
+	const FPhysicsControlModifierData& A, const FPhysicsControlModifierData& B, const float Weight)
+{
+	FPhysicsControlModifierData Output;
+	SELECT_PARAM(MovementType);
+	SELECT_PARAM(CollisionType);
+	INTERPOLATE_PARAM(GravityMultiplier);
+	INTERPOLATE_PARAM(PhysicsBlendWeight);
+	SELECT_PARAM(bUseSkeletalAnimation);
+	SELECT_PARAM(bUpdateKinematicFromSimulation);
+	return Output;
+}
+
+//======================================================================================================================
+FPhysicsControlModifierSparseData Interpolate(
+	const FPhysicsControlModifierSparseData& A, const FPhysicsControlModifierSparseData& B, const float Weight)
+{
+	FPhysicsControlModifierSparseData Output;
+	SELECT_PARAM(MovementType);
+	SELECT_PARAM(CollisionType);
+	INTERPOLATE_PARAM(GravityMultiplier);
+	INTERPOLATE_PARAM(PhysicsBlendWeight);
+	SELECT_PARAM(bUseSkeletalAnimation);
+	SELECT_PARAM(bUpdateKinematicFromSimulation);
+
+	SET_ENABLED_PARAM(MovementType);
+	SET_ENABLED_PARAM(CollisionType);
+	SET_ENABLED_PARAM(GravityMultiplier);
+	SET_ENABLED_PARAM(PhysicsBlendWeight);
+	SET_ENABLED_PARAM(bUseSkeletalAnimation);
+	SET_ENABLED_PARAM(bUpdateKinematicFromSimulation);
+	return Output;
+}
+
 //======================================================================================================================
 void FPhysicsControlData::UpdateFromSparseData(const FPhysicsControlSparseData& SparseData)
 {
@@ -17,14 +119,40 @@ void FPhysicsControlData::UpdateFromSparseData(const FPhysicsControlSparseData&
 	SET_SPARSE_DATA(MaxTorque);
 	SET_SPARSE_DATA(LinearTargetVelocityMultiplier);
 	SET_SPARSE_DATA(AngularTargetVelocityMultiplier);
+	SET_SPARSE_DATA(SkeletalAnimationVelocityMultiplier);
+	SET_SPARSE_DATA(CustomControlPoint);
 	SET_SPARSE_DATA(bEnabled);
+	SET_SPARSE_DATA(bUseCustomControlPoint);
+	SET_SPARSE_DATA(bUseSkeletalAnimation);
+	SET_SPARSE_DATA(bDisableCollision);
+	SET_SPARSE_DATA(bOnlyControlChildObject);
+}
+
+//======================================================================================================================
+void FPhysicsControlMultiplier::UpdateFromSparseData(const FPhysicsControlSparseMultiplier& SparseData)
+{
+	SET_SPARSE_DATA(LinearStrengthMultiplier);
+	SET_SPARSE_DATA(LinearDampingRatioMultiplier);
+	SET_SPARSE_DATA(LinearExtraDampingMultiplier);
+	SET_SPARSE_DATA(MaxForceMultiplier);
+	SET_SPARSE_DATA(AngularStrengthMultiplier);
+	SET_SPARSE_DATA(AngularDampingRatioMultiplier);
+	SET_SPARSE_DATA(AngularExtraDampingMultiplier);
+	SET_SPARSE_DATA(MaxTorqueMultiplier);
 }
 
 //======================================================================================================================
 void FPhysicsControlModifierData::UpdateFromSparseData(const FPhysicsControlModifierSparseData& SparseData)
 {
 	SET_SPARSE_DATA(MovementType);
+	SET_SPARSE_DATA(CollisionType);
 	SET_SPARSE_DATA(GravityMultiplier);
+	SET_SPARSE_DATA(PhysicsBlendWeight);
+	SET_SPARSE_DATA(bUseSkeletalAnimation);
+	SET_SPARSE_DATA(bUpdateKinematicFromSimulation);
 }
 
 #undef SET_SPARSE_DATA
+#undef INTERPOLATE_PARAM
+#undef SELECT_PARAM
+#undef SET_ENABLED_PARAM
diff --git a/Engine/Plugins/Experimental/PhysicsControl/Source/PhysicsControl/Private/PhysicsControlNameRecords.cpp b/Engine/Plugins/Experimental/PhysicsControl/Source/PhysicsControl/Private/PhysicsControlNameRecords.cpp
--- a/Engine/Plugins/Experimental/PhysicsControl/Source/PhysicsControl/Private/PhysicsControlNameRecords.cpp
+++ b/Engine/Plugins/Experimental/PhysicsControl/Source/PhysicsControl/Private/PhysicsControlNameRecords.cpp
@@ -1,11 +1,40 @@
 // Copyright Epic Games, Inc. All Rights Reserved.
 #include "PhysicsControlNameRecords.h"
 
+//======================================================================================================================
+void AddUniqueNamesToSet(const TArray<FName>& Names, const FName SetName, TMap<FName, TArray<FName>>& SetCollection)
+{
+	TArray<FName>& SetElementNames = SetCollection.FindOrAdd(SetName);
+
+	for (const FName Name : Names)
+	{
+		SetElementNames.AddUnique(Name);
+	}
+}
+
 //======================================================================================================================
 void FPhysicsControlNameRecords::AddControl(FName Name, FName SetName)
 {
-	ControlSets.FindOrAdd(SetName).Add(Name);
-	ControlSets.FindOrAdd("All").Add(Name);
+	ControlSets.FindOrAdd(SetName).AddUnique(Name);
+	ControlSets.FindOrAdd("All").AddUnique(Name);
+}
+
+//======================================================================================================================
+void FPhysicsControlNameRecords::AddControl(FName Name, const TArray<FName>& SetNames)
+{
+	for (const FName SetName : SetNames)
+	{
+		ControlSets.FindOrAdd(SetName).AddUnique(Name);
+	}
+	ControlSets.FindOrAdd("All").AddUnique(Name);
+}
+
+
+//======================================================================================================================
+void FPhysicsControlNameRecords::AddControls(const TArray<FName>& ControlNames, const FName SetName)
+{
+	AddUniqueNamesToSet(ControlNames, SetName, ControlSets);
+	AddUniqueNamesToSet(ControlNames, "All", ControlSets);
 }
 
 //======================================================================================================================
@@ -30,10 +59,27 @@ const TArray<FName>& FPhysicsControlNameRecords::GetControlNamesInSet(FName SetN
 }
 
 //======================================================================================================================
-void FPhysicsControlNameRecords::AddBodyModifier(FName Name, FName Set)
+void FPhysicsControlNameRecords::AddBodyModifier(FName Name, FName SetName)
+{
+	BodyModifierSets.FindOrAdd(SetName).AddUnique(Name);
+	BodyModifierSets.FindOrAdd("All").AddUnique(Name);
+}
+
+//======================================================================================================================
+void FPhysicsControlNameRecords::AddBodyModifier(FName Name, const TArray<FName>& SetNames)
+{
+	for (FName SetName : SetNames)
+	{
+		BodyModifierSets.FindOrAdd(SetName).AddUnique(Name);
+	}
+	BodyModifierSets.FindOrAdd("All").AddUnique(Name);
+}
+
+//======================================================================================================================
+void FPhysicsControlNameRecords::AddBodyModifiers(const TArray<FName>& BodyModifierNames, const FName SetName)
 {
-	BodyModifierSets.FindOrAdd(Set).Add(Name);	
-	BodyModifierSets.FindOrAdd("All").Add(Name);
+	AddUniqueNamesToSet(BodyModifierNames, SetName, BodyModifierSets);
+	AddUniqueNamesToSet(BodyModifierNames, "All", BodyModifierSets);
 }
 
 //======================================================================================================================
@@ -46,9 +92,9 @@ void FPhysicsControlNameRecords::RemoveBodyModifier(FName Name)
 }
 
 //======================================================================================================================
-const TArray<FName>& FPhysicsControlNameRecords::GetBodyModifierNamesInSet(FName Set) const
+const TArray<FName>& FPhysicsControlNameRecords::GetBodyModifierNamesInSet(FName SetName) const
 {
-	const TArray<FName>* Names = BodyModifierSets.Find(Set);
+	const TArray<FName>* Names = BodyModifierSets.Find(SetName);
 	if (Names)
 	{
 		return *Names;
@@ -56,3 +102,38 @@ const TArray<FName>& FPhysicsControlNameRecords::GetBodyModifierNamesInSet(FName
 	static TArray<FName> FailureResult;
 	return FailureResult;
 }
+
+//======================================================================================================================
+void FPhysicsControlNameRecords::Reset()
+{
+	ControlSets.Reset();
+	BodyModifierSets.Reset();
+}
+
+//======================================================================================================================
+TArray<FName> ExpandName(const FName InName, const TMap<FName, TArray<FName>>& SetNames)
+{
+	TArray<FName> OutputNames;
+	if (const TArray<FName>* const FoundSet = SetNames.Find(InName))
+	{
+		OutputNames.Append(*FoundSet);
+	}
+	else
+	{
+		OutputNames.Add(InName);
+	}
+	return OutputNames;
+}
+
+//======================================================================================================================
+TArray<FName> ExpandName(const TArray<FName>& InNames, const TMap<FName, TArray<FName>>& SetNames)
+{
+	TArray<FName> OutputNames;
+
+	for (const FName Name : InNames)
+	{
+		OutputNames.Append(ExpandName(Name, SetNames));
+	}
+
+	return OutputNames;
+}
diff --git a/Engine/Plugins/Experimental/PhysicsControl/Source/PhysicsControl/Private/PhysicsControlNameRecords.h b/Engine/Plugins/Experimental/PhysicsControl/Source/PhysicsControl/Private/PhysicsControlNameRecords.h
deleted file mode 100644
diff --git a/Engine/Plugins/Experimental/PhysicsControl/Source/PhysicsControl/Private/PhysicsControlOperatorNameGeneration.cpp b/Engine/Plugins/Experimental/PhysicsControl/Source/PhysicsControl/Private/PhysicsControlOperatorNameGeneration.cpp
--- a/Engine/Plugins/Experimental/PhysicsControl/Source/PhysicsControl/Private/PhysicsControlOperatorNameGeneration.cpp
+++ b/Engine/Plugins/Experimental/PhysicsControl/Source/PhysicsControl/Private/PhysicsControlOperatorNameGeneration.cpp
@@ -49,7 +49,7 @@ FName FindParentBodyBoneName(const FName BoneName, const FReferenceSkeleton& Ref
 }
 
 //======================================================================================================================
-template<typename TOperatorFunctorType> void CreateAdditionalBodyModifiers(const TArray<FRigidBodyModifierCreation> CreationSpecifiers, FRigidBodyNameRecords& NameRecords, TOperatorFunctorType& OperatorFunctor)
+template<typename TOperatorFunctorType> void CreateAdditionalBodyModifiers(const TArray<FRigidBodyModifierCreation> CreationSpecifiers, FPhysicsControlNameRecords& NameRecords, TOperatorFunctorType& OperatorFunctor)
 {
 	for (const FRigidBodyModifierCreation& Specifier : CreationSpecifiers)
 	{
@@ -72,7 +72,7 @@ template<typename TOperatorFunctorType> void CreateAdditionalBodyModifiers(const
 }
 
 //======================================================================================================================
-template<typename TOperatorFunctorType> void CreateAdditionalControls(const TArray<FRigidBodyControlCreation> CreationSpecifiers, FRigidBodyNameRecords& NameRecords, TOperatorFunctorType& OperatorFunctor)
+template<typename TOperatorFunctorType> void CreateAdditionalControls(const TArray<FRigidBodyControlCreation> CreationSpecifiers, FPhysicsControlNameRecords& NameRecords, TOperatorFunctorType& OperatorFunctor)
 {
 	for (const FRigidBodyControlCreation& Specifier : CreationSpecifiers)
 	{
@@ -104,7 +104,7 @@ template<typename TFunctor> void CreateControlsFromLimbBones(
 	const FReferenceSkeleton& RefSkeleton,
 	UPhysicsAsset* const PhysicsAsset,
 	const FPhysicsControlData& ControlData,
-	FRigidBodyNameRecords& NameRecords,
+	FPhysicsControlNameRecords& NameRecords,
 	TFunctor& CreateOperationLambda)
 {
 	for (const FName ChildBoneName : LimbBones.BoneNames)
@@ -150,7 +150,7 @@ template<typename TFunctor> void CreateBodyModifiersFromLimbBones(
 	const FName                   LimbName,
 	const FPhysicsControlLimbBones& LimbBones,
 	const FPhysicsControlModifierData& ModifierData,
-	FRigidBodyNameRecords& NameRecords,
+	FPhysicsControlNameRecords& NameRecords,
 	TFunctor& CreateOperationLambda)
 {
 	for (const FName BoneName : LimbBones.BoneNames)
@@ -170,7 +170,7 @@ template<typename TFunctor> void CreateBodyModifiersFromLimbBones(
 }
 
 //======================================================================================================================
-template<typename TControlFunctorType, typename TBodyModifierFunctorType> void ForEachPotentialOperator(const FAnimNode_RigidBodyWithControl* const Node, TMap<FName, FPhysicsControlLimbBones> AllLimbBones, const FReferenceSkeleton& RefSkeleton, UPhysicsAsset* const PhysicsAsset, FRigidBodyNameRecords& NameRecords, TControlFunctorType& ControlFunctor, TBodyModifierFunctorType& BodyModifierFunctor)
+template<typename TControlFunctorType, typename TBodyModifierFunctorType> void ForEachPotentialOperator(const FAnimNode_RigidBodyWithControl* const Node, TMap<FName, FPhysicsControlLimbBones> AllLimbBones, const FReferenceSkeleton& RefSkeleton, UPhysicsAsset* const PhysicsAsset, FPhysicsControlNameRecords& NameRecords, TControlFunctorType& ControlFunctor, TBodyModifierFunctorType& BodyModifierFunctor)
 {
 	if (Node != nullptr)
 	{
@@ -200,7 +200,7 @@ template<typename TControlFunctorType, typename TBodyModifierFunctorType> void F
 
 //======================================================================================================================
 // Slightly annoying to have to add the names individually, but we want to check they exist
-template<typename TBodyModifierNameContainerType, typename TControlNameContainerType> void CreateAdditionalSets_Implimentation(const FPhysicsControlSetUpdates& AdditionalSets, const TBodyModifierNameContainerType& BodyModifierNames, const TControlNameContainerType& ControlNames, FRigidBodyNameRecords& NameRecords)
+template<typename TBodyModifierNameContainerType, typename TControlNameContainerType> void CreateAdditionalSets_Implementation(const FPhysicsControlSetUpdates& AdditionalSets, const TBodyModifierNameContainerType& BodyModifierNames, const TControlNameContainerType& ControlNames, FPhysicsControlNameRecords& NameRecords)
 {
 	for (const FPhysicsControlSetUpdate& Set : AdditionalSets.ControlSetUpdates)
 	{
@@ -351,7 +351,7 @@ TMap<FName, FPhysicsControlLimbBones> GetLimbBones(
 	return Result;
 }
 
-void CollectOperatorNames(const FAnimNode_RigidBodyWithControl* const Node, TMap<FName, FPhysicsControlLimbBones> AllLimbBones, const FReferenceSkeleton& RefSkeleton, UPhysicsAsset* const PhysicsAsset, TSet<FName>& BodyModifierNames, TSet<FName>& ControlNames, FRigidBodyNameRecords& NameRecords)
+void CollectOperatorNames(const FAnimNode_RigidBodyWithControl* const Node, TMap<FName, FPhysicsControlLimbBones> AllLimbBones, const FReferenceSkeleton& RefSkeleton, UPhysicsAsset* const PhysicsAsset, TSet<FName>& BodyModifierNames, TSet<FName>& ControlNames, FPhysicsControlNameRecords& NameRecords)
 {
 	auto CollectControlName = [&ControlNames](const FName ParentBoneName, const FName ChildBoneName, const FPhysicsControlData& Data) -> FName { const FName Name = FindUniqueControlName(ParentBoneName, ChildBoneName, ControlNames); ControlNames.Add(Name); return Name; };
 	auto CollectBodyModifierName = [&BodyModifierNames](const FName BoneName, const FPhysicsControlModifierData& Data) -> FName { const FName Name = FindUniqueBodyModifierName(BoneName, BodyModifierNames); BodyModifierNames.Add(Name); return Name; };
@@ -359,7 +359,7 @@ void CollectOperatorNames(const FAnimNode_RigidBodyWithControl* const Node, TMap
 	ForEachPotentialOperator(Node, AllLimbBones, RefSkeleton, PhysicsAsset, NameRecords, CollectControlName, CollectBodyModifierName);
 }
 
-void CreateOperatorsForNode(FAnimNode_RigidBodyWithControl* const Node, TMap<FName, FPhysicsControlLimbBones> AllLimbBones, const FReferenceSkeleton& RefSkeleton, UPhysicsAsset* const PhysicsAsset, FRigidBodyNameRecords& NameRecords)
+void CreateOperatorsForNode(FAnimNode_RigidBodyWithControl* const Node, TMap<FName, FPhysicsControlLimbBones> AllLimbBones, const FReferenceSkeleton& RefSkeleton, UPhysicsAsset* const PhysicsAsset, FPhysicsControlNameRecords& NameRecords)
 {
 	auto CreateControl = [Node](const FName ParentBoneName, const FName ChildBoneName, const FPhysicsControlData& Data) -> FName { return Node->CreateControl(ParentBoneName, ChildBoneName, Data); };
 	auto CreateBodyModifier = [Node](const FName BoneName, const FPhysicsControlModifierData& Data) -> FName { return Node->CreateBodyModifier(BoneName, Data); };
@@ -367,12 +367,12 @@ void CreateOperatorsForNode(FAnimNode_RigidBodyWithControl* const Node, TMap<FNa
 	ForEachPotentialOperator(Node, AllLimbBones, RefSkeleton, PhysicsAsset, NameRecords, CreateControl, CreateBodyModifier);
 }
 
-void CreateAdditionalSets(const FPhysicsControlSetUpdates& AdditionalSets, const TSet<FName>& BodyModifierNames, const TSet<FName>& ControlNames, FRigidBodyNameRecords& NameRecords)
+void CreateAdditionalSets(const FPhysicsControlSetUpdates& AdditionalSets, const TSet<FName>& BodyModifierNames, const TSet<FName>& ControlNames, FPhysicsControlNameRecords& NameRecords)
 {
-	CreateAdditionalSets_Implimentation(AdditionalSets, BodyModifierNames, ControlNames, NameRecords);
+	CreateAdditionalSets_Implementation(AdditionalSets, BodyModifierNames, ControlNames, NameRecords);
 }
 
-void CreateAdditionalSets(const FPhysicsControlSetUpdates& AdditionalSets, const TMap<FName, FRigidBodyModifierRecord>& BodyModifierNames, const TMap<FName, FRigidBodyControlRecord>& ControlNames, FRigidBodyNameRecords& NameRecords)
+void CreateAdditionalSets(const FPhysicsControlSetUpdates& AdditionalSets, const TMap<FName, FRigidBodyModifierRecord>& BodyModifierNames, const TMap<FName, FRigidBodyControlRecord>& ControlNames, FPhysicsControlNameRecords& NameRecords)
 {
-	CreateAdditionalSets_Implimentation(AdditionalSets, BodyModifierNames, ControlNames, NameRecords);
+	CreateAdditionalSets_Implementation(AdditionalSets, BodyModifierNames, ControlNames, NameRecords);
 }
diff --git a/Engine/Plugins/Experimental/PhysicsControl/Source/PhysicsControl/Private/PhysicsControlProfileAsset.cpp b/Engine/Plugins/Experimental/PhysicsControl/Source/PhysicsControl/Private/PhysicsControlProfileAsset.cpp
--- a/Engine/Plugins/Experimental/PhysicsControl/Source/PhysicsControl/Private/PhysicsControlProfileAsset.cpp
+++ b/Engine/Plugins/Experimental/PhysicsControl/Source/PhysicsControl/Private/PhysicsControlProfileAsset.cpp
@@ -5,14 +5,19 @@
 
 //======================================================================================================================
 UPhysicsControlProfileAsset::UPhysicsControlProfileAsset()
-	: TestValue(1.0f)
 {
 }
 
 //======================================================================================================================
 void UPhysicsControlProfileAsset::Log()
 {
-	UE_LOG(LogTemp, Log, TEXT("%f"), TestValue)
+	UE_LOG(LogTemp, Log, TEXT("Number of initial controls = %d"), InitialControls.Num());
+	UE_LOG(LogTemp, Log, TEXT("Number of initial body modifiers = %d"), InitialBodyModifiers.Num());
+}
+
+//======================================================================================================================
+void UPhysicsControlProfileAsset::MakeControlsAndModifiersFromCharacterSetupData()
+{
 }
 
 #if WITH_EDITOR
diff --git a/Engine/Plugins/Experimental/PhysicsControl/Source/PhysicsControl/Private/PhysicsControlRecord.cpp b/Engine/Plugins/Experimental/PhysicsControl/Source/PhysicsControl/Private/PhysicsControlRecord.cpp
--- a/Engine/Plugins/Experimental/PhysicsControl/Source/PhysicsControl/Private/PhysicsControlRecord.cpp
+++ b/Engine/Plugins/Experimental/PhysicsControl/Source/PhysicsControl/Private/PhysicsControlRecord.cpp
@@ -8,31 +8,44 @@
 #include "PhysicsEngine/BodyInstance.h"
 
 //======================================================================================================================
-void FPhysicsControlState::Reset()
+void FPhysicsControlRecord::ResetConstraint()
 {
 	if (ConstraintInstance)
 	{
 		ConstraintInstance->TermConstraint();
 	}
-	*this = FPhysicsControlState();
+	ConstraintInstance.Reset();
 }
 
 //======================================================================================================================
-bool FPhysicsControlRecord::InitConstraint(UObject* ConstraintDebugOwner, FName ControlName)
+FVector FPhysicsControlRecord::GetControlPoint() const
 {
-	if (!PhysicsControlState.ConstraintInstance)
+	if (PhysicsControl.ControlData.bUseCustomControlPoint)
 	{
-		return false;
+		return PhysicsControl.ControlData.CustomControlPoint;
 	}
 
-	FConstraintInstance* ConstraintInstance = PhysicsControlState.ConstraintInstance.Get();
+	FBodyInstance* ChildBodyInstance = UE::PhysicsControlComponent::GetBodyInstance(
+		PhysicsControl.ChildMeshComponent.Get(), PhysicsControl.ChildBoneName);
+
+	return ChildBodyInstance ? ChildBodyInstance->GetMassSpaceLocal().GetTranslation() : FVector::ZeroVector;
+}
+
+//======================================================================================================================
+bool FPhysicsControlRecord::InitConstraint(UObject* ConstraintDebugOwner, FName ControlName)
+{
+	if (!ConstraintInstance)
+	{
+		ConstraintInstance = MakeShared<FConstraintInstance>();
+	}
+	check(ConstraintInstance);
 
 	FBodyInstance* ParentBody = UE::PhysicsControlComponent::GetBodyInstance(
-		PhysicsControl.ParentMeshComponent, PhysicsControl.ParentBoneName);
+		PhysicsControl.ParentMeshComponent.Get(), PhysicsControl.ParentBoneName);
 	FBodyInstance* ChildBody = UE::PhysicsControlComponent::GetBodyInstance(
-		PhysicsControl.ChildMeshComponent, PhysicsControl.ChildBoneName);
+		PhysicsControl.ChildMeshComponent.Get(), PhysicsControl.ChildBoneName);
 
-	if (PhysicsControl.ParentMeshComponent && !PhysicsControl.ParentBoneName.IsNone() && !ParentBody)
+	if (PhysicsControl.ParentMeshComponent.IsValid() && !PhysicsControl.ParentBoneName.IsNone() && !ParentBody)
 	{
 		UE_LOG(LogPhysicsControlComponent, Warning,
 			TEXT("Failed to find expected parent body %s when making constraint for control %s"),
@@ -40,7 +53,7 @@ bool FPhysicsControlRecord::InitConstraint(UObject* ConstraintDebugOwner, FName
 			*ControlName.ToString());
 		return false;
 	}
-	if (PhysicsControl.ChildMeshComponent && !PhysicsControl.ChildBoneName.IsNone() && !ChildBody)
+	if (PhysicsControl.ChildMeshComponent.IsValid() && !PhysicsControl.ChildBoneName.IsNone() && !ChildBody)
 	{
 		UE_LOG(LogPhysicsControlComponent, Warning,
 			TEXT("Failed to find expected child body %s when making constraint for control %s"),
@@ -50,7 +63,7 @@ bool FPhysicsControlRecord::InitConstraint(UObject* ConstraintDebugOwner, FName
 	}
 
 	ConstraintInstance->InitConstraint(ChildBody, ParentBody, 1.0f, ConstraintDebugOwner);
-	ConstraintInstance->SetDisableCollision(PhysicsControl.ControlSettings.bDisableCollision);
+	ConstraintInstance->SetDisableCollision(PhysicsControl.ControlData.bDisableCollision);
 	// These things won't change so set them once here
 	ConstraintInstance->SetLinearXMotion(ELinearConstraintMotion::LCM_Free);
 	ConstraintInstance->SetLinearYMotion(ELinearConstraintMotion::LCM_Free);
@@ -71,33 +84,16 @@ bool FPhysicsControlRecord::InitConstraint(UObject* ConstraintDebugOwner, FName
 	return true;
 }
 
-//======================================================================================================================
-// A constraint created with identity frames will just have the child frame at the mesh origin,
-// which is not necessarily where the center of gravity is.
-FConstraintInstance* FPhysicsControlRecord::CreateConstraint(UObject* ConstraintDebugOwner, FName ControlName)
-{
-	if (!PhysicsControlState.ConstraintInstance)
-	{
-		PhysicsControlState.ConstraintInstance = MakeShared<FConstraintInstance>();
-	}
-	FConstraintInstance* ConstraintInstance = PhysicsControlState.ConstraintInstance.Get();
-
-	InitConstraint(ConstraintDebugOwner, ControlName);
-
-	return ConstraintInstance;
-}
-
 //======================================================================================================================
 // Note that, by default, the constraint frames are simply identity. We only modify Frame1, which 
 // corresponds to the child frame. Frame2 will always be identity, because we never change it.
 void FPhysicsControlRecord::UpdateConstraintControlPoint()
 {
-	FConstraintInstance* ConstraintInstance = PhysicsControlState.ConstraintInstance.Get();
 	if (ConstraintInstance)
 	{
 		// Constraints are child then parent
 		FTransform Frame1 = ConstraintInstance->GetRefFrame(EConstraintFrame::Frame1);
-		Frame1.SetTranslation(PhysicsControl.ControlSettings.ControlPoint);
+		Frame1.SetTranslation(GetControlPoint());
 		ConstraintInstance->SetRefFrame(EConstraintFrame::Frame1, Frame1);
 	}
 }
@@ -105,12 +101,7 @@ void FPhysicsControlRecord::UpdateConstraintControlPoint()
 //======================================================================================================================
 void FPhysicsControlRecord::ResetControlPoint()
 {
-	FBodyInstance* ChildBodyInstance = UE::PhysicsControlComponent::GetBodyInstance(
-		PhysicsControl.ChildMeshComponent, PhysicsControl.ChildBoneName);
-
-	PhysicsControl.ControlSettings.ControlPoint = ChildBodyInstance
-		? ChildBodyInstance->GetMassSpaceLocal().GetTranslation() : FVector::ZeroVector;
-
+	PhysicsControl.ControlData.bUseCustomControlPoint = false;
 	UpdateConstraintControlPoint();
 }
 
diff --git a/Engine/Plugins/Experimental/PhysicsControl/Source/PhysicsControl/Private/PhysicsControlRecord.h b/Engine/Plugins/Experimental/PhysicsControl/Source/PhysicsControl/Private/PhysicsControlRecord.h
deleted file mode 100644
diff --git a/Engine/Plugins/Experimental/PhysicsControl/Source/PhysicsControl/Private/RigidBodyControlData.cpp b/Engine/Plugins/Experimental/PhysicsControl/Source/PhysicsControl/Private/RigidBodyControlData.cpp
--- a/Engine/Plugins/Experimental/PhysicsControl/Source/PhysicsControl/Private/RigidBodyControlData.cpp
+++ b/Engine/Plugins/Experimental/PhysicsControl/Source/PhysicsControl/Private/RigidBodyControlData.cpp
@@ -1,86 +1,14 @@
 // Copyright Epic Games, Inc. All Rights Reserved.
 
 #include "RigidBodyControlData.h"
-
-#define INTERPOLATE_PARAM(NAME) Output.NAME = FMath::Lerp(A.NAME, B.NAME, Weight)
-#define SET_ENABLED_PARAM(NAME) Output.bEnable##NAME = A.bEnable##NAME && B.bEnable##NAME
-
-//======================================================================================================================
-FPhysicsControlData Interpolate(
-	const FPhysicsControlData& A, const FPhysicsControlData& B, const float Weight)
-{
-	FPhysicsControlData Output;
-	INTERPOLATE_PARAM(LinearStrength);
-	INTERPOLATE_PARAM(LinearDampingRatio);
-	INTERPOLATE_PARAM(LinearExtraDamping);
-	INTERPOLATE_PARAM(AngularStrength);
-	INTERPOLATE_PARAM(AngularDampingRatio);
-	INTERPOLATE_PARAM(AngularExtraDamping);
-	INTERPOLATE_PARAM(LinearTargetVelocityMultiplier);
-	INTERPOLATE_PARAM(AngularTargetVelocityMultiplier);
-	Output.bEnabled = (Weight < 0.5f) ? A.bEnabled : B.bEnabled;
-	return Output;
-}
-
-//======================================================================================================================
-FPhysicsControlSparseData Interpolate(
-	const FPhysicsControlSparseData& A, const FPhysicsControlSparseData& B, const float Weight)
-{
-	FPhysicsControlSparseData Output;
-	INTERPOLATE_PARAM(LinearStrength);
-	INTERPOLATE_PARAM(LinearDampingRatio);
-	INTERPOLATE_PARAM(LinearExtraDamping);
-	INTERPOLATE_PARAM(AngularStrength);
-	INTERPOLATE_PARAM(AngularDampingRatio);
-	INTERPOLATE_PARAM(AngularExtraDamping);
-	INTERPOLATE_PARAM(LinearTargetVelocityMultiplier);
-	INTERPOLATE_PARAM(AngularTargetVelocityMultiplier);
-	Output.bEnabled = (Weight < 0.5f) ? A.bEnabled : B.bEnabled;
-
-	SET_ENABLED_PARAM(LinearStrength);
-	SET_ENABLED_PARAM(LinearDampingRatio);
-	SET_ENABLED_PARAM(LinearExtraDamping);
-	SET_ENABLED_PARAM(AngularStrength);
-	SET_ENABLED_PARAM(AngularDampingRatio);
-	SET_ENABLED_PARAM(AngularExtraDamping);
-	SET_ENABLED_PARAM(LinearTargetVelocityMultiplier);
-	SET_ENABLED_PARAM(AngularTargetVelocityMultiplier);
-	SET_ENABLED_PARAM(bEnabled);
-
-	return Output;
-}
+#include "Physics/ImmediatePhysics/ImmediatePhysicsActorHandle.h"
 
 //======================================================================================================================
-FPhysicsControlModifierData Interpolate(
-	const FPhysicsControlModifierData& A, const FPhysicsControlModifierData& B, const float Weight)
-{
-	FPhysicsControlModifierData Output;
-	Output.MovementType = (Weight < 0.5f) ? A.MovementType : B.MovementType;
-	INTERPOLATE_PARAM(GravityMultiplier);
-	return Output;
-}
-
-//======================================================================================================================
-FPhysicsControlModifierSparseData Interpolate(
-	const FPhysicsControlModifierSparseData& A, const FPhysicsControlModifierSparseData& B, const float Weight)
-{
-	FPhysicsControlModifierSparseData Output;
-	Output.MovementType = (Weight < 0.5f) ? A.MovementType : B.MovementType;
-	INTERPOLATE_PARAM(GravityMultiplier);
-
-	SET_ENABLED_PARAM(MovementType);
-	SET_ENABLED_PARAM(GravityMultiplier);
-	return Output;
-}
-
-#undef INTERPOLATE_PARAM
-#undef SET_ENABLED_PARAM
-
-//======================================================================================================================
-FRigidBodyControlRecord::FRigidBodyControlRecord(const FRigidBodyControl& InControl, ImmediatePhysics::FJointHandle* InJointHandle)
+FRigidBodyControlRecord::FRigidBodyControlRecord(
+	const FRigidBodyControl& InControl, ImmediatePhysics::FJointHandle* InJointHandle)
 	: Control(InControl)
 	, JointHandle(InJointHandle)
-	, CurrentData(Control.ControlData)
+	, ControlData(Control.ControlData)
 	, ChildBodyIndex(-1)
 	, ParentBodyIndex(-1)
 {
@@ -89,7 +17,7 @@ FRigidBodyControlRecord::FRigidBodyControlRecord(const FRigidBodyControl& InCont
 //======================================================================================================================
 void FRigidBodyControlRecord::ResetCurrent(bool bResetTarget)
 {
-	CurrentData = Control.ControlData;
+	ControlData = Control.ControlData;
 
 	if (bResetTarget)
 	{
@@ -98,16 +26,27 @@ void FRigidBodyControlRecord::ResetCurrent(bool bResetTarget)
 }
 
 //======================================================================================================================
-FRigidBodyModifierRecord::FRigidBodyModifierRecord(const FRigidBodyModifier& InModifier, ImmediatePhysics::FActorHandle* InActorHandle)
+FVector FRigidBodyControlRecord::GetControlPoint(const ImmediatePhysics::FActorHandle* ChildActorHandle) const
+{
+	if (ControlData.bUseCustomControlPoint)
+	{
+		return ControlData.CustomControlPoint;
+	}
+	return ChildActorHandle->GetLocalCoMLocation();
+}
+
+//======================================================================================================================
+FRigidBodyModifierRecord::FRigidBodyModifierRecord(
+	const FRigidBodyModifier& InModifier, ImmediatePhysics::FActorHandle* InActorHandle)
 	: Modifier(InModifier)
 	, ActorHandle(InActorHandle)
-	, CurrentData(Modifier.ModifierData)
+	, ModifierData(Modifier.ModifierData)
 {
 }
 
 //======================================================================================================================
 void FRigidBodyModifierRecord::ResetCurrent()
 {
-	CurrentData = Modifier.ModifierData;
+	ModifierData = Modifier.ModifierData;
 }
 
diff --git a/Engine/Plugins/Experimental/PhysicsControl/Source/PhysicsControl/Private/RigidBodyNameRecords.cpp b/Engine/Plugins/Experimental/PhysicsControl/Source/PhysicsControl/Private/RigidBodyNameRecords.cpp
deleted file mode 100644
diff --git a/Engine/Plugins/Experimental/PhysicsControl/Source/PhysicsControl/Private/RigidBodyWithControl.cpp b/Engine/Plugins/Experimental/PhysicsControl/Source/PhysicsControl/Private/RigidBodyWithControl.cpp
--- a/Engine/Plugins/Experimental/PhysicsControl/Source/PhysicsControl/Private/RigidBodyWithControl.cpp
+++ b/Engine/Plugins/Experimental/PhysicsControl/Source/PhysicsControl/Private/RigidBodyWithControl.cpp
@@ -50,17 +50,32 @@ static void SetConstraintEnabled(ImmediatePhysics::FJointHandle* const JointHand
 }
 
 //======================================================================================================================
-template<typename TRecord, typename TData> void SetRecordParameters(
-	const FName Name, const TData& Data, TMap<FName, TRecord>& Records)
+static void SetRecordParameters(
+	const FName Name, const FPhysicsControlModifierSparseData& Data, TMap<FName, FRigidBodyModifierRecord>& Records)
 {
-	if (TRecord* const RecordSearchResult = Records.Find(Name))
+	if (FRigidBodyModifierRecord* const RecordSearchResult = Records.Find(Name))
 	{
-		RecordSearchResult->CurrentData.UpdateFromSparseData(Data);
+		RecordSearchResult->ModifierData.UpdateFromSparseData(Data);
 	}
 	else
 	{
 		UE_LOG(LogRigidBodyWithControl, Warning,
-			TEXT("SetRecordParameters: Failed to find control/modifier with name %s"), *Name.ToString());
+			TEXT("SetRecordParameters: Failed to find modifier with name %s"), *Name.ToString());
+	}
+}
+
+//======================================================================================================================
+static void SetRecordParameters(
+	const FName Name, const FPhysicsControlSparseData& Data, TMap<FName, FRigidBodyControlRecord>& Records)
+{
+	if (FRigidBodyControlRecord* const RecordSearchResult = Records.Find(Name))
+	{
+		RecordSearchResult->ControlData.UpdateFromSparseData(Data);
+	}
+	else
+	{
+		UE_LOG(LogRigidBodyWithControl, Warning,
+			TEXT("SetRecordParameters: Failed to find control with name %s"), *Name.ToString());
 	}
 }
 
@@ -431,7 +446,8 @@ static void ConvertStrengthToSpringParams(
 
 //======================================================================================================================
 // Adjusts the constraint spring drive settings to reflect the control data
-static void UpdateDriveSpringDamperSettings(
+// Returns true if there is some control, false if the drive has no effect
+static bool UpdateDriveSpringDamperSettings(
 	Chaos::FPBDJointConstraintHandle* Constraint,
 	const Chaos::FPBDJointSettings&   Settings, 
 	const FPhysicsControlData&        ControlData)
@@ -446,9 +462,14 @@ static void UpdateDriveSpringDamperSettings(
 	ConvertStrengthToSpringParams(AngularSpring, AngularDamping, 
 		ControlData.AngularStrength, ControlData.AngularDampingRatio, ControlData.AngularExtraDamping);
 
-	Constraint->SetDriveProperties(
-		Chaos::FVec3(LinearSpring), Chaos::FVec3(LinearDamping), Chaos::FVec3(0),
-		Chaos::FVec3(AngularSpring), Chaos::FVec3(AngularDamping), Chaos::FVec3(0));
+	Constraint->SetDriveParams(
+		Chaos::FVec3(LinearSpring), Chaos::FVec3(LinearDamping), Chaos::FVec3(ControlData.MaxForce),
+		Chaos::FVec3(AngularSpring), Chaos::FVec3(AngularDamping), Chaos::FVec3(ControlData.MaxTorque));
+
+	bool bHaveAngular = (AngularSpring + AngularDamping) > 0;
+	bool bHaveLinear = (LinearSpring + LinearDamping) > 0;
+	return bHaveLinear || bHaveAngular;
+
 }
 
 //======================================================================================================================
@@ -484,82 +505,88 @@ void FAnimNode_RigidBodyWithControl::ApplyControl(FRigidBodyControlRecord& Contr
 		{
 			if (ControlRecord.ExpectedUpdateCounter.Get() != PoseData.UpdateCounter.Get())
 			{
+				// If we missed some intermediate updates, then we don't want to use the previous
+				// positions etc to calculate velocities. This will mean velocity/damping will be
+				// incorrect for one frame, but that's probably OK.
 				DeltaTime = 0.0f;
 			}
 
-			const FActorHandle* ChildActorHandle = JointHandle->GetActorHandles()[ConstraintChildIndex];
-			const FActorHandle* ParentActorHandle = JointHandle->GetActorHandles()[ConstraintParentIndex];
+			Constraint->SetCollisionEnabled(!ControlRecord.Control.ControlData.bDisableCollision);
+			Constraint->SetParentInvMassScale(ControlRecord.Control.ControlData.bOnlyControlChildObject ? 0 : 1);
 
-			if (ChildActorHandle && ParentActorHandle)
+			const Chaos::FPBDJointSettings& JointSettings = Constraint->GetSettings();
+			if (UpdateDriveSpringDamperSettings(Constraint, JointSettings, ControlRecord.ControlData))
 			{
-				const Chaos::FPBDJointSettings& JointSettings = Constraint->GetSettings();
+				const FActorHandle* ChildActorHandle = JointHandle->GetActorHandles()[ConstraintChildIndex];
+				const FActorHandle* ParentActorHandle = JointHandle->GetActorHandles()[ConstraintParentIndex];
 
-				// TODO
-				// - cache settings / previous input parameters to avoid unnecessary repeating
-				//   calculations and making physics API calls every update.
-
-				// Update the target point on the child
-				if (ControlRecord.ControlTarget.bUseTargetPoint)
-				{
-					Constraint->SetChildConnectorLocation(ControlRecord.ControlTarget.TargetPoint);
-				}
-				else
+				if (ChildActorHandle && ParentActorHandle)
 				{
-					const FVector ChildCoMPositionOffset = ChildActorHandle->GetLocalCoMLocation();
-					Constraint->SetChildConnectorLocation(ChildCoMPositionOffset);
-				}
 
-				checkSlow(FindBodyIndexFromBoneName(ControlRecord.Control.ChildBoneName) == ControlRecord.ChildBodyIndex);
-				checkSlow((ControlRecord.Control.ParentBoneName.IsNone() ? -1 : 
-					FindBodyIndexFromBoneName(ControlRecord.Control.ParentBoneName)) == ControlRecord.ParentBodyIndex);
+					// TODO
+					// - cache settings / previous input parameters to avoid unnecessary repeating
+					//   calculations and making physics API calls every update.
 
-				RigidBodyWithControl::FPosQuat TargetTM(
-					ControlRecord.ControlTarget.TargetOrientation, 
-					ControlRecord.ControlTarget.TargetPosition);
+					// Update the target point on the child
+					Constraint->SetChildConnectorLocation(ControlRecord.GetControlPoint(ChildActorHandle));
 
-				if (ControlRecord.ControlTarget.bUseSkeletalAnimation)
-				{
-					RigidBodyWithControl::FPosQuat AnimTargetTM = CalculateTargetTM(
-						JointSettings, PoseData, ControlRecord.ParentBodyIndex, ControlRecord.ChildBodyIndex);
-					TargetTM = TargetTM * AnimTargetTM;
-				}
+					checkSlow(FindBodyIndexFromBoneName(ControlRecord.Control.ChildBoneName) == ControlRecord.ChildBodyIndex);
+					checkSlow((ControlRecord.Control.ParentBoneName.IsNone() ? -1 : 
+						FindBodyIndexFromBoneName(ControlRecord.Control.ParentBoneName)) == ControlRecord.ParentBodyIndex);
 
-				Constraint->SetLinearDrivePositionTarget(TargetTM.GetTranslation());
-				Constraint->SetAngularDrivePositionTarget(TargetTM.GetRotation());
+					RigidBodyWithControl::FPosQuat TargetTM(
+						ControlRecord.ControlTarget.TargetOrientation, 
+						ControlRecord.ControlTarget.TargetPosition);
 
-				if ((DeltaTime * ControlRecord.CurrentData.LinearTargetVelocityMultiplier) != 0)
-				{
-					FVector Velocity = (TargetTM.GetTranslation() - ControlRecord.PrevTargetTM.GetTranslation()) / DeltaTime;
-					Constraint->SetLinearDriveVelocityTarget(
-						Velocity * ControlRecord.CurrentData.LinearTargetVelocityMultiplier);
-				}
-				else
-				{
-					Constraint->SetLinearDriveVelocityTarget(Chaos::FVec3(0));
-				}
+					if (ControlRecord.ControlData.bUseSkeletalAnimation)
+					{
+						RigidBodyWithControl::FPosQuat AnimTargetTM = CalculateTargetTM(
+							JointSettings, PoseData, ControlRecord.ParentBodyIndex, ControlRecord.ChildBodyIndex);
+						TargetTM = TargetTM * AnimTargetTM;
+					}
 
-				if ((DeltaTime * ControlRecord.CurrentData.AngularTargetVelocityMultiplier) != 0)
-				{
-					// Note that quats multiply in the opposite order to TMs, and must be in the same hemisphere.
-					const FQuat Q = TargetTM.GetRotation();
-					FQuat PrevQ = ControlRecord.PrevTargetTM.GetRotation();
-					PrevQ.EnforceShortestArcWith(Q);
-					const FQuat DeltaQ = Q * PrevQ.Inverse();
-					const FVector AngularVelocity = DeltaQ.ToRotationVector() / DeltaTime;
-
-					Constraint->SetAngularDriveVelocityTarget( 
-						AngularVelocity * ControlRecord.CurrentData.AngularTargetVelocityMultiplier);
+					Constraint->SetLinearDrivePositionTarget(TargetTM.GetTranslation());
+					Constraint->SetAngularDrivePositionTarget(TargetTM.GetRotation());
+
+					if ((DeltaTime * ControlRecord.ControlData.LinearTargetVelocityMultiplier) != 0)
+					{
+						FVector Velocity = (TargetTM.GetTranslation() - ControlRecord.PrevTargetTM.GetTranslation()) / DeltaTime;
+						Constraint->SetLinearDriveVelocityTarget(
+							Velocity * ControlRecord.ControlData.LinearTargetVelocityMultiplier);
+					}
+					else
+					{
+						Constraint->SetLinearDriveVelocityTarget(Chaos::FVec3(0));
+					}
+
+					if ((DeltaTime * ControlRecord.ControlData.AngularTargetVelocityMultiplier) != 0)
+					{
+						// Note that quats multiply in the opposite order to TMs, and must be in the same hemisphere.
+						const FQuat Q = TargetTM.GetRotation();
+						FQuat PrevQ = ControlRecord.PrevTargetTM.GetRotation();
+						PrevQ.EnforceShortestArcWith(Q);
+						const FQuat DeltaQ = Q * PrevQ.Inverse();
+						const FVector AngularVelocity = DeltaQ.ToRotationVector() / DeltaTime;
+
+						Constraint->SetAngularDriveVelocityTarget( 
+							AngularVelocity * ControlRecord.ControlData.AngularTargetVelocityMultiplier);
+					}
+					else
+					{
+						Constraint->SetAngularDriveVelocityTarget(Chaos::FVec3(0));
+					}
+
+
+					ControlRecord.PrevTargetTM = TargetTM;
+					ControlRecord.ExpectedUpdateCounter = PoseData.UpdateCounter;
+					ControlRecord.ExpectedUpdateCounter.Increment();
 				}
 				else
 				{
-					Constraint->SetAngularDriveVelocityTarget(Chaos::FVec3(0));
+					// Note that if we don't have any strength, then we don't calculate the targets.
+					// However, make sure that we don't apply velocities using the wrong calculation
+					// when the strength/damping is increased in the future
 				}
-
-				UpdateDriveSpringDamperSettings(Constraint, JointSettings, ControlRecord.CurrentData);
-
-				ControlRecord.PrevTargetTM = TargetTM;
-				ControlRecord.ExpectedUpdateCounter = PoseData.UpdateCounter;
-				ControlRecord.ExpectedUpdateCounter.Increment();
 			}
 		}
 	}
@@ -574,19 +601,23 @@ void FAnimNode_RigidBodyWithControl::ApplyModifier(
 	{
 		// Note that there's an early out if there's no change needed, so this should be OK.
 		BodyModifierRecord.ActorHandle->SetIsKinematic(
-			BodyModifierRecord.CurrentData.MovementType != EPhysicsMovementType::Simulated);
+			BodyModifierRecord.ModifierData.MovementType != EPhysicsMovementType::Simulated);
 
 		// Note that the actual kinematic targets will be set separately, since they need to be set
 		// for all kinematics whether or not they were under a modifier.
 
 		// Scale gravity
-		float GravityMultiplier = BodyModifierRecord.CurrentData.GravityMultiplier;
+		float GravityMultiplier = BodyModifierRecord.ModifierData.GravityMultiplier;
 		if (GravityMultiplier != 1 && BodyModifierRecord.ActorHandle->IsGravityEnabled())
 		{
 			float Mass = (float) BodyModifierRecord.ActorHandle->GetMass();
 			FVector AntiGravityForce = SimSpaceGravity * (-Mass * (1.0f - GravityMultiplier));
 			BodyModifierRecord.ActorHandle->AddForce(AntiGravityForce);
 		}
+
+		// Set collision
+		Chaos::FGeometryParticleHandle* ParticleHandle = BodyModifierRecord.ActorHandle->GetParticle();
+		ParticleHandle->SetHasCollision(CollisionEnabledHasPhysics(BodyModifierRecord.ModifierData.CollisionType));
 	}
 }
 
@@ -598,8 +629,7 @@ void FAnimNode_RigidBodyWithControl::ApplyControlAndModifierUpdatesAndParameters
 	SCOPE_CYCLE_COUNTER(STAT_RigidBodyNodeWithControl_ApplyControlAndModifierUpdatesAndParametersToRecords);
 
 	// Apply control and modifier parameters on a single-use basis
-	ApplyControlAndBodyModifierDatas(
-		Updates.ControlParameters, Updates.ModifierParameters);
+	ApplyControlAndBodyModifierDatas(Updates.ControlUpdates, Updates.ModifierUpdates);
 
 	// This goes through the records, resetting the update parts.
 	// Then the update structures get adjusted based on the parameters.
@@ -631,8 +661,8 @@ void FAnimNode_RigidBodyWithControl::ApplyControlAndModifierUpdatesAndParameters
 
 //======================================================================================================================
 void FAnimNode_RigidBodyWithControl::ApplyControlAndBodyModifierDatas(
-	const TArray<FPhysicsControlNamedControlParameters>& InControlParameters,
-	const TArray<FPhysicsControlNamedModifierParameters>& InModifierParameters)
+	const TArray<FPhysicsControlNamedControlParameters>&           InControlParameters,
+	const TArray<FPhysicsControlNamedModifierParameters>&          InModifierParameters)
 {
 	SCOPE_CYCLE_COUNTER(STAT_RigidBodyNodeWithControl_ApplyControlsAndModifierDatas);
 
diff --git a/Engine/Plugins/Experimental/PhysicsControl/Source/PhysicsControl/Public/AnimNode_RigidBodyWithControl.h b/Engine/Plugins/Experimental/PhysicsControl/Source/PhysicsControl/Public/AnimNode_RigidBodyWithControl.h
--- a/Engine/Plugins/Experimental/PhysicsControl/Source/PhysicsControl/Public/AnimNode_RigidBodyWithControl.h
+++ b/Engine/Plugins/Experimental/PhysicsControl/Source/PhysicsControl/Public/AnimNode_RigidBodyWithControl.h
@@ -7,15 +7,15 @@
 #pragma once
 
 #include "CoreMinimal.h"
-#include "RigidBodyPoseData.h"
-#include "RigidBodyControlData.h"
-#include "PhysicsControlLimbData.h"
-#include "RigidBodyNameRecords.h"
 #include "BoneControllers/AnimNode_SkeletalControlBase.h"
 #include "Components/SkeletalMeshComponent.h"
 #include "Physics/ImmediatePhysics/ImmediatePhysicsDeclares.h"
+#include "PhysicsControlLimbData.h"
+#include "PhysicsControlNameRecords.h"
 #include "PhysicsEngine/PhysicsAsset.h"
 #include "PhysicsProxy/PerSolverFieldSystem.h"
+#include "RigidBodyControlData.h"
+#include "RigidBodyPoseData.h"
 #include "Tasks/Task.h"
 
  // Use the simulation space functions from the RBAN
@@ -64,17 +64,37 @@ struct PHYSICSCONTROL_API FAnimNode_RigidBodyWithControl : public FAnimNode_Skel
 
 	virtual void AddImpulseAtLocation(FVector Impulse, FVector Location, FName BoneName = NAME_None);
 
-	// TEMP: Exposed for use in PhAt as a quick way to get drag handles working with Chaos
+	// TEMP: Exposed for use in PhAT as a quick way to get drag handles working with Chaos
 	virtual ImmediatePhysics::FSimulation* GetSimulation() { return PhysicsSimulation; }
 
+	/**
+	 * Set the override physics asset. This will automatically trigger a physics re-init if the 
+	 * override physics asset changes. Users can get access to this in the Animation Blueprint 
+	 * via the Animation Node Functions.
+	 */
+	void SetOverridePhysicsAsset(UPhysicsAsset* PhysicsAsset);
+
 	UPhysicsAsset* GetPhysicsAsset() const { return PhysicsAssetToUse; }
 
 public:
 	/** Physics asset to use. If empty use the skeletal mesh's default physics asset */
 	UPROPERTY(EditAnywhere, Category = Settings)
-	TObjectPtr<UPhysicsAsset> OverridePhysicsAsset;
+	TWeakObjectPtr<UPhysicsAsset> OverridePhysicsAsset;
+
+	/** 
+	 * Use the skeletal mesh physics asset as default in case set to True. The Override Physics Asset 
+	 * will always have priority over this. 
+	 */
+	UPROPERTY(EditAnywhere, Category = Settings)
+	bool bDefaultToSkeletalMeshPhysicsAsset = true;
 
 private:
+	/** 
+	 * Get the physics asset candidate to be used while respecting the bDefaultToSkeletalMeshPhysicsAsset 
+	 * and the priority to the override physics asset. 
+	 */
+	UPhysicsAsset* GetPhysicsAssetToBeUsed(const UAnimInstance* InAnimInstance) const;
+
 	FTransform PreviousCompWorldSpaceTM;
 	FTransform CurrentTransform;
 	FTransform PreviousTransform;
@@ -298,7 +318,7 @@ public:
 
 	const int32 GetNumBodies() const;
 
-	const FRigidBodyNameRecords& GetNameRecords() const { return NameRecords; }
+	const FPhysicsControlNameRecords& GetNameRecords() const { return NameRecords; }
 
 private:
 
@@ -499,7 +519,7 @@ private:
 	TMap<FName, FRigidBodyModifierRecord> ModifierRecords;
 
 	// Details about sets etc
-	FRigidBodyNameRecords NameRecords;
+	FPhysicsControlNameRecords NameRecords;
 
 	FName CurrentConstraintProfile;
 
diff --git a/Engine/Plugins/Experimental/PhysicsControl/Source/PhysicsControl/Public/PhysicsControlActor.h b/Engine/Plugins/Experimental/PhysicsControl/Source/PhysicsControl/Public/PhysicsControlActor.h
--- a/Engine/Plugins/Experimental/PhysicsControl/Source/PhysicsControl/Public/PhysicsControlActor.h
+++ b/Engine/Plugins/Experimental/PhysicsControl/Source/PhysicsControl/Public/PhysicsControlActor.h
@@ -22,7 +22,7 @@ struct PHYSICSCONTROL_API FInitialPhysicsControl
 
 	/** The owner of the mesh that will be doing the driving. Blank/non-existent means it will happen in world space */
 	UPROPERTY(EditAnywhere, BlueprintReadWrite, Category = PhysicsControl)
-	TObjectPtr<AActor> ParentActor;
+	TWeakObjectPtr<AActor> ParentActor;
 
 	/** 
 	 * The mesh that will be doing the driving. If this is blank but there is an actor, then we'll attempt to
@@ -40,7 +40,7 @@ struct PHYSICSCONTROL_API FInitialPhysicsControl
 
 	/** The owner of the mesh that the control will be driving */
 	UPROPERTY(EditAnywhere, BlueprintReadWrite, Category = PhysicsControl)
-	TObjectPtr<AActor> ChildActor;
+	TWeakObjectPtr<AActor> ChildActor;
 
 	/** 
 	 * The mesh that the control will be driving. If this is blank but there is an actor, then we'll attempt 
@@ -77,12 +77,6 @@ struct PHYSICSCONTROL_API FInitialPhysicsControl
 	 */
 	UPROPERTY(EditAnywhere, BlueprintReadWrite, Category = PhysicsControl)
 	FPhysicsControlTarget ControlTarget;
-
-	/**
-	 * More general settings for the control
-	 */
-	UPROPERTY(EditAnywhere, BlueprintReadWrite, Category = PhysicsControl)
-	FPhysicsControlSettings ControlSettings;
 };
 
 /**
@@ -97,7 +91,7 @@ struct PHYSICSCONTROL_API FInitialBodyModifier
 
 	/** The owner of the mesh that that we will modify */
 	UPROPERTY(EditAnywhere, BlueprintReadWrite, Category = PhysicsControl)
-	TObjectPtr<AActor> Actor;
+	TWeakObjectPtr<AActor> Actor;
 
 	/** 
 	 * The mesh that will be modify. If this is blank but there is an actor, then we'll attempt to
@@ -112,19 +106,9 @@ struct PHYSICSCONTROL_API FInitialBodyModifier
 	UPROPERTY(EditAnywhere, BlueprintReadWrite, Category = PhysicsControl)
 	FName BoneName;
 
-	/** How the associated body should move. */
+	/** How the body should move etc */
 	UPROPERTY(EditAnywhere, BlueprintReadWrite, Category = PhysicsControl)
-	EPhysicsMovementType MovementType = EPhysicsMovementType::Kinematic;
-
-	/** How the associated body should collide/interact */
-	ECollisionEnabled::Type CollisionType = ECollisionEnabled::QueryAndPhysics;
-
-	/**
-	 * Multiplier for gravity applied to the body. Note that if the body itself has gravity disabled, then
-	 * setting this to 1 will not enable gravity.
-	 */
-	UPROPERTY(EditAnywhere, BlueprintReadWrite, Category = PhysicsControl)
-	float GravityMultiplier = 1.0f;
+	FPhysicsControlModifierData BodyModifierData;
 
 	/**
 	 * The target position when kinematic. Note that this is applied on top of any animation
@@ -138,11 +122,7 @@ struct PHYSICSCONTROL_API FInitialBodyModifier
 	 * target if bUseSkeletalAnimation is set.
 	 */
 	UPROPERTY(EditAnywhere, BlueprintReadWrite, Category = PhysicsControl)
-	FQuat KinematicTargetOrientation = FQuat::Identity;
-
-	/** If true then the target will be applied on top of the skeletal animation (if there is any) */
-	UPROPERTY(EditAnywhere, BlueprintReadWrite, Category = PhysicsControl)
-	bool bUseSkeletalAnimation = true;
+	FRotator KinematicTargetOrientation = FRotator::ZeroRotator;
 };
 
 /**
@@ -158,7 +138,7 @@ struct PHYSICSCONTROL_API FInitialCharacterControls
 
 	/** The owner of the character skeletal mesh */
 	UPROPERTY(EditAnywhere, BlueprintReadWrite, Category = PhysicsControl)
-	TObjectPtr<AActor> CharacterActor;
+	TWeakObjectPtr<AActor> CharacterActor;
 
 	/**
 	 * The skeletal mesh that will have controls set up. If this is blank but there is an actor, then we'll attempt to
@@ -173,17 +153,11 @@ struct PHYSICSCONTROL_API FInitialCharacterControls
 	UPROPERTY(EditAnywhere, BlueprintReadWrite, Category = PhysicsControl)
 	FPhysicsControlData WorldSpaceControlData;
 
-	UPROPERTY(EditAnywhere, BlueprintReadWrite, Category = PhysicsControl)
-	FPhysicsControlSettings WorldSpaceControlSettings;
-
 	UPROPERTY(EditAnywhere, BlueprintReadWrite, Category = PhysicsControl)
 	FPhysicsControlData ParentSpaceControlData;
 
 	UPROPERTY(EditAnywhere, BlueprintReadWrite, Category = PhysicsControl)
-	FPhysicsControlSettings ParentSpaceControlSettings;
-
-	UPROPERTY(EditAnywhere, BlueprintReadWrite, Category = PhysicsControl)
-	EPhysicsMovementType PhysicsMovementType = EPhysicsMovementType::Kinematic;
+	FPhysicsControlModifierData BodyModifierData;
 };
 
 /**
diff --git a/Engine/Plugins/Experimental/PhysicsControl/Source/PhysicsControl/Public/PhysicsControlBPLibrary.h b/Engine/Plugins/Experimental/PhysicsControl/Source/PhysicsControl/Public/PhysicsControlBPLibrary.h
--- a/Engine/Plugins/Experimental/PhysicsControl/Source/PhysicsControl/Public/PhysicsControlBPLibrary.h
+++ b/Engine/Plugins/Experimental/PhysicsControl/Source/PhysicsControl/Public/PhysicsControlBPLibrary.h
@@ -103,6 +103,10 @@ class PHYSICSCONTROL_API UPhysicsControlBPLibrary : public UBlueprintFunctionLib
 		Result = (ConversionResult == EAnimNodeReferenceConversionResult::Succeeded);
 	}
 
+	/** Set the physics asset on the rigid body with control anim graph node (RBWC). */
+	UFUNCTION(BlueprintCallable, Category = "Animation|PhysicsControl", meta = (BlueprintThreadSafe))
+	static FRigidBodyWithControlReference SetOverridePhysicsAsset(const FRigidBodyWithControlReference& Node, UPhysicsAsset* PhysicsAsset);
+
 	/** Get the names of all the controls in a specified set managed by this Rigid Body With Control node. */
 	UFUNCTION(BlueprintPure, Category = "Animation|PhysicsControl", meta = (BlueprintThreadSafe))
 	static TArray<FName> GetControlNamesInSet(
diff --git a/Engine/Plugins/Experimental/PhysicsControl/Source/PhysicsControl/Public/PhysicsControlComponent.h b/Engine/Plugins/Experimental/PhysicsControl/Source/PhysicsControl/Public/PhysicsControlComponent.h
--- a/Engine/Plugins/Experimental/PhysicsControl/Source/PhysicsControl/Public/PhysicsControlComponent.h
+++ b/Engine/Plugins/Experimental/PhysicsControl/Source/PhysicsControl/Public/PhysicsControlComponent.h
@@ -4,6 +4,8 @@
 
 #include "PhysicsControlData.h"
 #include "PhysicsControlLimbData.h"
+#include "PhysicsControlNameRecords.h"
+#include "PhysicsControlRecord.h"
 
 #include "UObject/ObjectMacros.h"
 #include "Components/SceneComponent.h"
@@ -88,7 +90,6 @@ public:
 	 * 
 	 * @param ControlData   Describes the initial strength etc of the new control
 	 * @param ControlTarget Describes the initial target for the new control
-	 * @param ControlSettings General settings for the control
 	 * @param Set Which set to include the control in (optional). Note that it automatically
 	 *        gets added to the set "All"
 	 * @param bEnabled If true then the control will be enabled immediately. If false you will need to call
@@ -104,7 +105,6 @@ public:
 		const FName                   ChildBoneName,
 		const FPhysicsControlData     ControlData,
 		const FPhysicsControlTarget   ControlTarget, 
-		const FPhysicsControlSettings ControlSettings,
 		FName                         Set,
 		FString                       NamePrefix = TEXT("")
 	);
@@ -115,7 +115,6 @@ public:
 	 * @param The name of the control that will be created. Creation will fail if this name is already in use.
 	 * @param ControlData   Describes the initial strength etc of the new control
 	 * @param ControlTarget Describes the initial target for the new control
-	 * @param ControlSettings General settings for the control
 	 * @param Set Which set to include the control in (optional). Note that it automatically
 	 *        gets added to the set "All"
 	 * @param bEnabled If true then the control will be enabled immediately. If false you will need to call
@@ -131,7 +130,6 @@ public:
 		const FName                   ChildBoneName,
 		const FPhysicsControlData     ControlData, 
 		const FPhysicsControlTarget   ControlTarget, 
-		const FPhysicsControlSettings ControlSettings, 
 		const FName                   Set);
 
 	/**
@@ -142,7 +140,6 @@ public:
 	 * @param bIncludeSelf Whether or not to include BoneName when creating controls
 	 * @param ControlType What type of control to create. This determines what the parent will be for each control
 	 * @param ControlData   Describes the initial strength etc of the new control
-	 * @param ControlSettings General settings for the control
 	 * @param Set Which set to include the control in (optional). Note that it automatically
 	 *        gets added to the set "All"
 	 * @param bEnabled If true then the control will be enabled immediately. If false you will need to call
@@ -156,7 +153,6 @@ public:
 		const bool                    bIncludeSelf,
 		const EPhysicsControlType     ControlType,
 		const FPhysicsControlData     ControlData,
-		const FPhysicsControlSettings ControlSettings,
 		const FName                   Set);
 
 	/**
@@ -174,7 +170,7 @@ public:
 	 *                          profile if it can't be found) - though they will not match exactly if the linear
 	 *                          drive and different x/y/z values, or if the angular drive was using twist/swing
 	 *                          instead of slerp. Note also that the joint constraints do not use the animation
-	 *                          velocity as a target, so when creating controls in this way the control settings
+	 *                          velocity as a target, so when creating controls in this way the control 
 	 *                          will set the skeletal animation velocity multiplier to zero.
 	 * @param Set Which set to include the control in (optional). Note that it automatically
 	 *        gets added to the set "All"
@@ -198,7 +194,6 @@ public:
 	 * @param BoneNames The names of bones for which controls should be created. Each bone will be the child in a control
 	 * @param ControlType What type of control to create. This determines what the parent will be for each control
 	 * @param ControlData   Describes the initial strength etc of the new control
-	 * @param ControlSettings General settings for the control
 	 * @param Set Which set to include the control in (optional). Note that it automatically
 	 *        gets added to the set "All"
 	 * @param bEnabled If true then the control will be enabled immediately. If false you will need to call
@@ -211,7 +206,6 @@ public:
 		const TArray<FName>&          BoneNames,
 		const EPhysicsControlType     ControlType,
 		const FPhysicsControlData     ControlData,
-		const FPhysicsControlSettings ControlSettings,
 		const FName                   Set);
 
 	/**
@@ -228,7 +222,7 @@ public:
 	 *                          profile if it can't be found) - though they will not match exactly if the linear
 	 *                          drive and different x/y/z values, or if the angular drive was using twist/swing
 	 *                          instead of slerp. Note also that the joint constraints do not use the animation
-	 *                          velocity as a target, so when creating controls in this way the control settings
+	 *                          velocity as a target, so when creating controls in this way the control 
 	 *                          will set the skeletal animation velocity multiplier to zero.
 	 * @param Set Which set to include the control in (optional). Note that it automatically
 	 *        gets added to the set "All"
@@ -268,7 +262,6 @@ public:
 	 *        using GetLimbBonesFromSkeletalMesh
 	 * @param ControlType What type of control to create. This determines what the parent will be for each control
 	 * @param ControlData Describes the initial strength etc of the new control
-	 * @param ControlSettings General settings for the control
 	 * @param WorldComponent Optional component to use as the parent object for any "world-space" controls 
 	 *        that are created. Will be ignored if the controls being created are not world-space.
 	 * @param WorldBoneName Additional bone name to identify the world object if the WorldComponent is actually 
@@ -282,7 +275,6 @@ public:
 		const TMap<FName, FPhysicsControlLimbBones>& LimbBones,
 		const EPhysicsControlType                    ControlType,
 		const FPhysicsControlData                    ControlData,
-		const FPhysicsControlSettings                ControlSettings,
 		UMeshComponent*                              WorldComponent = nullptr,
 		FName                                        WorldBoneName = NAME_None,
 		FString                                      NamePrefix = TEXT(""));
@@ -302,7 +294,7 @@ public:
 	 *                          profile if it can't be found) - though they will not match exactly if the linear
 	 *                          drive and different x/y/z values, or if the angular drive was using twist/swing
 	 *                          instead of slerp. Note also that the joint constraints do not use the animation 
-	 *                          velocity as a target, so when creating controls in this way the control settings
+	 *                          velocity as a target, so when creating controls in this way the control 
 	 *                          will set the skeletal animation velocity multiplier to zero.
 	 * @param bEnabled If true then the control will be enabled immediately. If false you will need to call
 	 *                 SetControlEnabled(true) in order to enable it.
@@ -371,9 +363,8 @@ public:
 	/**
 	 * Modifies an existing control data - i.e. the strengths etc of the control driving towards the target
 	 *
-	 * @param Name The name of the control to modify. 
+	 * @param Name The name of the control to update. 
 	 * @param ControlData The new control data
-	 * @param bEnableControl Enables the control if it is currently disabled
 	 * @return true if the control was found and modified, false if not
 	 */
 	UFUNCTION(BlueprintCallable, Category = PhysicsControl)
@@ -382,10 +373,9 @@ public:
 	/**
 	 * Modifies existing control data - i.e. the strengths etc of the controls driving towards the targets
 	 *
-	 * @param Names The names of the controls to modify. Note that if you have these in a FPhysicsControlNameArray
+	 * @param Names The names of the controls to update. Note that if you have these in a FPhysicsControlNameArray
 	 *              then it can be split.
 	 * @param ControlData The new control data
-	 * @param bEnableControl Enables the control if it is currently disabled
 	 */
 	UFUNCTION(BlueprintCallable, Category = PhysicsControl)
 	void SetControlDatas(const TArray<FName>& Names, const FPhysicsControlData ControlData);
@@ -393,14 +383,43 @@ public:
 	/**
 	 * Modifies existing control data - i.e. the strengths etc of the controls driving towards the targets
 	 *
-	 * @param Set The set of controls to modify. Standard sets will include "All", "WorldSpace",
+	 * @param Set The set of controls to update. Standard sets will include "All", "WorldSpace",
 	 *        "ParentSpace" and things like "WorldSpace-ArmLeft", depending on how controls have been created.
 	 * @param ControlData The new control data
-	 * @param bEnableControl Enables the controls if currently disabled
 	 */
 	UFUNCTION(BlueprintCallable, Category = PhysicsControl)
 	void SetControlDatasInSet(const FName Set, const FPhysicsControlData ControlData);
 
+	/**
+	 * Modifies an existing control data - i.e. the strengths etc of the control driving towards the target
+	 *
+	 * @param Name The name of the control to update. 
+	 * @param ControlData The new control data
+	 * @return true if the control was found and modified, false if not
+	 */
+	UFUNCTION(BlueprintCallable, Category = PhysicsControl)
+	bool SetControlSparseData(const FName Name, const FPhysicsControlSparseData ControlData);
+
+	/**
+	 * Modifies existing control data - i.e. the strengths etc of the controls driving towards the targets
+	 *
+	 * @param Names The names of the controls to update. Note that if you have these in a FPhysicsControlNameArray
+	 *              then it can be split.
+	 * @param ControlData The new control data
+	 */
+	UFUNCTION(BlueprintCallable, Category = PhysicsControl)
+	void SetControlSparseDatas(const TArray<FName>& Names, const FPhysicsControlSparseData ControlData);
+
+	/**
+	 * Modifies existing control data - i.e. the strengths etc of the controls driving towards the targets
+	 *
+	 * @param Set The set of controls to update. Standard sets will include "All", "WorldSpace",
+	 *        "ParentSpace" and things like "WorldSpace-ArmLeft", depending on how controls have been created.
+	 * @param ControlData The new control data
+	 */
+	UFUNCTION(BlueprintCallable, Category = PhysicsControl)
+	void SetControlSparseDatasInSet(const FName Set, const FPhysicsControlSparseData ControlData);
+
 	/**
 	 * Modifies an existing control data using the multipliers
 	 *
@@ -410,7 +429,8 @@ public:
 	 * @return true if the control was found and modified, false if not
 	 */
 	UFUNCTION(BlueprintCallable, Category = PhysicsControl)
-	bool SetControlMultiplier(const FName Name, const FPhysicsControlMultiplier ControlMultiplier, const bool bEnableControl = true);
+	bool SetControlMultiplier
+	(const FName Name, const FPhysicsControlMultiplier ControlMultiplier, const bool bEnableControl = true);
 
 	/**
 	 * Modifies existing control data using the multipliers
@@ -440,6 +460,46 @@ public:
 		const FPhysicsControlMultiplier  ControlMultiplier, 
 		const bool                       bEnableControl = true);
 
+	/**
+	 * Modifies an existing control data using the multipliers
+	 *
+	 * @param Name The name of the control to modify. 
+	 * @param ControlMultipliers The new control multipliers
+	 * @param bEnableControl Enables the control if it is currently disabled
+	 * @return true if the control was found and modified, false if not
+	 */
+	UFUNCTION(BlueprintCallable, Category = PhysicsControl)
+	bool SetControlSparseMultiplier(
+		const FName Name, const FPhysicsControlSparseMultiplier ControlMultiplier, const bool bEnableControl = true);
+
+	/**
+	 * Modifies existing control data using the multipliers
+	 *
+	 * @param Names The names of the controls to modify. Note that if you have these in a FPhysicsControlNameArray
+	 *              then it can be split.
+	 * @param ControlMultiplier The new control multiplier
+	 * @param bEnableControl Enables the controls if currently disabled
+	 */
+	UFUNCTION(BlueprintCallable, Category = PhysicsControl)
+	void SetControlSparseMultipliers(
+		const TArray<FName>&                  Names, 
+		const FPhysicsControlSparseMultiplier ControlMultiplier, 
+		const bool                            bEnableControl = true);
+
+	/**
+	 * Modifies existing control data using the multipliers
+	 *
+	 * @param Set The set of controls to modify. Standard sets will include "All", "WorldSpace",
+	 *        "ParentSpace" and things like "WorldSpace-ArmLeft", depending on how controls have been created.
+	 * @param ControlMultiplier The new control multiplier
+	 * @param bEnableControl Enables the controls if currently disabled
+	 */
+	UFUNCTION(BlueprintCallable, Category = PhysicsControl)
+	void SetControlSparseMultipliersInSet(
+		const FName                           Set, 
+		const FPhysicsControlSparseMultiplier ControlMultiplier, 
+		const bool                            bEnableControl = true);
+
 	/**
 	 * Modifies an existing control's linear data - i.e. the strengths etc of the control driving towards the target
 	 *
@@ -885,30 +945,14 @@ public:
 	 * @param BoneName The bone name, if a skeletal mesh is used
 	 * @param Set Which set to include the body modifier in (optional). Note that it automatically 
 	 *        gets added to the set "All"
-	 * @param MovementType Whether to enable/disable simulation on the body
-	 * @param CollisionType Collision type to set on the body
-	 * @param GravityMultiplier The amount of gravity to apply when simulating
-	 * @param PhysicsBlendWeight The blend weight between the body transform coming from
-	 *        animation and that coming from simulation.
-	 * @param bUseSkeletalAnimation Whether the kinematic target is specified in the frame of the
-	 *        skeletal animation, rather than world space. Only relevant if the body is part of 
-	 *        a skeletal mesh.
-	 * @param bUpdateKinematicFromSimulation Whether the body should be updated from the simulation when
-	 *        it is kinematic, or whether it should track the kinematic target directly. This will be most
-	 *        likely useful when using async physics, in order to make kinematic parts behave the same as 
-	 *        dynamic ones.
+	 * @param BodyModifierData The initial properties of the modifier
 	 */
 	UFUNCTION(BlueprintCallable, Category = PhysicsControl)
 	FName CreateBodyModifier(
-		UMeshComponent*               MeshComponent,
-		const FName                   BoneName,
-		const FName                   Set,
-		const EPhysicsMovementType    MovementType = EPhysicsMovementType::Simulated,
-		const ECollisionEnabled::Type CollisionType = ECollisionEnabled::QueryAndPhysics,
-		const float                   GravityMultiplier = 1.0f,
-		const float                   PhysicsBlendWeight = 1.0f,
-		const bool                    bUseSkeletalAnimation = true,
-		const bool                    bUpdateKinematicFromSimulation = true);
+		UMeshComponent*                   MeshComponent,
+		const FName                       BoneName,
+		const FName                       Set,
+		const FPhysicsControlModifierData BodyModifierData);
 
 	/**
 	 * Creates a new body modifier for mesh components
@@ -918,31 +962,15 @@ public:
 	 * @param BoneName The bone name, if a skeletal mesh is used
 	 * @param Set Which set to include the body modifier in (optional). Note that it automatically
 	 *        gets added to the set "All"
-	 * @param MovementType Whether to enable/disable simulation on the body
-	 * @param CollisionType Collision type to set on the body
-	 * @param GravityMultiplier The amount of gravity to apply when simulating
-	 * @param PhysicsBlendWeight The blend weight between the body transform coming from
-	 *        animation and that coming from simulation.
-	 * @param bUseSkeletalAnimation Whether the kinematic target is specified in the frame of the
-	 *        skeletal animation, rather than world space. Only relevant if the body is part of 
-	 *        a skeletal mesh.
-	 * @param bUpdateKinematicFromSimulation Whether the body should be updated from the simulation when
-	 *        it is kinematic, or whether it should track the kinematic target directly. This will be most
-	 *        likely useful when using async physics, in order to make kinematic parts behave the same as
-	 *        dynamic ones.
+	 * @param BodyModifierData The initial properties of the modifier
 	 */
 	UFUNCTION(BlueprintCallable, Category = PhysicsControl)
 	bool CreateNamedBodyModifier(
-		const FName                   Name,
-		UMeshComponent*               MeshComponent,
-		const FName                   BoneName,
-		const FName                   Set,
-		const EPhysicsMovementType    MovementType = EPhysicsMovementType::Simulated,
-		const ECollisionEnabled::Type CollisionType = ECollisionEnabled::QueryAndPhysics,
-		const float                   GravityMultiplier = 1.0f,
-		const float                   PhysicsBlendWeight = 1.0f,
-		const bool                    bUseSkeletalAnimation = true,
-		const bool                    bUpdateKinematicFromSimulation = true);
+		const FName                       Name,
+		UMeshComponent*                   MeshComponent,
+		const FName                       BoneName,
+		const FName                       Set,
+		const FPhysicsControlModifierData BodyModifierData);
 
 	/**
 	 * Creates new body modifiers for skeletal mesh components
@@ -952,31 +980,15 @@ public:
 	 * @param bIncludeSelf Whether or not to include BoneName when creating modifiers
 	 * @param Set Which set to include the body modifier in (optional). Note that it automatically
 	 *        gets added to the set "All"
-	 * @param MovementType Whether to enable/disable simulation on the body
-	 * @param CollisionType Collision type to set on the body
-	 * @param GravityMultiplier The amount of gravity to apply when simulating
-	 * @param PhysicsBlendWeight The blend weight between the body transform coming from
-	 *        animation and that coming from simulation.
-	 * @param bUseSkeletalAnimation Whether the kinematic target is specified in the frame of the
-	 *        skeletal animation, rather than world space. Only relevant if the body is part of
-	 *        a skeletal mesh.
-	 * @param bUpdateKinematicFromSimulation Whether the body should be updated from the simulation when
-	 *        it is kinematic, or whether it should track the kinematic target directly. This will be most
-	 *        likely useful when using async physics, in order to make kinematic parts behave the same as
-	 *        dynamic ones.
+	 * @param BodyModifierData The initial properties of the modifier
 	 */
 	UFUNCTION(BlueprintCallable, Category = PhysicsControl)
 	TArray<FName> CreateBodyModifiersFromSkeletalMeshBelow(
-		USkeletalMeshComponent*       SkeletalMeshComponent,
-		const FName                   BoneName,
-		const bool                    bIncludeSelf,
-		const FName                   Set,
-		const EPhysicsMovementType    MovementType = EPhysicsMovementType::Simulated,
-		const ECollisionEnabled::Type CollisionType = ECollisionEnabled::QueryAndPhysics,
-		const float                   GravityMultiplier = 1.0f,
-		const float                   PhysicsBlendWeight = 1.0f,
-		const bool                    bUseSkeletalAnimation = true,
-		const bool                    bUpdateKinematicFromSimulation = true);
+		USkeletalMeshComponent*           SkeletalMeshComponent,
+		const FName                       BoneName,
+		const bool                        bIncludeSelf,
+		const FName                       Set,
+		const FPhysicsControlModifierData BodyModifierData);
 
 
 	/**
@@ -985,18 +997,7 @@ public:
 	 * @param AllControls A single container for all the controls that have been created
 	 * @param LimbBones A map relating the limbs and the bones that they contain. Typically create this 
 	 *                  using GetLimbBonesFromSkeletalMesh
-	 * @param MovementType Whether to enable/disable simulation on the body
-	 * @param CollisionType Collision type to set on the body
-	 * @param GravityMultiplier The amount of gravity to apply when simulating
-	 * @param PhysicsBlendWeight The blend weight between the body transform coming from
-	 *        animation and that coming from simulation.
-	 * @param bUseSkeletalAnimation Whether the kinematic target is specified in the frame of the
-	 *        skeletal animation, rather than world space. Only relevant if the body is part of
-	 *        a skeletal mesh.
-	 * @param bUpdateKinematicFromSimulation Whether the body should be updated from the simulation when
-	 *        it is kinematic, or whether it should track the kinematic target directly. This will be most
-	 *        likely useful when using async physics, in order to make kinematic parts behave the same as
-	 *        dynamic ones.
+	 * @param BodyModifierData The initial properties of the modifier
 	 *
 	 * @return A map containing the modifiers for each limb
 	 */
@@ -1004,12 +1005,7 @@ public:
 	TMap<FName, FPhysicsControlNames> CreateBodyModifiersFromLimbBones(
 		FPhysicsControlNames&                        AllBodyModifiers,
 		const TMap<FName, FPhysicsControlLimbBones>& LimbBones,
-		const EPhysicsMovementType                   MovementType = EPhysicsMovementType::Simulated,
-		const ECollisionEnabled::Type                CollisionType = ECollisionEnabled::QueryAndPhysics,
-		const float                                  GravityMultiplier = 1.0f,
-		const float                                  PhysicsBlendWeight = 1.0f,
-		const bool                                   bUseSkeletalAnimation = true,
-		const bool                                   bUpdateKinematicFromSimulation = true);
+		const FPhysicsControlModifierData            BodyModifierData);
 
 	/**
 	 * Destroys a BodyModifier
@@ -1038,6 +1034,66 @@ public:
 	UFUNCTION(BlueprintCallable, Category = PhysicsControl)
 	void DestroyBodyModifiersInSet(const FName Set);
 
+	/**
+	 * Modifies an existing Body Modifier Data
+	 *
+	 * @param Name The name of the modifier to update. 
+	 * @param ModifierData The new data
+	 * @return true if the modifier was found and modified, false if not
+	 */
+	UFUNCTION(BlueprintCallable, Category = PhysicsControl)
+	bool SetBodyModifierData(const FName Name, const FPhysicsControlModifierData ModifierData);
+
+	/**
+	 * Modifies existing Body Modifier Data
+	 *
+	 * @param Names The names of the modifiers to update. Note that if you have these in a FPhysicsControlNameArray
+	 *              then it can be split.
+	 * @param ModifierData The new data
+	 */
+	UFUNCTION(BlueprintCallable, Category = PhysicsControl)
+	void SetBodyModifierDatas(const TArray<FName>& Names, const FPhysicsControlModifierData ModifierData);
+
+	/**
+	 * Modifies existing Body Modifier Data
+	 *
+	 * @param Set The set of modifiers to update. Standard sets will include "All" and things like "ArmLeft", 
+	 *        depending on how body modifiers have been created.
+	 * @param ModifierData The new data
+	 */
+	UFUNCTION(BlueprintCallable, Category = PhysicsControl)
+	void SetBodyModifierDatasInSet(const FName Set, const FPhysicsControlModifierData ModifierData);
+
+	/**
+	 * Modifies an existing Body Modifier Data
+	 *
+	 * @param Name The name of the modifier to update. 
+	 * @param ModifierData The new data
+	 * @return true if the modifier was found and modified, false if not
+	 */
+	UFUNCTION(BlueprintCallable, Category = PhysicsControl)
+	bool SetBodyModifierSparseData(const FName Name, const FPhysicsControlModifierSparseData ModifierData);
+
+	/**
+	 * Modifies existing Body Modifier Data
+	 *
+	 * @param Names The names of the modifiers to update. Note that if you have these in a FPhysicsControlNameArray
+	 *              then it can be split.
+	 * @param ModifierData The new data
+	 */
+	UFUNCTION(BlueprintCallable, Category = PhysicsControl)
+	void SetBodyModifierSparseDatas(const TArray<FName>& Names, const FPhysicsControlModifierSparseData ModifierData);
+
+	/**
+	 * Modifies existing Body Modifier Data
+	 *
+	 * @param Set The set of modifiers to update. Standard sets will include "All" and things like "ArmLeft", 
+	 *        depending on how body modifiers have been created.
+	 * @param ModifierData The new data
+	 */
+	UFUNCTION(BlueprintCallable, Category = PhysicsControl)
+	void SetBodyModifierSparseDatasInSet(const FName Set, const FPhysicsControlModifierSparseData ModifierData);
+
 	/**
 	 * Sets the kinematic target transform for a body modifier. 
 	 * 
@@ -1313,12 +1369,8 @@ public:
 		USkeletalMeshComponent*                     SkeletalMeshComponent,
 		const TArray<FPhysicsControlLimbSetupData>& LimbSetupData,
 		const FPhysicsControlData                   WorldSpaceControlData,
-		const FPhysicsControlSettings               WorldSpaceControlSettings,
 		const FPhysicsControlData                   ParentSpaceControlData,
-		const FPhysicsControlSettings               ParentSpaceControlSettings,
-		const EPhysicsMovementType                  PhysicsMovementType = EPhysicsMovementType::Static,
-		const float                                 GravityMultiplier = 1.0f,
-		const float                                 PhysicsBlendWeight = 1.0f,
+		const FPhysicsControlModifierData           BodyModifierData,
 		UMeshComponent*                             WorldComponent = nullptr,
 		FName                                       WorldBoneName = NAME_None);
 
@@ -1564,7 +1616,16 @@ public:
 		const FName                        SetName,
 		const EResetToCachedTargetBehavior Behavior = EResetToCachedTargetBehavior::ResetImmediately);
 
+	/** Indicates if a control with the name exists (doesn't produce a warning if it doesn't) */
+	UFUNCTION(BlueprintCallable, Category = PhysicsControl)
+	bool GetControlExists(const FName Name) const;
+
+	/** Indicates if a body modifier with the name exists (doesn't produce a warning if it doesn't) */
+	UFUNCTION(BlueprintCallable, Category = PhysicsControl)
+	bool GetBodyModifierExists(const FName Name) const;
+
 public:
+	// Public property data
 
 	/**
 	 * If the component moves by more than this distance then it is treated as a teleport,
@@ -1623,13 +1684,13 @@ public:
 	UPROPERTY(EditAnywhere, BlueprintReadWrite, Category = Settings)
 	int32 MaxNumControlsOrModifiersPerName = 256;
 
-protected:
-
-	// Begin UActorComponent Interface
-	virtual void InitializeComponent() override;
-	virtual void BeginDestroy() override;
-	virtual void TickComponent(float DeltaTime, enum ELevelTick TickType, FActorComponentTickFunction* ThisTickFunction) override;
-	// End UActorComponent Interface
+	/** 
+	 * Warn if an an invalid control or body modifier name is used. This can happen quite easily since 
+	 * they're only referenced through names, which are likely auto-generated. However, it may happen for
+	 * valid reasons, in which case you'll want to disable this.
+	 */
+	UPROPERTY(EditAnywhere, BlueprintReadWrite, Category = Debug)
+	bool bWarnAboutInvalidNames = true;
 
 public:
 #if WITH_EDITOR
@@ -1637,13 +1698,183 @@ public:
 	virtual void OnRegister() override;
 	//End ActorComponent interface
 
-	void UpdateSpriteTexture();
 	// Used by the component visualizer
-	void DebugDraw(FPrimitiveDrawInterface* PDI) const;
-	void DebugDrawControl(FPrimitiveDrawInterface* PDI, const FPhysicsControlRecord& Record, const FName ControlName) const;
+	virtual void DebugDraw(FPrimitiveDrawInterface* PDI) const;
+	virtual void DebugDrawControl(FPrimitiveDrawInterface* PDI, const FPhysicsControlRecord& Record, const FName ControlName) const;
 #endif
 
 protected:
 
-	TPimplPtr<FPhysicsControlComponentImpl> Implementation;
+	// Begin UActorComponent Interface
+	virtual void InitializeComponent() override;
+	virtual void BeginDestroy() override;
+	virtual void TickComponent(float DeltaTime, enum ELevelTick TickType, FActorComponentTickFunction* ThisTickFunction) override;
+	// End UActorComponent Interface
+
+
+	/**
+	 * Retrieves the bone data for the specified bone given the skeletal mesh component.
+	 * 
+	 * @param OutBoneData If successful, this will contain the output bone data.
+	 * @param InSkeletalMeshComponent Required to be a valid pointer to a skeletal mesh component
+	 * @param InBoneName The name of the bone to retrieve data for
+	 * @return true if the bone and data were found, false if not (in which case warnings will be logged)
+	 */
+	bool GetBoneData(
+		FCachedSkeletalMeshData::FBoneData& OutBoneData,
+		const USkeletalMeshComponent*       InSkeletalMeshComponent,
+		const FName                         InBoneName) const;
+
+	/*
+	 * Retrieves the bone data for the specified bone given the skeletal mesh component, for modification
+	 *
+	 * @param OutBoneData If successful, this will point to the bone data.
+	 * @param InSkeletalMeshComponent Required to be a valid pointer to a skeletal mesh component
+	 * @param InBoneName The name of the bone to retrieve data for
+	 * @return true if the bone and data were found, false if not (in which case warnings will be logged)
+	 */
+	bool GetModifiableBoneData(
+		FCachedSkeletalMeshData::FBoneData*& OutBoneData,
+		const USkeletalMeshComponent*        InSkeletalMeshComponent,
+		const FName                          InBoneName);
+
+	/**
+	 * Retrieves the control for the name. Returns a null pointer if the name cannot be found
+	 */
+	FPhysicsControl* FindControl(const FName Name);
+	const FPhysicsControl* FindControl(const FName Name) const;
+
+	/**
+	 * Retrieves the control record for the name. Returns a null pointer if the name cannot be found
+	 */
+	FPhysicsControlRecord* FindControlRecord(const FName Name);
+	const FPhysicsControlRecord* FindControlRecord(const FName Name) const;
+
+	/**
+	 * Updates the world-space bone positions etc for each skeleton we're tracking
+	 */
+	void UpdateCachedSkeletalBoneData(float DeltaTime);
+
+	/**
+	 * @return true if the difference between the old and new TMs exceeds the teleport
+	 * translation/rotation thresholds
+	 */
+	bool DetectTeleport(const FTransform& OldTM, const FTransform& NewTM) const;
+
+	/**
+	 * @return true if the difference between the old and new TMs exceeds the teleport
+	 * translation/rotation thresholds
+	 */
+	bool DetectTeleport(
+		const FVector& OldPosition, const FQuat& OldOrientation,
+		const FVector& NewPosition, const FQuat& NewOrientation) const;
+
+	/**
+	 * Terminates the underlying physical constraints, resets our internal stored state for each control,
+	 * and optionally deletes all record of the controls.
+	 */
+	void ResetControls(bool bKeepControlRecords);
+
+	/**
+	 * Starts caching skeletal mesh poses, and registers for a tick pre-requisite
+	 */
+	void AddSkeletalMeshReferenceForCaching(USkeletalMeshComponent* SkeletalMeshComponent);
+
+	/**
+	 * Stops caching skeletal mesh poses (if this is the last one), and deregisters for a tick pre-requisite.
+	 * Returns true if this was the last reference, false otherwise.
+	 */
+	bool RemoveSkeletalMeshReferenceForCaching(USkeletalMeshComponent* SkeletalMeshComponent);
+
+	/**
+	 * Records that a modifier is working with the skeletal mesh and stores original data if necessary
+	 */
+	void AddSkeletalMeshReferenceForModifier(USkeletalMeshComponent* SkeletalMeshComponent);
+
+	/**
+	 * Records that a modifier has stopped working with the skeletal mesh and restores original data if necessary.
+	 * Returns true if this was the last reference, false otherwise.
+	 */
+	bool RemoveSkeletalMeshReferenceForModifier(USkeletalMeshComponent* SkeletalMeshComponent);
+
+	/** Update the constraint based on the record. */
+	void ApplyControl(FPhysicsControlRecord& Record);
+
+	/**
+	 * Updates the constraint strengths. Returns true if there is some strength, 
+	 * false otherwise (e.g. to skip setting targets)
+	 */
+	bool ApplyControlStrengths(FPhysicsControlRecord& Record, FConstraintInstance* ConstraintInstance);
+
+	/**
+	 * Calculates the Target TM and velocities that will get passed to the constraint - so this 
+	 * is a target that is defined in the space of the parent body (or in world space, if it doesn't exist).
+	 */
+	void CalculateControlTargetData(
+		FTransform&                  OutTargetTM, 
+		FVector&                     OutTargetVelocity,
+		FVector&                     OutTargetAngularVelocity,
+		const FPhysicsControlRecord& Record,
+		bool                         bCalculateVelocity) const;
+
+	/** Updates the body based on the modifier */
+	void ApplyBodyModifier(FPhysicsBodyModifier& BodyModifier);
+
+	/**
+	 * This will set the kinematic target for the appropriate body based on the weighted target position and 
+	 * orientation (and whether any were found) for any controls that are related to the body modifier.
+	 */
+	void ApplyKinematicTarget(const FPhysicsBodyModifier& BodyModifier) const;
+
+	/**
+	 * Sets the body (simulated or kinematic) to have the position/velocity etc state that has been cached. 
+	 * Has no effect if there is no cached target data.
+	 */
+	void ResetToCachedTarget(const FPhysicsBodyModifier& BodyModifier) const;
+
+	/**
+	 * Retrieves the body modifier for the name. Note that if Name is blank then the first modifier will be returned,
+	 * assuming there is one.
+	 */
+	FPhysicsBodyModifier* FindBodyModifier(const FName Name);
+	const FPhysicsBodyModifier* FindBodyModifier(const FName Name) const;
+
+	/** 
+	 * When destroying a control or modifier, the record will normally be removed, but it can be retained if you
+	 * will subsequently update the records
+	 */
+	enum class EDestroyBehavior : uint8
+	{
+		KeepRecord,
+		RemoveRecord
+	};
+
+	/** Destroys the control. It will optionally be removed from the array of records too */
+	bool DestroyControl(const FName Name, const EDestroyBehavior DestroyBehavior);
+
+	/** Destroys the modifier. It will optionally be removed from the array of records too */
+	bool DestroyBodyModifier(const FName Name, const EDestroyBehavior DestroyBehavior);
+
+	/** Calculates and returns a unique name based on the bone names and any existing records */
+	FName GetUniqueControlName(const FName ParentBoneName, const FName ChildBoneName, const FString& NamePrefix) const;
+
+	/** Calculates and returns a unique name based on the bone name and any existing records */
+	FName GetUniqueBodyModifierName(const FName BoneName) const;
+
+protected:
+
+	// Cached transforms from each skeletal mesh we're working with. Will be updated at the
+	// beginning of each tick.
+	TMap<TWeakObjectPtr<USkeletalMeshComponent>, FCachedSkeletalMeshData> CachedSkeletalMeshDatas;
+
+	// Track which skeletons have been affected by a body modifier - some settings get overridden
+	// and then need to be restored when the last body modifier is destroyed.
+	TMap<TWeakObjectPtr<USkeletalMeshComponent>, FModifiedSkeletalMeshData> ModifiedSkeletalMeshDatas;
+
+	TMap<FName, FPhysicsControlRecord> PhysicsControlRecords;
+	TMap<FName, FPhysicsBodyModifier> PhysicsBodyModifiers;
+
+	// Keep track of the names of everything we have created
+	FPhysicsControlNameRecords NameRecords;
+
 };
diff --git a/Engine/Plugins/Experimental/PhysicsControl/Source/PhysicsControl/Public/PhysicsControlData.h b/Engine/Plugins/Experimental/PhysicsControl/Source/PhysicsControl/Public/PhysicsControlData.h
--- a/Engine/Plugins/Experimental/PhysicsControl/Source/PhysicsControl/Public/PhysicsControlData.h
+++ b/Engine/Plugins/Experimental/PhysicsControl/Source/PhysicsControl/Public/PhysicsControlData.h
@@ -4,6 +4,7 @@
 
 #include "UObject/ObjectMacros.h"
 #include "Components/ActorComponent.h"
+#include "Components/MeshComponent.h"
 #include "EngineDefines.h"
 #include "PhysicsControlData.generated.h"
 
@@ -12,8 +13,6 @@
  * and the RigidBodyWithControl node.
  */
 
-class UMeshComponent;
-
 /**
  * Used by Body Modifiers to specify how the physical bodies should move.  
  */
@@ -111,26 +110,84 @@ struct PHYSICSCONTROL_API FPhysicsControlModifierSparseData
 	GENERATED_BODY();
 
 	FPhysicsControlModifierSparseData(
-		const EPhysicsMovementType InMovementType = EPhysicsMovementType::Simulated,
-		const float                InGravityMultiplier = 1.0f)
+		const EPhysicsMovementType    InMovementType = EPhysicsMovementType::Simulated,
+		const ECollisionEnabled::Type InCollisionType = ECollisionEnabled::QueryAndPhysics,
+		const float                   InGravityMultiplier = 1.0f,
+		const float                   InPhysicsBlendWeight = 1.0f,
+		const bool                    InUseSkeletalAnimation = true,
+		const bool                    InUpdateKinematicFromSimulation = true)
 		: MovementType(InMovementType)
+		, CollisionType(InCollisionType)
 		, GravityMultiplier(InGravityMultiplier)
-		, bEnableMovementType(1)
-		, bEnableGravityMultiplier(1)
+		, PhysicsBlendWeight(InPhysicsBlendWeight)
+		, bUseSkeletalAnimation(InUseSkeletalAnimation)
+		, bUpdateKinematicFromSimulation(InUpdateKinematicFromSimulation)
+		// The flags
+		, bEnableMovementType(true)
+		, bEnableCollisionType(true)
+		, bEnableGravityMultiplier(true)
+		, bEnablePhysicsBlendWeight(true)
+		, bEnablebUseSkeletalAnimation(true)
+		, bEnablebUpdateKinematicFromSimulation(true)
 	{
 	}
 
+	/**
+	 * How the body should move.
+	 */
 	UPROPERTY(EditAnywhere, BlueprintReadWrite, Category = PhysicsControl, meta = (editcondition = "bEnableMovementType"))
 	EPhysicsMovementType MovementType;
 
+	/**
+	 * The collision type. Note that PhysicsControlComponent uses the full filtering, but
+	 * RigidBodyWithControl just treats this as enable/disable collision
+	 */
+	UPROPERTY(EditAnywhere, BlueprintReadWrite, Category = PhysicsControl, meta = (editcondition = "bEnableCollisionType"))
+	TEnumAsByte<ECollisionEnabled::Type> CollisionType;
+
+	/**
+	 * How much gravity to use when simulating - typically between 0 and 1
+	 */
 	UPROPERTY(EditAnywhere, BlueprintReadWrite, Category = PhysicsControl, meta = (editcondition = "bEnableGravityMultiplier"))
 	float GravityMultiplier;
 
+	/**
+	 * When converting back from simulation to animation, how much to use the simulation output
+	 * versus the original animation input.
+	 * NOTE This is currently only used in the PhysicsControlComponent implementation
+	 */
+	UPROPERTY(EditAnywhere, BlueprintReadWrite, Category = PhysicsControl, meta = (editcondition = "bEnablePhysicsBlendWeight"))
+	float PhysicsBlendWeight;
+
+	/** If true then the target will be applied on top of the skeletal animation (if there is any) */
+	UPROPERTY(EditAnywhere, BlueprintReadWrite, Category = PhysicsControl, meta = (editcondition = "bEnablebUseSkeletalAnimation"))
+	uint8 bUseSkeletalAnimation : 1;
+
+	/**
+	 * If true then the associated actor's transform will be updated from the simulation when it is
+	 * kinematic. This is most likely useful when using async physics in order to prevent different
+	 * parts of the skeleton from being torn apart.
+	 */
+	UPROPERTY(EditAnywhere, BlueprintReadWrite, Category = PhysicsControl, meta = (editcondition = "bEnablebUpdateKinematicFromSimulation"))
+	uint8 bUpdateKinematicFromSimulation : 1;
+
 	UPROPERTY(EditAnywhere, Category = PhysicsControl, meta = (InlineEditConditionToggle))
 	uint8 bEnableMovementType : 1;
 
+	UPROPERTY(EditAnywhere, Category = PhysicsControl, meta = (InlineEditConditionToggle))
+	uint8 bEnableCollisionType : 1;
+
 	UPROPERTY(EditAnywhere, Category = PhysicsControl, meta = (InlineEditConditionToggle))
 	uint8 bEnableGravityMultiplier : 1;
+
+	UPROPERTY(EditAnywhere, Category = PhysicsControl, meta = (InlineEditConditionToggle))
+	uint8 bEnablePhysicsBlendWeight : 1;
+
+	UPROPERTY(EditAnywhere, Category = PhysicsControl, meta = (InlineEditConditionToggle))
+	uint8 bEnablebUseSkeletalAnimation : 1;
+
+	UPROPERTY(EditAnywhere, Category = PhysicsControl, meta = (InlineEditConditionToggle))
+	uint8 bEnablebUpdateKinematicFromSimulation : 1;
 };
 
 /**
@@ -142,20 +199,65 @@ struct PHYSICSCONTROL_API FPhysicsControlModifierData
 	GENERATED_BODY();
 
 	FPhysicsControlModifierData(
-		const EPhysicsMovementType InMovementType = EPhysicsMovementType::Simulated,
-		const float                InGravityMultiplier = 1.0f)
+		const EPhysicsMovementType    InMovementType = EPhysicsMovementType::Simulated,
+		const ECollisionEnabled::Type InCollisionType = ECollisionEnabled::QueryAndPhysics,
+		const float                   InGravityMultiplier = 1.0f,
+		const float                   InPhysicsBlendWeight = 1.0f,
+		const bool                    InUseSkeletalAnimation = true,
+		const bool                    InUpdateKinematicFromSimulation = true)
 		: MovementType(InMovementType)
+		, CollisionType(InCollisionType)
 		, GravityMultiplier(InGravityMultiplier)
+		, PhysicsBlendWeight(InPhysicsBlendWeight)
+		, bUseSkeletalAnimation(InUseSkeletalAnimation)
+		, bUpdateKinematicFromSimulation(InUpdateKinematicFromSimulation)
 	{
 	}
 
 	void UpdateFromSparseData(const FPhysicsControlModifierSparseData& SparseData);
 
+	/**
+	 * How the body should move.
+	 */
 	UPROPERTY(EditAnywhere, BlueprintReadWrite, Category = PhysicsControl)
-	EPhysicsMovementType MovementType;
+	EPhysicsMovementType MovementType = EPhysicsMovementType::Simulated;
 
+	/**
+	 * The collision type. Note that PhysicsControlComponent uses the full filtering, but 
+	 * RigidBodyWithControl just treats this as enable/disable collision
+	 */
 	UPROPERTY(EditAnywhere, BlueprintReadWrite, Category = PhysicsControl)
-	float GravityMultiplier;
+	TEnumAsByte<ECollisionEnabled::Type> CollisionType = ECollisionEnabled::QueryAndPhysics;
+
+	/**
+	 * How much gravity to use when simulating - typically between 0 and 1
+	 */
+	UPROPERTY(EditAnywhere, BlueprintReadWrite, Category = PhysicsControl)
+	float GravityMultiplier = 1.0f;
+
+	/**
+	 * When converting back from simulation to animation, how much to use the simulation output
+	 * versus the original animation input.
+	 * NOTE This is currently only used in the PhysicsControlComponent implementation
+	 */
+	UPROPERTY(EditAnywhere, BlueprintReadWrite, Category = PhysicsControl)
+	float PhysicsBlendWeight = 1.0f;
+
+	/** 
+	 * If true then the target will be applied on top of the skeletal animation (if there is any) 
+	 * NOTE This is currently only used in the PhysicsControlComponent implementation
+	 */
+	UPROPERTY(EditAnywhere, BlueprintReadWrite, Category = PhysicsControl)
+	uint8 bUseSkeletalAnimation : 1;
+
+	/**
+	 * If true then the associated actor's transform will be updated from the simulation when it is
+	 * kinematic. This is most likely useful when using async physics in order to prevent different
+	 * parts of the skeleton from being torn apart.
+	 * NOTE This is currently only used in the PhysicsControlComponent implementation (is it needed in RBWC?)
+	 */
+	UPROPERTY(EditAnywhere, BlueprintReadWrite, Category = PhysicsControl)
+	uint8 bUpdateKinematicFromSimulation : 1;
 };
 
 
@@ -168,24 +270,25 @@ struct PHYSICSCONTROL_API FPhysicsControlSparseData
 {
 	GENERATED_BODY();
 
-	FPhysicsControlSparseData(
-		float InLinearStrength = 0.0f, float InLinearDampingRatio = 1.0f, 
-		float InLinearExtraDamping = 0.0f, float InMaxForce = 0.0f,
-		float InAngularStrength = 0.0f, float InAngularDampingRatio = 1.0f, 
-		float InAngularExtraDamping = 0.0f, float InMaxTorque = 0.0f,
-		float InLinearTargetVelocityMultiplier = 1.0f, float InAngularTargetVelocityMultiplier = 1.0f,
-		bool bInEnabled = true)
-		: LinearStrength(InLinearStrength)
-		, LinearDampingRatio(InLinearDampingRatio)
-		, LinearExtraDamping(InLinearExtraDamping)
-		, MaxForce(InMaxForce)
-		, AngularStrength(InAngularStrength)
-		, AngularDampingRatio(InAngularDampingRatio)
-		, AngularExtraDamping(InAngularExtraDamping)
-		, MaxTorque(InMaxTorque)
-		, LinearTargetVelocityMultiplier(InLinearTargetVelocityMultiplier)
-		, AngularTargetVelocityMultiplier(InAngularTargetVelocityMultiplier)
-		, bEnabled(bInEnabled)
+	FPhysicsControlSparseData()
+		: LinearStrength(0)
+		, LinearDampingRatio(1)
+		, LinearExtraDamping(0)
+		, MaxForce(0)
+		, AngularStrength(0)
+		, AngularDampingRatio(1)
+		, AngularExtraDamping(0)
+		, MaxTorque(0)
+		, LinearTargetVelocityMultiplier(1)
+		, AngularTargetVelocityMultiplier(1)
+		, SkeletalAnimationVelocityMultiplier(1)
+		, CustomControlPoint(FVector::ZeroVector)
+		, bEnabled(true)
+		, bUseCustomControlPoint(false)
+		, bUseSkeletalAnimation(true)
+		, bDisableCollision(false)
+		, bOnlyControlChildObject(false)
+		// the flags
 		, bEnableLinearStrength(true)
 		, bEnableLinearDampingRatio(true)
 		, bEnableLinearExtraDamping(true)
@@ -196,7 +299,13 @@ struct PHYSICSCONTROL_API FPhysicsControlSparseData
 		, bEnableMaxTorque(true)
 		, bEnableLinearTargetVelocityMultiplier(true)
 		, bEnableAngularTargetVelocityMultiplier(true)
+		, bEnableSkeletalAnimationVelocityMultiplier(true)
+		, bEnableCustomControlPoint(true)
 		, bEnablebEnabled(true)
+		, bEnablebUseCustomControlPoint(true)
+		, bEnablebUseSkeletalAnimation(true)
+		, bEnablebDisableCollision(true)
+		, bEnablebOnlyControlChildObject(true)
 	{
 	}
 
@@ -256,6 +365,7 @@ struct PHYSICSCONTROL_API FPhysicsControlSparseData
 	 * Multiplier on the velocity, which gets applied to the damping. A value of 1 means the animation target
 	 * velocity is used, which helps it track the animation. A value of 0 means damping happens in "world space" 
 	 * - so damping acts like drag on the movement.
+	 * NOTE This is currently only used in the RigidBodyWithControl implementation
 	 */
 	UPROPERTY(EditAnywhere, BlueprintReadWrite, Category = PhysicsControl, meta = (ClampMin = "0.0", editcondition = "bEnableLinearTargetVelocityMultiplier"))
 	float LinearTargetVelocityMultiplier;
@@ -264,16 +374,51 @@ struct PHYSICSCONTROL_API FPhysicsControlSparseData
 	 * Multiplier on the angular velocity, which gets applied to the damping. A value of 1 means the animation target
 	 * velocity is used, which helps it track the animation. A value of 0 means damping happens in "world space" 
 	 * - so damping acts like drag on the movement.
+	 * NOTE This is currently only used in the RigidBodyWithControl implementation
 	 */
 	UPROPERTY(EditAnywhere, BlueprintReadWrite, Category = PhysicsControl, meta = (ClampMin = "0.0", editcondition = "bEnableAngularTargetVelocityMultiplier"))
 	float AngularTargetVelocityMultiplier;
 
+	/** 
+	 * The amount of skeletal animation velocity to use in the targets 
+	 * NOTE This is currently only used in the PhysicsControlComponent implementation
+	 */
+	UPROPERTY(EditAnywhere, BlueprintReadWrite, Category = PhysicsControl, meta = (ClampMin = "0.0", editcondition = "bEnableSkeletalAnimationVelocityMultiplier"))
+	float SkeletalAnimationVelocityMultiplier;
+
 	/**
-	 * Whether this control should be enabled
+	 * The position of the control point relative to the child mesh, when using a custom control point. 
 	 */
-	UPROPERTY(EditAnywhere, BlueprintReadWrite, Category = PhysicsControl, meta = (editcondition = "bEnablebEnabled"))
+	UPROPERTY(EditAnywhere, BlueprintReadWrite, Category = PhysicsControl, meta = (ClampMin = "0.0", editcondition = "bEnableCustomControlPoint"))
+	FVector CustomControlPoint;
+
+	/** Whether this control should be enabled */
+	UPROPERTY(EditAnywhere, BlueprintReadWrite, Category = PhysicsControl, meta = (ClampMin = "0.0", editcondition = "bEnablebEnabled"))
 	uint8 bEnabled : 1;
 
+	/** Whether or not to use the custom control point position */
+	UPROPERTY(EditAnywhere, BlueprintReadWrite, Category = PhysicsControl, meta = (ClampMin = "0.0", editcondition = "bEnablebUseCustomControlPoint"))
+	uint8 bUseCustomControlPoint : 1;
+
+	/** If true then the target will be applied on top of the skeletal animation (if there is any) */
+	UPROPERTY(EditAnywhere, BlueprintReadWrite, Category = PhysicsControl, meta = (ClampMin = "0.0", editcondition = "bEnablebUseSkeletalAnimation"))
+	uint8 bUseSkeletalAnimation : 1;
+
+	/**
+	 * Whether or not this control should disable collision between the parent and child bodies (only
+	 * has an effect if there is a parent body)
+	 */
+	UPROPERTY(EditAnywhere, BlueprintReadWrite, Category = PhysicsControl, meta = (ClampMin = "0.0", editcondition = "bEnablebDisableCollision"))
+	uint8 bDisableCollision : 1;
+
+	/** If true then the control will only affect the child object, not the parent */
+	UPROPERTY(EditAnywhere, BlueprintReadWrite, Category = PhysicsControl, meta = (ClampMin = "0.0", editcondition = "bEnablebOnlyControlChildObject"))
+	uint8 bOnlyControlChildObject : 1;
+
+	//==================================================================================================================
+	// Flags for sparse access
+	//==================================================================================================================
+
 	UPROPERTY(EditAnywhere, Category = PhysicsControl, meta = (InlineEditConditionToggle))
 	uint8 bEnableLinearStrength : 1;
 
@@ -304,8 +449,26 @@ struct PHYSICSCONTROL_API FPhysicsControlSparseData
 	UPROPERTY(EditAnywhere, Category = PhysicsControl, meta = (InlineEditConditionToggle))
 	uint8 bEnableAngularTargetVelocityMultiplier : 1;
 
+	UPROPERTY(EditAnywhere, Category = PhysicsControl, meta = (InlineEditConditionToggle))
+	uint8 bEnableSkeletalAnimationVelocityMultiplier : 1;
+
+	UPROPERTY(EditAnywhere, Category = PhysicsControl, meta = (InlineEditConditionToggle))
+	uint8 bEnableCustomControlPoint : 1;
+
 	UPROPERTY(EditAnywhere, Category = PhysicsControl, meta = (InlineEditConditionToggle))
 	uint8 bEnablebEnabled : 1;
+
+	UPROPERTY(EditAnywhere, Category = PhysicsControl, meta = (InlineEditConditionToggle))
+	uint8 bEnablebUseCustomControlPoint: 1;
+
+	UPROPERTY(EditAnywhere, Category = PhysicsControl, meta = (InlineEditConditionToggle))
+	uint8 bEnablebUseSkeletalAnimation : 1;
+
+	UPROPERTY(EditAnywhere, Category = PhysicsControl, meta = (InlineEditConditionToggle))
+	uint8 bEnablebDisableCollision : 1;
+
+	UPROPERTY(EditAnywhere, Category = PhysicsControl, meta = (InlineEditConditionToggle))
+	uint8 bEnablebOnlyControlChildObject : 1;
 };
 
 /**
@@ -328,7 +491,13 @@ struct PHYSICSCONTROL_API FPhysicsControlData
 		, MaxTorque(0.0f)
 		, LinearTargetVelocityMultiplier(1.0f)
 		, AngularTargetVelocityMultiplier(1.0f)
+		, SkeletalAnimationVelocityMultiplier(1.0f)
+		, CustomControlPoint(FVector::ZeroVector)
 		, bEnabled(true)
+		, bUseCustomControlPoint(false)
+		, bUseSkeletalAnimation(true)
+		, bDisableCollision(false)
+		, bOnlyControlChildObject(false)
 	{
 	}
 
@@ -354,10 +523,7 @@ struct PHYSICSCONTROL_API FPhysicsControlData
 	UPROPERTY(EditAnywhere, BlueprintReadWrite, Category = PhysicsControl, meta = (ClampMin = "0.0"))
 	float LinearExtraDamping;
 
-	/** 
-	 * The maximum force used to drive the linear motion. Zero indicates no limit. 
-	 * Note - not yet implemented for RigidBodyWithControl 
-	 */
+	/** The maximum force used to drive the linear motion. Zero indicates no limit. */
 	UPROPERTY(EditAnywhere, BlueprintReadWrite, Category = PhysicsControl, meta = (ClampMin = "0.0"))
 	float MaxForce;
 
@@ -380,10 +546,7 @@ struct PHYSICSCONTROL_API FPhysicsControlData
 	UPROPERTY(EditAnywhere, BlueprintReadWrite, Category = PhysicsControl, meta = (ClampMin = "0.0"))
 	float AngularExtraDamping;
 
-	/** 
-	 * The maximum torque used to drive the angular motion. Zero indicates no limit. 
- 	 * Note - not yet implemented for RigidBodyWithControl 
-	 */
+	/** The maximum torque used to drive the angular motion. Zero indicates no limit. */
 	UPROPERTY(EditAnywhere, BlueprintReadWrite, Category = PhysicsControl, meta = (ClampMin = "0.0"))
 	float MaxTorque;
 
@@ -391,6 +554,7 @@ struct PHYSICSCONTROL_API FPhysicsControlData
 	 * Multiplier on the velocity, which gets applied to the damping. A value of 1 means the animation target
 	 * velocity is used, which helps it track the animation. A value of 0 means damping happens in "world space" 
 	 * - so damping acts like drag on the movement.
+	 * NOTE This is currently only used in the RigidBodyWithControl implementation
 	 */
 	UPROPERTY(EditAnywhere, BlueprintReadWrite, Category = PhysicsControl, meta = (ClampMin = "0.0"))
 	float LinearTargetVelocityMultiplier;
@@ -399,15 +563,139 @@ struct PHYSICSCONTROL_API FPhysicsControlData
 	 * Multiplier on the angular velocity, which gets applied to the damping. A value of 1 means the animation target
 	 * velocity is used, which helps it track the animation. A value of 0 means damping happens in "world space" 
 	 * - so damping acts like drag on the movement.
+	 * NOTE This is currently only used in the RigidBodyWithControl implementation
 	 */
 	UPROPERTY(EditAnywhere, BlueprintReadWrite, Category = PhysicsControl, meta = (ClampMin = "0.0"))
 	float AngularTargetVelocityMultiplier;
 
-	/**
-	 * Whether this control should be enabled
+	/** 
+	 * The amount of skeletal animation velocity to use in the targets 
+	 * NOTE This is currently only used in the PhysicsControlComponent implementation
 	 */
 	UPROPERTY(EditAnywhere, BlueprintReadWrite, Category = PhysicsControl)
+	float SkeletalAnimationVelocityMultiplier;
+
+	/** The position of the control point relative to the child mesh, when using a custom control point. */
+	UPROPERTY(EditAnywhere, BlueprintReadWrite, Category = PhysicsControl)
+	FVector CustomControlPoint;
+
+	/** Whether this control should be enabled */
+	UPROPERTY(EditAnywhere, BlueprintReadWrite, Category = PhysicsControl)
 	uint8 bEnabled : 1;
+
+	/** Whether or not to use the custom control point position */
+	UPROPERTY(EditAnywhere, BlueprintReadWrite, Category = PhysicsControl)
+	uint8 bUseCustomControlPoint : 1;
+
+	/** If true then the target will be applied on top of the skeletal animation (if there is any) */
+	UPROPERTY(EditAnywhere, BlueprintReadWrite, Category = PhysicsControl)
+	uint8 bUseSkeletalAnimation : 1;
+
+	/**
+	 * Whether or not this control should disable collision between the parent and child bodies (only
+	 * has an effect if there is a parent body)
+	 */
+	UPROPERTY(EditAnywhere, BlueprintReadWrite, Category = PhysicsControl)
+	uint8 bDisableCollision : 1;
+
+	/** If true then the control will only affect the child object, not the parent */
+	UPROPERTY(EditAnywhere, BlueprintReadWrite, Category = PhysicsControl)
+	uint8 bOnlyControlChildObject : 1;
+};
+
+/**
+ * These parameters allow modification of the parameters in FPhysicsControlData for two reasons:
+ * 1. They allow per-axis settings for the linear components (e.g. so you can drive an object 
+ *    horizontally but still let it fall under gravity)
+ * 2. They make it easy to create the controls with "default" strength/damping (e.g. taken from the
+ *    physics asset) in FPhysicsControlData, and then the strength/damping etc can be scaled every 
+ *    tick (typically between 0 and 1, though that is up to the user).
+ */
+USTRUCT(BlueprintType)
+struct PHYSICSCONTROL_API FPhysicsControlSparseMultiplier
+{
+	GENERATED_BODY()
+
+	FPhysicsControlSparseMultiplier()
+		: LinearStrengthMultiplier(1.0)
+		, LinearDampingRatioMultiplier(1.0)
+		, LinearExtraDampingMultiplier(1.0)
+		, MaxForceMultiplier(1.0)
+		, AngularStrengthMultiplier(1.0)
+		, AngularDampingRatioMultiplier(1.0)
+		, AngularExtraDampingMultiplier(1.0)
+		, MaxTorqueMultiplier(1.0)
+		// the flags
+		, bEnableLinearStrengthMultiplier(true)
+		, bEnableLinearDampingRatioMultiplier(true)
+		, bEnableLinearExtraDampingMultiplier(true)
+		, bEnableMaxForceMultiplier(true)
+		, bEnableAngularStrengthMultiplier(true)
+		, bEnableAngularDampingRatioMultiplier(true)
+		, bEnableAngularExtraDampingMultiplier(true)
+		, bEnableMaxTorqueMultiplier(true)
+	{
+	}
+
+	// Per-direction multiplier on the linear strength.
+	UPROPERTY(EditAnywhere, BlueprintReadWrite, Category = PhysicsControl, meta = (ClampMin = "0.0", editcondition = "bEnableLinearStrengthMultiplier"))
+	FVector LinearStrengthMultiplier;
+
+	// Per-direction multiplier on the linear damping ratio.
+	UPROPERTY(EditAnywhere, BlueprintReadWrite, Category = PhysicsControl, meta = (ClampMin = "0.0", editcondition = "bEnableLinearDampingRatioMultiplier"))
+	FVector LinearDampingRatioMultiplier;
+
+	// Per-direction multiplier on the linear extra damping.
+	UPROPERTY(EditAnywhere, BlueprintReadWrite, Category = PhysicsControl, meta = (ClampMin = "0.0", editcondition = "bEnableLinearExtraDampingMultiplier"))
+	FVector LinearExtraDampingMultiplier;
+
+	// Per-direction multiplier on the maximum force that can be applied. Note that zero means zero force.
+	UPROPERTY(EditAnywhere, BlueprintReadWrite, Category = PhysicsControl, meta = (ClampMin = "0.0", editcondition = "bEnableMaxForceMultiplier"))
+	FVector MaxForceMultiplier;
+
+	// Multiplier on the angular strength.
+	UPROPERTY(EditAnywhere, BlueprintReadWrite, Category = PhysicsControl, meta = (ClampMin = "0.0", editcondition = "bEnableAngularStrengthMultiplier"))
+	float AngularStrengthMultiplier;
+
+	// Multiplier on the angular damping ratio.
+	UPROPERTY(EditAnywhere, BlueprintReadWrite, Category = PhysicsControl, meta = (ClampMin = "0.0", editcondition = "bEnableAngularDampingRatioMultiplier"))
+	float AngularDampingRatioMultiplier;
+
+	// Multiplier on the angular extra damping.
+	UPROPERTY(EditAnywhere, BlueprintReadWrite, Category = PhysicsControl, meta = (ClampMin = "0.0", editcondition = "bEnableAngularExtraDampingMultiplier"))
+	float AngularExtraDampingMultiplier;
+
+	// Per-direction multiplier on the maximum torque that can be applied. Note that zero means zero torque.
+	UPROPERTY(EditAnywhere, BlueprintReadWrite, Category = PhysicsControl, meta = (ClampMin = "0.0", editcondition = "bEnableMaxTorqueMultiplier"))
+	float MaxTorqueMultiplier;
+
+	//==================================================================================================================
+	// Flags for sparse access
+	//==================================================================================================================
+
+	UPROPERTY(EditAnywhere, Category = PhysicsControl, meta = (InlineEditConditionToggle))
+	uint8 bEnableLinearStrengthMultiplier : 1;
+
+	UPROPERTY(EditAnywhere, Category = PhysicsControl, meta = (InlineEditConditionToggle))
+	uint8 bEnableLinearDampingRatioMultiplier : 1;
+
+	UPROPERTY(EditAnywhere, Category = PhysicsControl, meta = (InlineEditConditionToggle))
+	uint8 bEnableLinearExtraDampingMultiplier : 1;
+
+	UPROPERTY(EditAnywhere, Category = PhysicsControl, meta = (InlineEditConditionToggle))
+	uint8 bEnableMaxForceMultiplier : 1;
+
+	UPROPERTY(EditAnywhere, Category = PhysicsControl, meta = (InlineEditConditionToggle))
+	uint8 bEnableAngularStrengthMultiplier : 1;
+
+	UPROPERTY(EditAnywhere, Category = PhysicsControl, meta = (InlineEditConditionToggle))
+	uint8 bEnableAngularDampingRatioMultiplier : 1;
+
+	UPROPERTY(EditAnywhere, Category = PhysicsControl, meta = (InlineEditConditionToggle))
+	uint8 bEnableAngularExtraDampingMultiplier : 1;
+
+	UPROPERTY(EditAnywhere, Category = PhysicsControl, meta = (InlineEditConditionToggle))
+	uint8 bEnableMaxTorqueMultiplier : 1;
 };
 
 /**
@@ -435,6 +723,9 @@ struct PHYSICSCONTROL_API FPhysicsControlMultiplier
 	{
 	}
 
+	/** Applies the values that have been flagged as enabled from the sparse data */
+	void UpdateFromSparseData(const FPhysicsControlSparseMultiplier& SparseData);
+
 	// Per-direction multiplier on the linear strength.
 	UPROPERTY(EditAnywhere, BlueprintReadWrite, Category = PhysicsControl, meta = (ClampMin = "0.0"))
 	FVector LinearStrengthMultiplier;
@@ -516,52 +807,6 @@ struct PHYSICSCONTROL_API FPhysicsControlTarget
 	bool bApplyControlPointToTarget;
 };
 
-/**
- * General settings for a control
- */
-USTRUCT(BlueprintType)
-struct PHYSICSCONTROL_API FPhysicsControlSettings
-{
-	GENERATED_BODY()
-
-	FPhysicsControlSettings()
-		: ControlPoint(ForceInitToZero)
-		, bUseSkeletalAnimation(true)
-		, SkeletalAnimationVelocityMultiplier(1.0f)
-		, bDisableCollision(false)
-		, bOnlyControlChildObject(false)
-	{
-	}
-
-	/**
-	 * The position of the control point relative to the child mesh. Note that this can't be authored
-	 * directly here/on creation - it needs to be set after creation in UPhysicsControlComponent::SetControlPoint
-	 */
-	UPROPERTY()
-	FVector ControlPoint;
-
-	/** If true then the target will be applied on top of the skeletal animation (if there is any) */
-	UPROPERTY(EditAnywhere, BlueprintReadWrite, Category = PhysicsControl)
-	bool bUseSkeletalAnimation;
-
-	/** The amount of skeletal animation velocity to use in the targets */
-	UPROPERTY(EditAnywhere, BlueprintReadWrite, Category = PhysicsControl)
-	float SkeletalAnimationVelocityMultiplier;
-
-	/**
-	 * Whether or not this control should disable collision between the parent and child bodies (only
-	 * has an effect if there is a parent body)
-	 */
-	UPROPERTY(EditAnywhere, BlueprintReadWrite, Category = PhysicsControl)
-	bool bDisableCollision;
-
-	/**
-	 * If set to true, then the control will only affect the child object, not the parent
-	 */
-	UPROPERTY(EditAnywhere, BlueprintReadWrite, Category = PhysicsControl)
-	bool bOnlyControlChildObject;
-};
-
 /**
  * Structure that determines a "control" - this contains all the information needed to drive (with spring-dampers)
  * a child body relative to a parent body. These bodies will be associated with either a static or skeletal mesh.
@@ -579,21 +824,19 @@ struct PHYSICSCONTROL_API FPhysicsControl
 		UMeshComponent*              InChildMeshComponent,
 		const FName&                 InChildBoneName,
 		const FPhysicsControlData&   InControlData,
-		const FPhysicsControlTarget& InControlTarget,
-		const FPhysicsControlSettings& InControlSettings)
+		const FPhysicsControlTarget& InControlTarget)
 		: ParentMeshComponent(InParentMeshComponent)
 		, ParentBoneName(InParentBoneName)
 		, ChildMeshComponent(InChildMeshComponent)
 		, ChildBoneName(InChildBoneName)
 		, ControlData(InControlData)
 		, ControlTarget(InControlTarget)
-		, ControlSettings(InControlSettings)
 	{
 	}
 
 	/**  The mesh that will be doing the driving. Blank/non-existent means it will happen in world space */
 	UPROPERTY(EditAnywhere, BlueprintReadWrite, Category = PhysicsControl)
-	TObjectPtr<UMeshComponent> ParentMeshComponent;
+	TWeakObjectPtr<UMeshComponent> ParentMeshComponent;
 
 	/** The name of the skeletal mesh bone or the name of the static mesh body that will be doing the driving. */
 	UPROPERTY(EditAnywhere, BlueprintReadWrite, Category = PhysicsControl)
@@ -601,7 +844,7 @@ struct PHYSICSCONTROL_API FPhysicsControl
 
 	/** The mesh that the control will be driving. */
 	UPROPERTY(EditAnywhere, BlueprintReadWrite, Category = PhysicsControl)
-	TObjectPtr<UMeshComponent> ChildMeshComponent;
+	TWeakObjectPtr<UMeshComponent> ChildMeshComponent;
 
 	/** 
 	 * The name of the skeletal mesh bone or the name of the static mesh body that the control 
@@ -630,11 +873,130 @@ struct PHYSICSCONTROL_API FPhysicsControl
 	 */
 	UPROPERTY(EditAnywhere, BlueprintReadWrite, Category = PhysicsControl)
 	FPhysicsControlTarget ControlTarget;
+};
+
+/**
+ * Data that can be used to parameterize (modify/update) a control 
+ */
+USTRUCT(BlueprintType)
+struct PHYSICSCONTROL_API FPhysicsControlNamedControlParameters
+{
+	GENERATED_BODY();
+
+	FPhysicsControlNamedControlParameters() {}
+
+	FPhysicsControlNamedControlParameters(FName InName, const FPhysicsControlSparseData& InData)
+		: Name(InName), Data(InData) {}
+
+	// The name of the control (or set of controls) to update
+	UPROPERTY(EditAnywhere, BlueprintReadWrite, Category = PhysicsControl)
+	FName Name;
+
+	UPROPERTY(EditAnywhere, BlueprintReadWrite, Category = PhysicsControl)
+	FPhysicsControlSparseData Data;
+};
+
+/**
+ * Data that can be used to parameterize (modify/update) a control multiplier
+ */
+USTRUCT(BlueprintType)
+struct PHYSICSCONTROL_API FPhysicsControlNamedControlMultiplierParameters
+{
+	GENERATED_BODY();
+
+	FPhysicsControlNamedControlMultiplierParameters() {}
+
+	FPhysicsControlNamedControlMultiplierParameters(FName InName, const FPhysicsControlSparseMultiplier& InData)
+		: Name(InName), Data(InData) {}
+
+	// The name of the control (or set of controls) to update
+	UPROPERTY(EditAnywhere, BlueprintReadWrite, Category = PhysicsControl)
+	FName Name;
+
+	UPROPERTY(EditAnywhere, BlueprintReadWrite, Category = PhysicsControl)
+	FPhysicsControlSparseMultiplier Data;
+};
+
+
+/**
+ * Data that can be used to parameterize(modify / update) a modifier
+ */
+USTRUCT(BlueprintType)
+struct PHYSICSCONTROL_API FPhysicsControlNamedModifierParameters
+{
+	GENERATED_BODY();
+
+	FPhysicsControlNamedModifierParameters() {}
+
+	FPhysicsControlNamedModifierParameters(FName InName, const FPhysicsControlModifierSparseData& InData)
+		: Name(InName), Data(InData) {}
+
+	// The name of the modifier (or set of modifiers) to update
+	UPROPERTY(EditAnywhere, BlueprintReadWrite, Category = PhysicsControl)
+	FName Name;
+
+	UPROPERTY(EditAnywhere, BlueprintReadWrite, Category = PhysicsControl)
+	FPhysicsControlModifierSparseData Data;
+};
+
+/**
+ * These apply temporary/ephemeral changes to the controls that only persist for one tick.
+ */
+USTRUCT(BlueprintType)
+struct PHYSICSCONTROL_API FPhysicsControlControlAndModifierParameters
+{
+	GENERATED_BODY();
 
 	/**
-	 * More general settings for the control
+	 * Parameters for existing controls. Each name can be the name of a control, or the name of a 
+	 * set of controls. They will only apply for one tick/update. They will be applied in order (so 
+	 * subsequent entries will override earlier ones if they apply to the same control).
+	 */ 
+	UPROPERTY(EditAnywhere, BlueprintReadWrite, Category = Settings)
+	TArray<FPhysicsControlNamedControlParameters> ControlParameters;
+	
+	/**
+	 *  Parameters for existing modifiers. Each name can be the name of a modifier, or the name of a 
+	 * set of modifiers. They will only apply for one tick/update.
 	 */
+	UPROPERTY(EditAnywhere, BlueprintReadWrite, Category = Settings)
+	TArray<FPhysicsControlNamedModifierParameters> ModifierParameters;
+
+	void Add(const FPhysicsControlNamedControlParameters& InParameters) { ControlParameters.Add(InParameters); }
+	void Add(const FPhysicsControlNamedModifierParameters& InParameters) { ModifierParameters.Add(InParameters); }
+};
+
+/**
+ * These apply permanent changes to the controls and modifiers, allowing all the settings to be changed
+ * (apart from the actual bodies that are being controlled/affected)
+ */
+USTRUCT(BlueprintType)
+struct PHYSICSCONTROL_API FPhysicsControlControlAndModifierUpdates
+{
+	GENERATED_BODY();
+
+	/** Modifications to the underlying controls */
+	UPROPERTY(EditAnywhere, BlueprintReadWrite, Category = PhysicsControl)
+	TArray<FPhysicsControlNamedControlParameters> ControlUpdates;
+
+	/** Modifications to the underlying control multipliers */
+	UPROPERTY(EditAnywhere, BlueprintReadWrite, Category = PhysicsControl)
+	TArray<FPhysicsControlNamedControlMultiplierParameters> ControlMultiplierUpdates;
+
+	/** Modifications to the underlying modifiers */
 	UPROPERTY(EditAnywhere, BlueprintReadWrite, Category = PhysicsControl)
-	FPhysicsControlSettings ControlSettings;
+	TArray<FPhysicsControlNamedModifierParameters> ModifierUpdates;
 };
 
+// Functions to interpolate between various control/modifier datas
+PHYSICSCONTROL_API FPhysicsControlData Interpolate(
+	const FPhysicsControlData& A, const FPhysicsControlData& B, const float Weight);
+PHYSICSCONTROL_API FPhysicsControlSparseData Interpolate(
+	const FPhysicsControlSparseData& A, const FPhysicsControlSparseData& B, const float Weight);
+
+PHYSICSCONTROL_API FPhysicsControlModifierData Interpolate(
+	const FPhysicsControlModifierData& A, const FPhysicsControlModifierData& B, const float Weight);
+PHYSICSCONTROL_API FPhysicsControlModifierSparseData Interpolate(
+	const FPhysicsControlModifierSparseData& A, const FPhysicsControlModifierSparseData& B, const float Weight);
+
+
diff --git a/Engine/Plugins/Experimental/PhysicsControl/Source/PhysicsControl/Public/PhysicsControlLimbData.h b/Engine/Plugins/Experimental/PhysicsControl/Source/PhysicsControl/Public/PhysicsControlLimbData.h
--- a/Engine/Plugins/Experimental/PhysicsControl/Source/PhysicsControl/Public/PhysicsControlLimbData.h
+++ b/Engine/Plugins/Experimental/PhysicsControl/Source/PhysicsControl/Public/PhysicsControlLimbData.h
@@ -73,7 +73,7 @@ struct PHYSICSCONTROL_API FPhysicsControlLimbBones
 	{}
 
 	/** The Skeletal mesh that this limb is associated with */
-	TObjectPtr<USkeletalMeshComponent> SkeletalMeshComponent;
+	TWeakObjectPtr<USkeletalMeshComponent> SkeletalMeshComponent;
 
 	/** The names of the bones in the limb */
 	TArray<FName> BoneNames;
diff --git a/Engine/Plugins/Experimental/PhysicsControl/Source/PhysicsControl/Public/PhysicsControlNameRecords.h b/Engine/Plugins/Experimental/PhysicsControl/Source/PhysicsControl/Public/PhysicsControlNameRecords.h
new file mode 100644
--- /dev/null
+++ b/Engine/Plugins/Experimental/PhysicsControl/Source/PhysicsControl/Public/PhysicsControlNameRecords.h
@@ -0,0 +1,55 @@
+// Copyright Epic Games, Inc. All Rights Reserved.
+
+#pragma once
+
+#include "CoreMinimal.h"
+
+/**
+ * All the controls and body modifiers we create get given names, so we can keep a record in a form
+ * that will be useful to users. The three general uses are:
+ * 1. All controls or modifiers
+ * 2. The type of control - the obvious ones are "world space" and "parent space", but users might make others
+ * 3. Sets of control/modifiers - the obvious ones are limbs (e.g. a set of world-space controls on the leg)
+ *    but users might want to make other sets - e.g. "UpperBody"
+ */
+struct FPhysicsControlNameRecords
+{
+	/** Adds Name to SetName as well as adding it to the set "All" */
+	void AddControl(FName Name, FName SetName);
+	/** Adds Name to SetNames as well as adding it to the set "All" */
+	void AddControl(FName Name, const TArray<FName>& SetNames);
+	/** Adds a collection of Names to SetName as well as adding them to the set "All"  */
+	void AddControls(const TArray<FName>& ControlNames, const FName SetName);
+	/** Removes Name from all Sets */
+	void RemoveControl(FName Name);
+	const TArray<FName>& GetControlNamesInSet(FName SetName) const;
+
+	/** Adds Name to SetName as well as adding it to the set "All" */
+	void AddBodyModifier(FName Name, FName SetName);
+	/** Adds Name to SetNames as well as adding it to the set "All" */
+	void AddBodyModifier(FName Name, const TArray<FName>& SetNames);
+	/** Adds a collection of Names to SetName as well as adding them to the set "All"  */
+	void AddBodyModifiers(const TArray<FName>& BodyModifierNames, const FName SetName);
+	/** Removes Name from all Sets */
+	void RemoveBodyModifier(FName Name);
+	const TArray<FName>& GetBodyModifierNamesInSet(FName SetName) const;
+
+	/** Remove all control and modifier sets. */
+	void Reset();
+
+	/**
+	 * All the control sets we've created, arranged by set name. Note that some controls will
+	 * not necessarily have been added to a set.
+	 */
+	TMap<FName, TArray<FName>> ControlSets;
+
+	/**
+	 * All the body modifiers we've created, arranged by set name. Note that some modifiers will
+	 * not necessarily have been added to a set.
+	 * */
+	TMap<FName, TArray<FName>> BodyModifierSets;
+};
+
+/** Returns a copy of the supplied list of names in which any names that match keys in the supplied 'Sets' map are replaced with all the names in the array associated with that key. */
+PHYSICSCONTROL_API TArray<FName> ExpandName(const FName InName, const TMap<FName, TArray<FName>>& Sets);
+PHYSICSCONTROL_API TArray<FName> ExpandName(const TArray<FName>& InNames, const TMap<FName, TArray<FName>>& Sets);
diff --git a/Engine/Plugins/Experimental/PhysicsControl/Source/PhysicsControl/Public/PhysicsControlOperatorNameGeneration.h b/Engine/Plugins/Experimental/PhysicsControl/Source/PhysicsControl/Public/PhysicsControlOperatorNameGeneration.h
--- a/Engine/Plugins/Experimental/PhysicsControl/Source/PhysicsControl/Public/PhysicsControlOperatorNameGeneration.h
+++ b/Engine/Plugins/Experimental/PhysicsControl/Source/PhysicsControl/Public/PhysicsControlOperatorNameGeneration.h
@@ -15,7 +15,7 @@ struct FPhysicsControlSetUpdates;
 struct FReferenceSkeleton;
 struct FRigidBodyControlRecord;
 struct FRigidBodyModifierRecord;
-struct FRigidBodyNameRecords;
+struct FPhysicsControlNameRecords;
 
 constexpr int32 MaxNumControlsOrModifiersPerName = 16;
 
@@ -23,16 +23,16 @@ constexpr int32 MaxNumControlsOrModifiersPerName = 16;
 PHYSICSCONTROL_API TMap<FName, FPhysicsControlLimbBones> GetLimbBones(const TArray<FPhysicsControlLimbSetupData>& LimbSetupData, const FReferenceSkeleton& RefSkeleton, UPhysicsAsset* const PhysicsAsset);
 
 // Populates the supplied body modifier names, control names and name records structures with the names and sets that could be created for the supplied node, limb bones, skeleton and physics asset.
-PHYSICSCONTROL_API void CollectOperatorNames(const FAnimNode_RigidBodyWithControl* const Node, TMap<FName, FPhysicsControlLimbBones> AllLimbBones, const FReferenceSkeleton& RefSkeleton, UPhysicsAsset* const PhysicsAsset, TSet<FName>& BodyModifierNames, TSet<FName>& ControlNames, FRigidBodyNameRecords& NameRecords);
+PHYSICSCONTROL_API void CollectOperatorNames(const FAnimNode_RigidBodyWithControl* const Node, TMap<FName, FPhysicsControlLimbBones> AllLimbBones, const FReferenceSkeleton& RefSkeleton, UPhysicsAsset* const PhysicsAsset, TSet<FName>& BodyModifierNames, TSet<FName>& ControlNames, FPhysicsControlNameRecords& NameRecords);
 
 // Creates the body modifiers, controls and sets for the supplied node, limb bones, skeleton and physics asset.
-PHYSICSCONTROL_API void CreateOperatorsForNode(FAnimNode_RigidBodyWithControl* const Node, TMap<FName, FPhysicsControlLimbBones> AllLimbBones, const FReferenceSkeleton& RefSkeleton, UPhysicsAsset* const PhysicsAsset, FRigidBodyNameRecords& NameRecords);
+PHYSICSCONTROL_API void CreateOperatorsForNode(FAnimNode_RigidBodyWithControl* const Node, TMap<FName, FPhysicsControlLimbBones> AllLimbBones, const FReferenceSkeleton& RefSkeleton, UPhysicsAsset* const PhysicsAsset, FPhysicsControlNameRecords& NameRecords);
 
 // Adds the specified additional sets to the supplied Name Records structure.
-PHYSICSCONTROL_API void CreateAdditionalSets(const FPhysicsControlSetUpdates& AdditionalSets, const TSet<FName>& BodyModifierNames, const TSet<FName>& ControlNames, FRigidBodyNameRecords& NameRecords);
+PHYSICSCONTROL_API void CreateAdditionalSets(const FPhysicsControlSetUpdates& AdditionalSets, const TSet<FName>& BodyModifierNames, const TSet<FName>& ControlNames, FPhysicsControlNameRecords& NameRecords);
 
 // Adds the specified additional sets to the supplied Name Records structure.
-PHYSICSCONTROL_API void CreateAdditionalSets(const FPhysicsControlSetUpdates& AdditionalSets, const TMap<FName, FRigidBodyModifierRecord>& BodyModifierNames, const TMap<FName, FRigidBodyControlRecord>& ControlNames, FRigidBodyNameRecords& NameRecords);
+PHYSICSCONTROL_API void CreateAdditionalSets(const FPhysicsControlSetUpdates& AdditionalSets, const TMap<FName, FRigidBodyModifierRecord>& BodyModifierNames, const TMap<FName, FRigidBodyControlRecord>& ControlNames, FPhysicsControlNameRecords& NameRecords);
 
 FName FindUniqueName(const FString& NameBase, const TSet<FName>& Keys, const int32 MaxNmeaIndex);
 FName FindUniqueBodyModifierName(const FName BodyName, const TSet<FName>& ExistingNames);
diff --git a/Engine/Plugins/Experimental/PhysicsControl/Source/PhysicsControl/Public/PhysicsControlProfileAsset.h b/Engine/Plugins/Experimental/PhysicsControl/Source/PhysicsControl/Public/PhysicsControlProfileAsset.h
--- a/Engine/Plugins/Experimental/PhysicsControl/Source/PhysicsControl/Public/PhysicsControlProfileAsset.h
+++ b/Engine/Plugins/Experimental/PhysicsControl/Source/PhysicsControl/Public/PhysicsControlProfileAsset.h
@@ -2,12 +2,68 @@
 
 #pragma once
 
+#include "PhysicsControlData.h"
+#include "PhysicsControlLimbData.h"
+
 #include "Interfaces/Interface_PreviewMeshProvider.h"
 
 #include "PhysicsControlProfileAsset.generated.h"
 
 class USkeletalMesh;
 
+
+USTRUCT(BlueprintType)
+struct PHYSICSCONTROL_API FPhysicsControlProfileInitialControl
+{
+	GENERATED_BODY()
+
+	UPROPERTY(EditAnywhere, BlueprintReadWrite, Category = PhysicsControl)
+	FName ParentBoneName;
+
+	UPROPERTY(EditAnywhere, BlueprintReadWrite, Category = PhysicsControl)
+	FName ChildBoneName;
+
+	UPROPERTY(EditAnywhere, BlueprintReadWrite, Category = PhysicsControl)
+	FPhysicsControlData ControlData;
+
+	UPROPERTY(EditAnywhere, BlueprintReadWrite, Category = PhysicsControl)
+	TArray<FName> Sets;
+};
+
+
+USTRUCT(BlueprintType)
+struct PHYSICSCONTROL_API FPhysicsControlProfileInitialBodyModifier
+{
+	GENERATED_BODY()
+
+	UPROPERTY(EditAnywhere, BlueprintReadWrite, Category = PhysicsControl)
+	FName BoneName;
+
+	UPROPERTY(EditAnywhere, BlueprintReadWrite, Category = PhysicsControl)
+	FPhysicsControlModifierData BodyModifierData;
+
+	UPROPERTY(EditAnywhere, BlueprintReadWrite, Category = PhysicsControl)
+	TArray<FName> Sets;
+};
+
+USTRUCT(BlueprintType)
+struct FPhysicsControlProfileCharacterSetup
+{
+	GENERATED_BODY();
+
+	UPROPERTY(EditAnywhere, Category = ControlSetup)
+	TArray<FPhysicsControlLimbSetupData> LimbSetupData;
+
+	UPROPERTY(EditAnywhere, Category = ControlSetup)
+	FPhysicsControlData DefaultWorldSpaceControlData;
+
+	UPROPERTY(EditAnywhere, Category = ControlSetup)
+	FPhysicsControlData DefaultParentSpaceControlData;
+
+	UPROPERTY(EditAnywhere, Category = ControlSetup)
+	FPhysicsControlModifierData DefaultBodyModifierData;
+};
+
 /**
  * Asset for storing Physics Control Profiles. These will contain data that define:
  * - Controls and body modifiers to be created on a mesh
@@ -25,19 +81,28 @@ class PHYSICSCONTROL_API UPhysicsControlProfileAsset : public UObject, public II
 public:
 	UPhysicsControlProfileAsset();
 
-	/** Placeholder for testing/developing */
-	UFUNCTION(CallInEditor)
+	/** Shows all the controls etc that would be made */
+	UFUNCTION(CallInEditor, Category = Development)
 	void Log();
 
-	/** IInterface_PreviewMeshProvider interface */
-	virtual void SetPreviewMesh(USkeletalMesh* PreviewMesh, bool bMarkAsDirty = true) override;
-	virtual USkeletalMesh* GetPreviewMesh() const override;
-	/** END IInterface_PreviewMeshProvider interface */
+	/** Sets the initial controls/modifiers based on character setup data */
+	UFUNCTION(CallInEditor, Category = Development)
+	void MakeControlsAndModifiersFromCharacterSetupData();
 
-#if WITH_EDITOR
-/* Get name of Preview Mesh property */
-	static const FName GetPreviewMeshPropertyName();
-#endif
+	/**
+	 * We can define controls in the form of limbs etc here
+	 */
+	UPROPERTY(EditAnywhere, Category = Development)
+	FPhysicsControlProfileCharacterSetup CharacterSetupData;
+
+public:
+	/** A profile asset to inherit from (can be null). If set, we will just add/modify data in that */
+	UPROPERTY(EditAnywhere, Category = Inheritance)
+	TObjectPtr<UPhysicsControlProfileAsset> ParentAsset;
+
+	/** Additional profile assets to use for the profiles */
+	UPROPERTY(EditAnywhere, Category = Inheritance)
+	TArray<TObjectPtr<UPhysicsControlProfileAsset>> AdditionalProfileAssets;
 
 public:
 	/** 
@@ -47,7 +112,33 @@ public:
 	UPROPERTY(AssetRegistrySearchable, EditAnywhere, Category = PreviewMesh)
 	TSoftObjectPtr<USkeletalMesh> PreviewSkeletalMesh;
 
-	/** Placeholder for testing/developing */
-	UPROPERTY(EditAnywhere, Category=PhysicsControl)
-	float TestValue;
+	/**
+	 * Additional controls. If these have the same name as one that's already created, they'll just override it.
+	 */
+	UPROPERTY(EditAnywhere, Category = ControlsAndBodyModifiers)
+	TMap<FName, FPhysicsControlProfileInitialControl> InitialControls;
+
+	/**
+	 * Additional body modifiers. If these have the same name as one that's already created, they'll just override it.
+	 */
+	UPROPERTY(EditAnywhere, Category = ControlsAndBodyModifiers)
+	TMap<FName, FPhysicsControlProfileInitialBodyModifier> InitialBodyModifiers;
+
+	/**
+	 * The named profiles, which are essentially control and modifier updates
+	 */
+	UPROPERTY(EditAnywhere, Category = Profiles)
+	TMap<FName, FPhysicsControlControlAndModifierUpdates> Profiles;
+
+public:
+	/** IInterface_PreviewMeshProvider interface */
+	virtual void SetPreviewMesh(USkeletalMesh* PreviewMesh, bool bMarkAsDirty = true) override;
+	virtual USkeletalMesh* GetPreviewMesh() const override;
+	/** END IInterface_PreviewMeshProvider interface */
+
+#if WITH_EDITOR
+/* Get name of Preview Mesh property */
+	static const FName GetPreviewMeshPropertyName();
+#endif
+
 };
diff --git a/Engine/Plugins/Experimental/PhysicsControl/Source/PhysicsControl/Public/PhysicsControlRecord.h b/Engine/Plugins/Experimental/PhysicsControl/Source/PhysicsControl/Public/PhysicsControlRecord.h
new file mode 100644
--- /dev/null
+++ b/Engine/Plugins/Experimental/PhysicsControl/Source/PhysicsControl/Public/PhysicsControlRecord.h
@@ -0,0 +1,193 @@
+// Copyright Epic Games, Inc. All Rights Reserved.
+
+#pragma once
+
+#include "PhysicsControlData.h"
+#include "Components/SkeletalMeshComponent.h"
+#include "Engine/EngineTypes.h"
+
+struct FConstraintInstance;
+struct FBodyInstance;
+class UMeshComponent;
+
+/**
+ * There will be a PhysicsControlRecord created at runtime for every Control that has been created
+ */
+struct FPhysicsControlRecord
+{
+	FPhysicsControlRecord(const FPhysicsControl& InControl) : PhysicsControl(InControl) {}
+
+	// Indicates if the control is enabled
+	bool Enabled() const { return PhysicsControl.ControlData.bEnabled; }
+
+	/** Removes any constraint and resets the state */
+	void ResetConstraint();
+
+	/** Returns the control point, which may be custom or automatic (centre of mass) */
+	FVector GetControlPoint() const;
+
+	/**
+	 * Creates the constraint if necessary and stores it. Then initializes the constraint with the bodies.
+	 * Returns true/false on success/failure */
+	bool InitConstraint(UObject* ConstraintDebugOwner, FName ControlName);
+
+	/** Ensures the constraint frame matches the control point in the record. */
+	void UpdateConstraintControlPoint();
+
+	/** Sets the control point to the center of mass of the child mesh (or to zero if that fails). */
+	void ResetControlPoint();
+
+	/** The configuration data */
+	FPhysicsControl PhysicsControl;
+
+	/** The underlying constraint used to implement the control. */
+	TSharedPtr<FConstraintInstance> ConstraintInstance;
+
+};
+
+/**
+ * There will be a PhysicsBodyModifier created at runtime for every BodyInstance involved in the component
+ */
+struct FPhysicsBodyModifier
+{
+	FPhysicsBodyModifier(
+		TWeakObjectPtr<UMeshComponent>  InMeshComponent, 
+		const FName&                InBoneName, 
+		FPhysicsControlModifierData InBodyModifierData)
+		: MeshComponent(InMeshComponent)
+		, BoneName(InBoneName)
+		, BodyModifierData(InBodyModifierData)
+		, KinematicTargetPosition(FVector::ZeroVector)
+		, KinematicTargetOrientation(FQuat::Identity)
+		, bResetToCachedTarget(false)
+	{}
+
+	/**  The mesh that will be modified. */
+	TWeakObjectPtr<UMeshComponent> MeshComponent;
+
+	/** The name of the skeletal mesh bone or the name of the static mesh body that will be modified. */
+	FName BoneName;
+
+	/** Basic data about how the body should be modified */
+	FPhysicsControlModifierData BodyModifierData;
+
+	/** 
+	 * The target position when kinematic. Note that this is applied on top of any animation 
+	 * target if bUseSkeletalAnimation is set. 
+	 */
+	FVector KinematicTargetPosition = FVector::ZeroVector;
+
+	/**
+	 * The target orientation when kinematic. Note that this is applied on top of any animation
+	 * target if bUseSkeletalAnimation is set.
+	 */
+	FQuat KinematicTargetOrientation = FQuat::Identity;
+
+	/** 
+	 * If true then the body will be set to the transform/velocity stored in any cached target (if that
+	 * exists), and then this flag will be cleared.
+	 */
+	uint8 bResetToCachedTarget:1;
+};
+
+/**
+ * Used internally/only at runtime to track when a SkeletalMeshComponent is being controlled through
+ * a modifier, and to restore settings when that stops.
+ */
+struct FModifiedSkeletalMeshData
+{
+public: 
+	FModifiedSkeletalMeshData() : ReferenceCount(0) {}
+
+public:
+
+	/** The original setting for restoration when we're deleted */
+	uint8 bOriginalUpdateMeshWhenKinematic : 1;
+
+	/** The original setting for restoration when we're deleted */
+	EKinematicBonesUpdateToPhysics::Type OriginalKinematicBonesUpdateType;
+
+	/**
+	 * Track when skeletal meshes are going to be used so this entry can be removed
+	 */
+	int32 ReferenceCount;
+};
+
+/**
+ * Used internally/only at runtime to cache skeletal transforms at the start of the tick, to avoid
+ * calculating them separately for every control.
+ */
+struct FCachedSkeletalMeshData
+{
+public:
+	FCachedSkeletalMeshData() : ReferenceCount(0) {}
+
+public:
+	struct FBoneData
+	{
+		FBoneData()
+			: Position(FVector::ZeroVector), Orientation(FQuat::Identity)
+			, Velocity(FVector::ZeroVector), AngularVelocity(FVector::ZeroVector) {}
+		FBoneData(const FVector& InPosition, const FQuat& InOrientation)
+			: Position(InPosition), Orientation(InOrientation),
+			Velocity(FVector::ZeroVector), AngularVelocity(FVector::ZeroVector) {}
+
+		/**
+		 * Sets position/orientation and calculates velocity/angular velocity - requires Dt > 0 
+		 */
+		void Update(const FVector& InPosition, const FQuat& InOrientation, float Dt);
+
+		/**
+		 * Sets position/orientation, and sets velocity to zero
+		 */
+		void Update(const FVector& InPosition, const FQuat& InOrientation);
+
+		FTransform GetTM() const { return FTransform(Orientation, Position); }
+
+		FVector Position;
+		FQuat   Orientation;
+		FVector Velocity;
+		FVector AngularVelocity;
+	};
+
+public:
+
+	/**
+	 * The cached skeletal data, updated at the start of each tick
+	 */
+	TArray<FBoneData> BoneData;
+
+	/**
+	 * The component transform. This is only stored so we can detect teleports
+	 */
+	FTransform ComponentTM;
+
+	/**
+	 * Track when skeletal meshes are going to be used so this entry can be removed, and also so we
+	 * can add a tick dependency
+	 */
+	int32 ReferenceCount;
+};
+
+//======================================================================================================================
+inline void FCachedSkeletalMeshData::FBoneData::Update(const FVector& InPosition, const FQuat& InOrientation, float Dt)
+{
+	check(Dt > 0);
+
+	Velocity = (InPosition - Position) / Dt;
+	Orientation.EnforceShortestArcWith(InOrientation);
+	// Note that quats multiply in the opposite order to TMs
+	const FQuat DeltaQ = InOrientation * Orientation.Inverse();
+	AngularVelocity = DeltaQ.ToRotationVector() / Dt;
+	Position = InPosition;
+	Orientation = InOrientation;
+}
+
+//======================================================================================================================
+inline void FCachedSkeletalMeshData::FBoneData::Update(const FVector& InPosition, const FQuat& InOrientation)
+{
+	Position = InPosition;
+	Orientation = InOrientation;
+	Velocity = FVector::ZeroVector;
+	AngularVelocity = FVector::ZeroVector;
+}
diff --git a/Engine/Plugins/Experimental/PhysicsControl/Source/PhysicsControl/Public/RigidBodyControlData.h b/Engine/Plugins/Experimental/PhysicsControl/Source/PhysicsControl/Public/RigidBodyControlData.h
--- a/Engine/Plugins/Experimental/PhysicsControl/Source/PhysicsControl/Public/RigidBodyControlData.h
+++ b/Engine/Plugins/Experimental/PhysicsControl/Source/PhysicsControl/Public/RigidBodyControlData.h
@@ -121,103 +121,6 @@ struct PHYSICSCONTROL_API FRigidBodyControlAndBodyModifierCreations
 	TArray<FRigidBodyModifierCreation> Modifiers;
 };
 
-PHYSICSCONTROL_API FPhysicsControlData Interpolate(
-	const FPhysicsControlData& A, const FPhysicsControlData& B, const float Weight);
-PHYSICSCONTROL_API FPhysicsControlSparseData Interpolate(
-	const FPhysicsControlSparseData& A, const FPhysicsControlSparseData& B, const float Weight);
-
-/**
- * Data that can be used to parameterize (modify/update) a control 
- */
-USTRUCT(BlueprintType)
-struct PHYSICSCONTROL_API FPhysicsControlNamedControlParameters
-{
-	GENERATED_BODY();
-
-	FPhysicsControlNamedControlParameters() {}
-
-	FPhysicsControlNamedControlParameters(FName InName, const FPhysicsControlSparseData& InData)
-		: Name(InName), Data(InData) {}
-
-	// The name of the control (or set of controls) to update
-	UPROPERTY(EditAnywhere, BlueprintReadWrite, Category = PhysicsControl)
-	FName Name;
-
-	UPROPERTY(EditAnywhere, BlueprintReadWrite, Category = PhysicsControl)
-	FPhysicsControlSparseData Data;
-};
-
-PHYSICSCONTROL_API FPhysicsControlModifierData Interpolate(
-	const FPhysicsControlModifierData& A, const FPhysicsControlModifierData& B, const float Weight);
-PHYSICSCONTROL_API FPhysicsControlModifierSparseData Interpolate(
-	const FPhysicsControlModifierSparseData& A, const FPhysicsControlModifierSparseData& B, const float Weight);
-
-/**
- * Data that can be used to parameterize(modify / update) a control
- */
-USTRUCT(BlueprintType)
-struct PHYSICSCONTROL_API FPhysicsControlNamedModifierParameters
-{
-	GENERATED_BODY();
-
-	FPhysicsControlNamedModifierParameters() {}
-
-	FPhysicsControlNamedModifierParameters(FName InName, const FPhysicsControlModifierSparseData& InData)
-		: Name(InName), Data(InData) {}
-
-	// The name of the modifier (or set of modifiers) to update
-	UPROPERTY(EditAnywhere, BlueprintReadWrite, Category = PhysicsControl)
-	FName Name;
-
-	UPROPERTY(EditAnywhere, BlueprintReadWrite, Category = PhysicsControl)
-	FPhysicsControlModifierSparseData Data;
-};
-
-/**
- * These apply temporary/ephemeral changes to the controls that only persist for one tick.
- */
-USTRUCT(BlueprintType)
-struct PHYSICSCONTROL_API FPhysicsControlControlAndModifierParameters
-{
-	GENERATED_BODY();
-
-	/**
-	 * Parameters for existing controls. Each name can be the name of a control, or the name of a 
-	 * set of controls. They will only apply for one tick/update. They will be applied in order (so 
-	 * subsequent entries will override earlier ones if they apply to the same control).
-	 */ 
-	UPROPERTY(EditAnywhere, BlueprintReadWrite, Category = Settings)
-	TArray<FPhysicsControlNamedControlParameters> ControlParameters;
-	
-	/**
-	 *  Parameters for existing modifiers. Each name can be the name of a modifier, or the name of a 
-	 * set of modifiers. They will only apply for one tick/update.
-	 */
-	UPROPERTY(EditAnywhere, BlueprintReadWrite, Category = Settings)
-	TArray<FPhysicsControlNamedModifierParameters> ModifierParameters;
-
-	void Add(const FPhysicsControlNamedControlParameters& InParameters) { ControlParameters.Add(InParameters); }
-	void Add(const FPhysicsControlNamedModifierParameters& InParameters) { ModifierParameters.Add(InParameters); }
-};
-
-/**
- * These apply permanent changes to the controls and modifiers, allowing all the settings to be changed
- * (apart from the actual bodies that are being controlled/affected)
- */
-USTRUCT(BlueprintType)
-struct PHYSICSCONTROL_API FPhysicsControlControlAndModifierUpdates
-{
-	GENERATED_BODY();
-
-	/** Modifications to the underlying controls - these will persist */
-	UPROPERTY(EditAnywhere, BlueprintReadWrite, Category = PhysicsControl)
-	TArray<FPhysicsControlNamedControlParameters> ControlParameters;
-
-	/** Modifications to the underlying modifiers - these will persist */
-	UPROPERTY(EditAnywhere, BlueprintReadWrite, Category = PhysicsControl)
-	TArray<FPhysicsControlNamedModifierParameters> ModifierParameters;
-};
-
 /**
  * A single target for a control, which may be defined as an offset from the (implicit) animation target.
  */
@@ -229,9 +132,6 @@ struct PHYSICSCONTROL_API FRigidBodyControlTarget
 	FRigidBodyControlTarget()
 		: TargetPosition(ForceInitToZero)
 		, TargetOrientation(ForceInitToZero)
-		, TargetPoint(ForceInitToZero)
-		, bUseSkeletalAnimation(true)
-		, bUseTargetPoint(false)
 	{
 	}
 
@@ -242,21 +142,6 @@ struct PHYSICSCONTROL_API FRigidBodyControlTarget
 	/** The target orientation of the child body, relative to the parent body */
 	UPROPERTY(EditAnywhere, BlueprintReadWrite, Category = PhysicsControl)
 	FRotator TargetOrientation;
-
-	/** 
-	 * The point on the controlled (child) object that should be driven towards the target position.
-	 * Note that if this is not set (i.e. if UseTargetPoint is false) then the centre of mass will be used. 
-	 */
-	UPROPERTY(EditAnywhere, BlueprintReadWrite, Category = PhysicsControl, meta = (editcondition = "bUseTargetPoint"))
-	FVector TargetPoint;
-
-	/** If true then the target will be applied on top of the skeletal animation */
-	UPROPERTY(EditAnywhere, BlueprintReadWrite, Category = PhysicsControl)
-	uint8 bUseSkeletalAnimation : 1;
-
-	/** If true then the target will be applied on top of the skeletal animation */
-	UPROPERTY(EditAnywhere, BlueprintReadWrite, Category = PhysicsControl)
-	uint8 bUseTargetPoint : 1;
 };
 
 /**
@@ -325,7 +210,10 @@ struct FRigidBodyControlRecord
 	void ResetCurrent(bool bResetTarget);
 
 	// Note that this is only correct when called during or after the update has been done
-	bool IsEnabled() const { return CurrentData.bEnabled; }
+	bool IsEnabled() const { return ControlData.bEnabled; }
+
+	/** Returns the control point, which may be custom or automatic (centre of mass) */
+	FVector GetControlPoint(const ImmediatePhysics::FActorHandle* ChildActorHandle) const;
 
 	// TODO - might benefit from smaller - non-blueprint data members here ?
 	FRigidBodyControl               Control;
@@ -333,7 +221,7 @@ struct FRigidBodyControlRecord
 
 	// This contains the currently active control data. It will be updated just prior to
 	// applying the controls, by setting it to the default, and then updating it with any parameters.
-	FPhysicsControlData CurrentData;
+	FPhysicsControlData ControlData;
 
 	// Contains any control target that has been set
 	FRigidBodyControlTarget ControlTarget;
@@ -370,7 +258,6 @@ struct FRigidBodyModifierRecord
 
 	// This contains the currently active modifier data. It will be updated just prior to
 	// applying the controls, by setting it to the default, and then updating it with any parameters.
-	FPhysicsControlModifierData CurrentData;
+	FPhysicsControlModifierData ModifierData;
 };
 
-
diff --git a/Engine/Plugins/Experimental/PhysicsControl/Source/PhysicsControl/Public/RigidBodyNameRecords.h b/Engine/Plugins/Experimental/PhysicsControl/Source/PhysicsControl/Public/RigidBodyNameRecords.h
deleted file mode 100644
diff --git a/Engine/Plugins/Experimental/PhysicsControl/Source/PhysicsControlEditor/Private/OperatorEditor/SOperatorTreeWidget.cpp b/Engine/Plugins/Experimental/PhysicsControl/Source/PhysicsControlEditor/Private/OperatorEditor/SOperatorTreeWidget.cpp
--- a/Engine/Plugins/Experimental/PhysicsControl/Source/PhysicsControlEditor/Private/OperatorEditor/SOperatorTreeWidget.cpp
+++ b/Engine/Plugins/Experimental/PhysicsControl/Source/PhysicsControlEditor/Private/OperatorEditor/SOperatorTreeWidget.cpp
@@ -315,7 +315,7 @@ void SOperatorTreeWidget::Refresh()
 			// Create node item
 			NodeItemPtr Node = MakeShared<OperatorTreeNodeItem>(*Itr, FName(*NodeUniqueName), BlueprintName);
 
-			if (!Itr->Node.OverridePhysicsAsset)
+			if (!Itr->Node.OverridePhysicsAsset.IsValid())
 			{
 				Node->AddChild(MakeShared<OperatorTreeMessageItem>(*Itr, OperatorTreeMessageItem::Warning, LOCTEXT("OperatorNoPhysicsAssetDefinedText", "Can't create operators for limbs because no Physics Asset has been defined for this node.")));
 			}
diff --git a/Engine/Plugins/Experimental/PhysicsControl/Source/PhysicsControlEditor/Private/PhysicsControlEditorModule.cpp b/Engine/Plugins/Experimental/PhysicsControl/Source/PhysicsControlEditor/Private/PhysicsControlEditorModule.cpp
--- a/Engine/Plugins/Experimental/PhysicsControl/Source/PhysicsControlEditor/Private/PhysicsControlEditorModule.cpp
+++ b/Engine/Plugins/Experimental/PhysicsControl/Source/PhysicsControlEditor/Private/PhysicsControlEditorModule.cpp
@@ -14,17 +14,16 @@ static const FName PhysicsControlEditorModule_PhysicsControlEditorInterface("Phy
 
 #define LOCTEXT_NAMESPACE "PhysicsControlModule"
 
-//#define ENABLE_PHYSICS_CONTROL_PROFILE_EDITOR
-
 //======================================================================================================================
 void FPhysicsControlEditorModule::StartupModule()
 {
-	// Physics Control Profile editor/asset is disabled for now
-#ifdef ENABLE_PHYSICS_CONTROL_PROFILE_EDITOR
+#ifdef ENABLE_PHYSICS_CONTROL_PROFILE_ASSET
 	IAssetTools& AssetTools = FModuleManager::LoadModuleChecked<FAssetToolsModule>("AssetTools").Get();
 	PhysicsControlProfileAssetActions = MakeShared<FPhysicsControlProfileAssetActions>();
 	AssetTools.RegisterAssetTypeActions(PhysicsControlProfileAssetActions.ToSharedRef());
+#endif
 
+#if ENABLE_PHYSICS_CONTROL_PROFILE_EDITOR
 	FEditorModeRegistry::Get().RegisterMode<FPhysicsControlProfileEditorMode>(
 		FPhysicsControlProfileEditorMode::ModeName, 
 		LOCTEXT("PhysicsControlProfileEditorMode", "PhysicsControlProfile"), 
@@ -65,9 +64,11 @@ void FPhysicsControlEditorModule::ShutdownModule()
 	}
 	
 	// Physics Control Profile editor/asset is disabled for now
-#ifdef ENABLE_PHYSICS_CONTROL_PROFILE_EDITOR
+#if ENABLE_PHYSICS_CONTROL_PROFILE_EDITOR
 	FEditorModeRegistry::Get().UnregisterMode(FPhysicsControlProfileEditorMode::ModeName);
+#endif
 
+#ifdef ENABLE_PHYSICS_CONTROL_PROFILE_ASSET
 	if (FModuleManager::Get().IsModuleLoaded("AssetTools"))
 	{
 		FAssetToolsModule::GetModule().Get().UnregisterAssetTypeActions(
diff --git a/Engine/Plugins/Experimental/PhysicsControl/Source/PhysicsControlEditor/Private/PhysicsControlProfileAssetActions.cpp b/Engine/Plugins/Experimental/PhysicsControl/Source/PhysicsControlEditor/Private/PhysicsControlProfileAssetActions.cpp
--- a/Engine/Plugins/Experimental/PhysicsControl/Source/PhysicsControlEditor/Private/PhysicsControlProfileAssetActions.cpp
+++ b/Engine/Plugins/Experimental/PhysicsControl/Source/PhysicsControlEditor/Private/PhysicsControlProfileAssetActions.cpp
@@ -29,6 +29,7 @@ uint32 FPhysicsControlProfileAssetActions::GetCategories()
 	return EAssetTypeCategories::Physics;
 }
 
+#if ENABLE_PHYSICS_CONTROL_PROFILE_EDITOR
 //======================================================================================================================
 void FPhysicsControlProfileAssetActions::OpenAssetEditor(
 	const TArray<UObject*>& InObjects, TSharedPtr<IToolkitHost> EditWithinLevelEditor)
@@ -45,3 +46,4 @@ void FPhysicsControlProfileAssetActions::OpenAssetEditor(
 		}
 	}
 }
+#endif
diff --git a/Engine/Plugins/Experimental/PhysicsControl/Source/PhysicsControlEditor/Private/PhysicsControlProfileAssetActions.h b/Engine/Plugins/Experimental/PhysicsControl/Source/PhysicsControlEditor/Private/PhysicsControlProfileAssetActions.h
--- a/Engine/Plugins/Experimental/PhysicsControl/Source/PhysicsControlEditor/Private/PhysicsControlProfileAssetActions.h
+++ b/Engine/Plugins/Experimental/PhysicsControl/Source/PhysicsControlEditor/Private/PhysicsControlProfileAssetActions.h
@@ -5,6 +5,11 @@
 #include "CoreMinimal.h"
 #include "AssetTypeActions_Base.h"
 
+#define ENABLE_PHYSICS_CONTROL_PROFILE_ASSET
+
+// Enable/disable the ability to actually create control profile assets in editor
+#define ENABLE_PHYSICS_CONTROL_PROFILE_EDITOR 0
+
 class FPhysicsControlProfileAssetActions : public FAssetTypeActions_Base
 {
 public:
@@ -13,8 +18,9 @@ public:
 	FColor GetTypeColor() const override;
 	uint32 GetCategories() override;
 
+#if ENABLE_PHYSICS_CONTROL_PROFILE_EDITOR
 	void OpenAssetEditor(
 		const TArray<UObject*>& InObjects, 
 		TSharedPtr<class IToolkitHost> EditWithinLevelEditor = TSharedPtr<IToolkitHost>()) override;
-
+#endif
 };
diff --git a/Engine/Plugins/Experimental/PhysicsControl/Source/PhysicsControlEditor/Private/PhysicsControlProfileAssetFactory.cpp b/Engine/Plugins/Experimental/PhysicsControl/Source/PhysicsControlEditor/Private/PhysicsControlProfileAssetFactory.cpp
--- a/Engine/Plugins/Experimental/PhysicsControl/Source/PhysicsControlEditor/Private/PhysicsControlProfileAssetFactory.cpp
+++ b/Engine/Plugins/Experimental/PhysicsControl/Source/PhysicsControlEditor/Private/PhysicsControlProfileAssetFactory.cpp
@@ -4,7 +4,8 @@
 #include "PhysicsControlProfileAsset.h"
 
 //======================================================================================================================
-UPhysicsControlProfileAssetFactory::UPhysicsControlProfileAssetFactory()
+UPhysicsControlProfileAssetFactory::UPhysicsControlProfileAssetFactory(const FObjectInitializer& ObjectInitializer)
+	: Super(ObjectInitializer)
 {
 	SupportedClass = UPhysicsControlProfileAsset::StaticClass();
 	bCreateNew = true;
diff --git a/Engine/Plugins/Experimental/PhysicsControl/Source/PhysicsControlEditor/Private/PhysicsControlProfileAssetFactory.h b/Engine/Plugins/Experimental/PhysicsControl/Source/PhysicsControlEditor/Private/PhysicsControlProfileAssetFactory.h
--- a/Engine/Plugins/Experimental/PhysicsControl/Source/PhysicsControlEditor/Private/PhysicsControlProfileAssetFactory.h
+++ b/Engine/Plugins/Experimental/PhysicsControl/Source/PhysicsControlEditor/Private/PhysicsControlProfileAssetFactory.h
@@ -11,7 +11,7 @@ class UPhysicsControlProfileAssetFactory : public UFactory
 {
 	GENERATED_BODY()
 public:
-	UPhysicsControlProfileAssetFactory();
+	UPhysicsControlProfileAssetFactory(const FObjectInitializer& ObjectInitializer);
 	UObject* FactoryCreateNew(
 		UClass* Class, UObject* InParent, FName Name, EObjectFlags Flags, UObject* Context, FFeedbackContext* Warn);
 };
diff --git a/Engine/Plugins/Experimental/PhysicsControl/Source/PhysicsControlUncookedOnly/Private/AnimGraphNode_RigidBodyWithControl.cpp b/Engine/Plugins/Experimental/PhysicsControl/Source/PhysicsControlUncookedOnly/Private/AnimGraphNode_RigidBodyWithControl.cpp
--- a/Engine/Plugins/Experimental/PhysicsControl/Source/PhysicsControlUncookedOnly/Private/AnimGraphNode_RigidBodyWithControl.cpp
+++ b/Engine/Plugins/Experimental/PhysicsControl/Source/PhysicsControlUncookedOnly/Private/AnimGraphNode_RigidBodyWithControl.cpp
@@ -1,14 +1,5 @@
 // Copyright Epic Games, Inc. All Rights Reserved.
 
-/* <<<< THIS IS TEMPORARY PROTOTYPE CODE >>>>
- *
- * This is essentially a copy + paste of the 
- * equivalent file in engine code and is 
- * intended for prototyping work as part of 
- * the Ch5 locomotion initiative.
- * 
- */
-
 #include "AnimGraphNode_RigidBodyWithControl.h"
 #include "AnimNodeEditModes.h"
 #include "Features/IModularFeatures.h"
@@ -322,13 +313,13 @@ TArray<TPair<FName, TArray<FName>>> UAnimGraphNode_RigidBodyWithControl::Generat
 
 		// These functions will create the base set of controls and modifiers from SetupData
 		TMap<FName, FPhysicsControlLimbBones> AllLimbBones =
-			GetLimbBones(Node.SetupData.LimbSetupData, RefSkeleton, Node.OverridePhysicsAsset);
+			GetLimbBones(Node.SetupData.LimbSetupData, RefSkeleton, Node.OverridePhysicsAsset.Get());
 
 		TSet<FName> BodyModifierNames;
 		TSet<FName> ControlNames;
-		FRigidBodyNameRecords NameRecords;
+		FPhysicsControlNameRecords NameRecords;
 
-		CollectOperatorNames(&Node, AllLimbBones, RefSkeleton, Node.OverridePhysicsAsset, BodyModifierNames, ControlNames, NameRecords);
+		CollectOperatorNames(&Node, AllLimbBones, RefSkeleton, Node.OverridePhysicsAsset.Get(), BodyModifierNames, ControlNames, NameRecords);
 
 		// Create any additional sets that have been requested
 		CreateAdditionalSets(Node.AdditionalSets, BodyModifierNames, ControlNames, NameRecords);
diff --git a/Engine/Source/Runtime/Experimental/Chaos/Private/Chaos/PBDJointConstraints.cpp b/Engine/Source/Runtime/Experimental/Chaos/Private/Chaos/PBDJointConstraints.cpp
--- a/Engine/Source/Runtime/Experimental/Chaos/Private/Chaos/PBDJointConstraints.cpp
+++ b/Engine/Source/Runtime/Experimental/Chaos/Private/Chaos/PBDJointConstraints.cpp
@@ -210,6 +210,16 @@ namespace Chaos
 		GetConstraintSettingsInternal().AngularDriveMaxTorque = MaxTorque;
 	}
 
+	void FPBDJointConstraintHandle::SetCollisionEnabled(const bool bCollisionEnabled)
+	{
+		GetConstraintSettingsInternal().bCollisionEnabled = bCollisionEnabled;
+	}
+
+	void FPBDJointConstraintHandle::SetParentInvMassScale(const FReal ParentInvMassScale)
+	{
+		GetConstraintSettingsInternal().ParentInvMassScale = ParentInvMassScale;
+	}
+
 	void FPBDJointConstraintHandle::SetDriveParams(
 		const FVec3 LinearStiffness, const FVec3 LinearDamping, const FVec3 MaxForce,
 		const FVec3 AngularStiffness, const FVec3 AngularDamping, const FVec3 MaxTorque)
diff --git a/Engine/Source/Runtime/Experimental/Chaos/Public/Chaos/Collision/BasicBroadPhase.h b/Engine/Source/Runtime/Experimental/Chaos/Public/Chaos/Collision/BasicBroadPhase.h
--- a/Engine/Source/Runtime/Experimental/Chaos/Public/Chaos/Collision/BasicBroadPhase.h
+++ b/Engine/Source/Runtime/Experimental/Chaos/Public/Chaos/Collision/BasicBroadPhase.h
@@ -86,6 +86,10 @@ namespace Chaos
 			FParticleHandle* ParticleB,
 			FParticleHandle* SearchParticle)
 		{
+			if (!ParticleA->HasCollision() || !ParticleB->HasCollision())
+			{
+				return;
+			}
 			const FAABB3& Box0 = ParticleA->WorldSpaceInflatedBounds();
 			const FAABB3& Box1 = ParticleB->WorldSpaceInflatedBounds();
 			if (Box0.Intersects(Box1))
diff --git a/Engine/Source/Runtime/Experimental/Chaos/Public/Chaos/PBDJointConstraints.h b/Engine/Source/Runtime/Experimental/Chaos/Public/Chaos/PBDJointConstraints.h
--- a/Engine/Source/Runtime/Experimental/Chaos/Public/Chaos/PBDJointConstraints.h
+++ b/Engine/Source/Runtime/Experimental/Chaos/Public/Chaos/PBDJointConstraints.h
@@ -73,6 +73,9 @@ namespace Chaos
 		CHAOS_API void SetAngularDriveDamping(const FVec3 Damping);
 		CHAOS_API void SetAngularDriveMaxTorque(const FVec3 MaxTorque);
 
+		CHAOS_API void SetCollisionEnabled(const bool bCollisionEnabled);
+		CHAOS_API void SetParentInvMassScale(const FReal ParentInvMassScale);
+
 		/**
 		 * This allows the most common drive parameters to be set in one call. Note that the 
 		 * individual drive elements will be enabled/disabled depending on the strength/damping
@@ -84,14 +87,6 @@ namespace Chaos
 			const FVec3 LinearStiffness, const FVec3 LinearDamping, const FVec3 MaxForce,
 			const FVec3 AngularStiffness, const FVec3 AngularDamping, const FVec3 MaxTorque);
 
-		// TODO remove - this is just a temporary workaround before fixing calling code
-		void SetDriveProperties(
-			const FVec3 LinearStiffness, const FVec3 LinearDamping, const FVec3 MaxForce,
-			const FVec3 AngularStiffness, const FVec3 AngularDamping, const FVec3 MaxTorque)
-		{
-			SetDriveParams(LinearStiffness, LinearDamping, MaxForce, AngularStiffness, AngularDamping, MaxTorque);
-		}
-
 		CHAOS_API TVec2<FGeometryParticleHandle*> GetConstrainedParticles() const override final;
 
 		static const FConstraintHandleTypeID& StaticType()
