Fix use of the physics override asset in the RigidBodyWithControl node (problems running cooked builds)

diff --git a/Engine/Plugins/Experimental/PhysicsControl/Source/PhysicsControl/Private/AnimNode_RigidBodyWithControl.cpp b/Engine/Plugins/Experimental/PhysicsControl/Source/PhysicsControl/Private/AnimNode_RigidBodyWithControl.cpp
--- a/Engine/Plugins/Experimental/PhysicsControl/Source/PhysicsControl/Private/AnimNode_RigidBodyWithControl.cpp
+++ b/Engine/Plugins/Experimental/PhysicsControl/Source/PhysicsControl/Private/AnimNode_RigidBodyWithControl.cpp
@@ -962,9 +962,9 @@ void ComputeBodyInsertionOrderWithControl(TArray<FBoneIndexType>& InsertionOrder
 
 UPhysicsAsset* FAnimNode_RigidBodyWithControl::GetPhysicsAssetToBeUsed(const UAnimInstance* InAnimInstance) const
 {
-	if (OverridePhysicsAsset.IsValid())
+	if (OverridePhysicsAsset)
 	{
-		return OverridePhysicsAsset.Get();
+		return ToRawPtr(OverridePhysicsAsset);
 	}
 
 	if (bDefaultToSkeletalMeshPhysicsAsset && ensure(InAnimInstance))
diff --git a/Engine/Plugins/Experimental/PhysicsControl/Source/PhysicsControl/Public/AnimNode_RigidBodyWithControl.h b/Engine/Plugins/Experimental/PhysicsControl/Source/PhysicsControl/Public/AnimNode_RigidBodyWithControl.h
--- a/Engine/Plugins/Experimental/PhysicsControl/Source/PhysicsControl/Public/AnimNode_RigidBodyWithControl.h
+++ b/Engine/Plugins/Experimental/PhysicsControl/Source/PhysicsControl/Public/AnimNode_RigidBodyWithControl.h
@@ -79,7 +79,7 @@ struct PHYSICSCONTROL_API FAnimNode_RigidBodyWithControl : public FAnimNode_Skel
 public:
 	/** Physics asset to use. If empty use the skeletal mesh's default physics asset */
 	UPROPERTY(EditAnywhere, Category = Settings)
-	TWeakObjectPtr<UPhysicsAsset> OverridePhysicsAsset;
+	TObjectPtr<UPhysicsAsset> OverridePhysicsAsset;
 
 	/** 
 	 * Use the skeletal mesh physics asset as default in case set to True. The Override Physics Asset 
diff --git a/Engine/Plugins/Experimental/PhysicsControl/Source/PhysicsControlEditor/Private/OperatorEditor/SOperatorTreeWidget.cpp b/Engine/Plugins/Experimental/PhysicsControl/Source/PhysicsControlEditor/Private/OperatorEditor/SOperatorTreeWidget.cpp
--- a/Engine/Plugins/Experimental/PhysicsControl/Source/PhysicsControlEditor/Private/OperatorEditor/SOperatorTreeWidget.cpp
+++ b/Engine/Plugins/Experimental/PhysicsControl/Source/PhysicsControlEditor/Private/OperatorEditor/SOperatorTreeWidget.cpp
@@ -315,7 +315,7 @@ void SOperatorTreeWidget::Refresh()
 			// Create node item
 			NodeItemPtr Node = MakeShared<OperatorTreeNodeItem>(*Itr, FName(*NodeUniqueName), BlueprintName);
 
-			if (!Itr->Node.OverridePhysicsAsset.IsValid())
+			if (!Itr->Node.OverridePhysicsAsset)
 			{
 				Node->AddChild(MakeShared<OperatorTreeMessageItem>(*Itr, OperatorTreeMessageItem::Warning, LOCTEXT("OperatorNoPhysicsAssetDefinedText", "Can't create operators for limbs because no Physics Asset has been defined for this node.")));
 			}
