Implementing Physics Control Profiles. This continues refactoring to align PhysicsControlComponent and RigidBodyWithControl. Also fixes some problems with collisions

diff --git a/Engine/Plugins/Experimental/PhysicsControl/Source/PhysicsControl/Private/AnimNode_RigidBodyWithControl.cpp b/Engine/Plugins/Experimental/PhysicsControl/Source/PhysicsControl/Private/AnimNode_RigidBodyWithControl.cpp
--- a/Engine/Plugins/Experimental/PhysicsControl/Source/PhysicsControl/Private/AnimNode_RigidBodyWithControl.cpp
+++ b/Engine/Plugins/Experimental/PhysicsControl/Source/PhysicsControl/Private/AnimNode_RigidBodyWithControl.cpp
@@ -1,9 +1,5 @@
 // Copyright Epic Games, Inc. All Rights Reserved.
 
-/* 
- * This is essentially a copy + paste of the RigidBody file
- */
-
 #include "AnimNode_RigidBodyWithControl.h"
 #include "PhysicsControlLog.h"
 #include "AnimationRuntime.h"
@@ -26,17 +22,19 @@
 #include "PhysicsEngine/PhysicsSettings.h"
 #include "Logging/MessageLog.h"
 
+#include "PhysicsProxy/SingleParticlePhysicsProxy.h"
+
 #include "Chaos/PBDJointConstraintTypes.h"
 
 #include UE_INLINE_GENERATED_CPP_BY_NAME(AnimNode_RigidBodyWithControl)
 
 LLM_DEFINE_TAG(Animation_RigidBodyWithControl);
 
-//UE_DISABLE_OPTIMIZATION;
-
 /////////////////////////////////////////////////////
 // FAnimNode_RigidBodyWithControl
 
+//UE_DISABLE_OPTIMIZATION
+
 #define LOCTEXT_NAMESPACE "ImmediatePhysicsWithControl"
 
 DECLARE_CYCLE_STAT(TEXT("RigidBodyNodeWithControl_InitPhysics"), STAT_RigidBodyWithControlInitPhysicsTime, STATGROUP_Anim);
@@ -44,6 +42,7 @@ DECLARE_CYCLE_STAT(TEXT("RigidBodyNodeWithControl_SetupControls"), STAT_RigidBod
 
 DECLARE_CYCLE_STAT(TEXT("RigidBodyNodeWithControl_Eval"), STAT_RigidBodyNodeWithControl_Eval, STATGROUP_Anim);
 DECLARE_CYCLE_STAT(TEXT("RigidBodyNodeWithControl_Simulation"), STAT_RigidBodyNodeWithControl_Simulation, STATGROUP_Anim);
+DECLARE_CYCLE_STAT(TEXT("RigidBodyNodeWithControl_UpdateWorldObjects"), STAT_RigidBodyNodeWithControl_UpdateWorldObjects, STATGROUP_Anim);
 DECLARE_CYCLE_STAT(TEXT("RigidBodyNodeWithControl_SimulationWait"), STAT_RigidBodyNodeWithControl_SimulationWait, STATGROUP_Anim);
 DECLARE_CYCLE_STAT(TEXT("RigidBodyNodeWithControl_PreUpdate"), STAT_RigidBodyNodeWithControl_PreUpdate, STATGROUP_Anim);
 DECLARE_CYCLE_STAT(TEXT("RigidBodyNodeWithControl_Update"), STAT_RigidBodyNodeWithControl_Update, STATGROUP_Anim);
@@ -130,9 +129,10 @@ FAnimNode_RigidBodyWithControl::FAnimNode_RigidBodyWithControl()
 	, ComponentAppliedLinearAccClamp(10000.0f)
 	, SimSpaceSettings()
 	, CachedBoundsScale(1.2f)
-	, UpdateCacheEveryFrame(true)
+	, bUpdateCacheEveryFrame(true)
 	, BaseBoneRef()
 	, OverlapChannel(ECC_WorldStatic)
+	, bCalculateVelocitiesForWorldGeometry(true)
 	, SimulationSpace(ESimulationSpace::ComponentSpace)
 	, bForceDisableCollisionBetweenConstraintBodies(false)
 	, bUseExternalClothCollision(false)
@@ -173,7 +173,6 @@ FAnimNode_RigidBodyWithControl::FAnimNode_RigidBodyWithControl()
 	, TotalMass(0.0f)
 	, CachedBounds(FVector::ZeroVector, 0.0f)
 	, QueryParams(NAME_None, FCollisionQueryParams::GetUnknownStatId())
-	, PhysScene(nullptr)
 	, UnsafeWorld(nullptr)
 	, UnsafeOwner(nullptr)
 	, CapturedBoneVelocityBoneContainer()
@@ -336,29 +335,25 @@ void FAnimNode_RigidBodyWithControl::InitSimulationSpace(
 }
 
 void FAnimNode_RigidBodyWithControl::CalculateSimulationSpace(
-	ESimulationSpace Space, 
-	const FTransform& ComponentToWorld, 
-	const FTransform& BoneToComponent,
-	const float Dt,
+	ESimulationSpace         Space, 
+	const FTransform&        ComponentToWorld, 
+	const FTransform&        BoneToComponent,
+	const float              Dt,
 	const FSimSpaceSettings& Settings,
-	FTransform& SpaceTransform, 
-	FVector& SpaceLinearVel, 
-	FVector& SpaceAngularVel, 
-	FVector& SpaceLinearAcc, 
-	FVector& SpaceAngularAcc)
+	FSimulationSpaceData&    OutSimulationSpaceData)
 {
 	// World-space transform of the simulation space
-	SpaceTransform = SpaceToWorldTransform(Space, ComponentToWorld, BoneToComponent);
-	SpaceLinearVel = FVector::ZeroVector;
-	SpaceAngularVel = FVector::ZeroVector;
-	SpaceLinearAcc = FVector::ZeroVector;
-	SpaceAngularAcc = FVector::ZeroVector;
+	OutSimulationSpaceData.Transform = SpaceToWorldTransform(Space, ComponentToWorld, BoneToComponent);
+	OutSimulationSpaceData.LinearVel = FVector::ZeroVector;
+	OutSimulationSpaceData.AngularVel = FVector::ZeroVector;
+	OutSimulationSpaceData.LinearAcc = FVector::ZeroVector;
+	OutSimulationSpaceData.AngularAcc = FVector::ZeroVector;
 
 	// The simulation scale does not change - we scale the inputs and outputs instead.
 	// This means we do not support phantom forces resulting from scale changes, but that's ok.
 	// NOTE: If we don't clear the scale, rapid scaling to zero can introduce large phantom forces
 	// leading to major instability in the simulation
-	SpaceTransform.SetScale3D(FVector::One());
+	OutSimulationSpaceData.Transform.SetScale3D(FVector::One());
 
 	// If the system is disabled, nothing else to do
 	if ((Settings.WorldAlpha == 0.0f) || (Dt < SMALL_NUMBER))
@@ -368,8 +363,8 @@ void FAnimNode_RigidBodyWithControl::CalculateSimulationSpace(
 
 	if (Space == ESimulationSpace::WorldSpace)
 	{
-		SpaceLinearVel = Settings.ExternalLinearVelocity;
-		SpaceAngularVel = Settings.ExternalAngularVelocity;
+		OutSimulationSpaceData.LinearVel = Settings.ExternalLinearVelocity;
+		OutSimulationSpaceData.AngularVel = Settings.ExternalAngularVelocity;
 		return;
 	}
 
@@ -387,10 +382,10 @@ void FAnimNode_RigidBodyWithControl::CalculateSimulationSpace(
 		CompLinVel.Z *= Settings.VelocityScaleZ;
 		CompLinAcc.Z *= Settings.VelocityScaleZ;
 
-		SpaceLinearVel = CompLinVel.GetClampedToMaxSize(Settings.MaxLinearVelocity) + Settings.ExternalLinearVelocity;
-		SpaceAngularVel = CompAngVel.GetClampedToMaxSize(Settings.MaxAngularVelocity) + Settings.ExternalAngularVelocity;
-		SpaceLinearAcc = CompLinAcc.GetClampedToMaxSize(Settings.MaxLinearAcceleration);
-		SpaceAngularAcc = CompAngAcc.GetClampedToMaxSize(Settings.MaxAngularAcceleration);
+		OutSimulationSpaceData.LinearVel = CompLinVel.GetClampedToMaxSize(Settings.MaxLinearVelocity) + Settings.ExternalLinearVelocity;
+		OutSimulationSpaceData.AngularVel = CompAngVel.GetClampedToMaxSize(Settings.MaxAngularVelocity) + Settings.ExternalAngularVelocity;
+		OutSimulationSpaceData.LinearAcc = CompLinAcc.GetClampedToMaxSize(Settings.MaxLinearAcceleration);
+		OutSimulationSpaceData.AngularAcc = CompAngAcc.GetClampedToMaxSize(Settings.MaxAngularAcceleration);
 		return;
 	}
 	
@@ -428,10 +423,10 @@ void FAnimNode_RigidBodyWithControl::CalculateSimulationSpace(
 		NetLinVel.Z *= Settings.VelocityScaleZ;
 		NetLinAcc.Z *= Settings.VelocityScaleZ;
 
-		SpaceLinearVel = NetLinVel.GetClampedToMaxSize(Settings.MaxLinearVelocity) + Settings.ExternalLinearVelocity;
-		SpaceAngularVel = NetAngVel.GetClampedToMaxSize(Settings.MaxAngularVelocity) + Settings.ExternalAngularVelocity;
-		SpaceLinearAcc = NetLinAcc.GetClampedToMaxSize(Settings.MaxLinearAcceleration);
-		SpaceAngularAcc = NetAngAcc.GetClampedToMaxSize(Settings.MaxAngularAcceleration);
+		OutSimulationSpaceData.LinearVel = NetLinVel.GetClampedToMaxSize(Settings.MaxLinearVelocity) + Settings.ExternalLinearVelocity;
+		OutSimulationSpaceData.AngularVel = NetAngVel.GetClampedToMaxSize(Settings.MaxAngularVelocity) + Settings.ExternalAngularVelocity;
+		OutSimulationSpaceData.LinearAcc = NetLinAcc.GetClampedToMaxSize(Settings.MaxLinearAcceleration);
+		OutSimulationSpaceData.AngularAcc = NetAngAcc.GetClampedToMaxSize(Settings.MaxAngularAcceleration);
 		return;
 	}
 }
@@ -467,33 +462,6 @@ void FAnimNode_RigidBodyWithControl::DestroyPhysicsSimulation()
 	PhysicsSimulation = nullptr;
 }
 
-static FQuat CalculateTargetQuat(
-	const Chaos::FPBDJointSettings& Settings, 
-	const FTransform&               ChildBoneTM,
-	const FTransform&               ParentBoneTM)
-{
-	FQuat ChildBoneQuat = ChildBoneTM.GetRotation();
-	FQuat ParentBoneQuat = ParentBoneTM.GetRotation();
-	FQuat ChildConnectorQuat = Settings.ConnectorTransforms[0].GetRotation();
-	FQuat ParentConnectorQuat = Settings.ConnectorTransforms[1].GetRotation();
-	FQuat LimitQuat = ParentBoneQuat * ParentConnectorQuat;
-	FQuat DialQuat = ChildBoneQuat * ChildConnectorQuat;
-	return LimitQuat.Inverse() * DialQuat;
-}
-
-static FTransform CalculateJointTargetTransform(
-	const Chaos::FPBDJointSettings& Settings,
-	const FTransform& ChildBoneTM,
-	const FTransform& ParentBoneTM)
-{
-	const FTransform& ChildLocalConnectorTM = Settings.ConnectorTransforms[0];
-	const FTransform& ParentLocalConnectorTM = Settings.ConnectorTransforms[1];
-	const FTransform ChildConnectorTM = ChildBoneTM * ChildLocalConnectorTM;
-	const FTransform ParentConnectorTM = ParentBoneTM * ParentLocalConnectorTM;
-
-	return ParentConnectorTM.Inverse() * ChildConnectorTM;
-}
-
 void FAnimNode_RigidBodyWithControl::SetupControls(USkeletalMeshComponent* const SkeletalMeshComponent)
 {
 	SCOPE_CYCLE_COUNTER(STAT_RigidBodyWithControlSetupControlsTime);
@@ -510,12 +478,12 @@ void FAnimNode_RigidBodyWithControl::SetupControls(USkeletalMeshComponent* const
 		}
 		else
 		{
-			UE_LOG(LogRigidBodyWithControl, Warning, TEXT("Invalid Skeletal Mesh"));
+			UE_LOG(LogPhysicsControl, Warning, TEXT("Invalid Skeletal Mesh"));
 		}
 	}
 	else
 	{
-		UE_LOG(LogRigidBodyWithControl, Warning, TEXT("Invalid Skeletal Mesh Component"));
+		UE_LOG(LogPhysicsControl, Warning, TEXT("Invalid Skeletal Mesh Component"));
 	}
 }
 
@@ -570,7 +538,7 @@ void FAnimNode_RigidBodyWithControl::EvaluateSkeletalControl_AnyThread(FComponen
 			// being off-screen or LOD changes.
 			if (WorldTimeSeconds - (LastEvalTimeSeconds + DeltaSeconds) > EvaluationResetTime)
 			{
-				UE_LOG(LogRigidBodyWithControl, Verbose, TEXT("%s Time-Based Reset"), *Output.AnimInstanceProxy->GetAnimInstanceName());
+				UE_LOG(LogPhysicsControl, Verbose, TEXT("%s Time-Based Reset"), *Output.AnimInstanceProxy->GetAnimInstanceName());
 				ResetSimulatedTeleportType = ETeleportType::ResetPhysics;
 			}
 		}
@@ -649,7 +617,7 @@ void FAnimNode_RigidBodyWithControl::EvaluateSkeletalControl_AnyThread(FComponen
 			{
 				case ETeleportType::TeleportPhysics:
 				{
-					UE_LOG(LogRigidBodyWithControl, Verbose, TEXT("%s TeleportPhysics (Scale: %f %f %f)"), *Output.AnimInstanceProxy->GetAnimInstanceName(), CompWorldSpaceTM.GetScale3D().X, CompWorldSpaceTM.GetScale3D().Y, CompWorldSpaceTM.GetScale3D().Z);
+					UE_LOG(LogPhysicsControl, Verbose, TEXT("%s TeleportPhysics (Scale: %f %f %f)"), *Output.AnimInstanceProxy->GetAnimInstanceName(), CompWorldSpaceTM.GetScale3D().X, CompWorldSpaceTM.GetScale3D().Y, CompWorldSpaceTM.GetScale3D().Z);
 
 					// Teleport bodies.
 					for (const RigidBodyWithControl::FOutputBoneData& OutputData : OutputBoneData)
@@ -676,7 +644,7 @@ void FAnimNode_RigidBodyWithControl::EvaluateSkeletalControl_AnyThread(FComponen
 
 				case ETeleportType::ResetPhysics:
 				{
-					UE_LOG(LogRigidBodyWithControl, Verbose, TEXT("%s ResetPhysics (Scale: %f %f %f)"), *Output.AnimInstanceProxy->GetAnimInstanceName(), CompWorldSpaceTM.GetScale3D().X, CompWorldSpaceTM.GetScale3D().Y, CompWorldSpaceTM.GetScale3D().Z);
+					UE_LOG(LogPhysicsControl, Verbose, TEXT("%s ResetPhysics (Scale: %f %f %f)"), *Output.AnimInstanceProxy->GetAnimInstanceName(), CompWorldSpaceTM.GetScale3D().X, CompWorldSpaceTM.GetScale3D().Y, CompWorldSpaceTM.GetScale3D().Z);
 
 					InitSimulationSpace(CompWorldSpaceTM, BaseBoneTM);
 
@@ -794,12 +762,18 @@ void FAnimNode_RigidBodyWithControl::EvaluateSkeletalControl_AnyThread(FComponen
 				ApplyCurrentConstraintProfile();
 			}
 
+			if (ControlProfile != CurrentControlProfile)
+			{
+				CurrentControlProfile = ControlProfile;
+				ApplyCurrentControlProfile();
+			}
+
 			// Apply the controls. Note that these won't set kinematic targets - we'll do that afterwards as they
 			// can apply to bodies that aren't under the influence of a modifier
 			ApplyControlAndModifierUpdatesAndParametersToRecords(ControlAndModifierUpdates, ControlAndModifierParameters);
 			ApplyControlsAndModifiers(SimSpaceGravity, DeltaSeconds);
 
-			// @todo(ccaulfield): We should be interpolating kinematic targets for each sub-step below
+			// Note that the simulation interpolates kinematic targets to handle substepping
 			for (const RigidBodyWithControl::FOutputBoneData& OutputData : OutputBoneData)
 			{
 				const int32 BodyIndex = OutputData.BodyIndex;
@@ -816,42 +790,26 @@ void FAnimNode_RigidBodyWithControl::EvaluateSkeletalControl_AnyThread(FComponen
 			// the duplication is not likely to be expensive)
 			ApplyKinematicTargets();
 
-			FSimSpaceSettings* UseSimSpaceSettings = &SimSpaceSettings;
+			FSimSpaceSettings* SimSpaceSettingsToUse = &SimSpaceSettings;
 			if (bRBANWithControl_SimSpace_EnableOverride)
 			{
-				UseSimSpaceSettings = &RBANWithControl_SimSpaceOverride;
+				SimSpaceSettingsToUse = &RBANWithControl_SimSpaceOverride;
 			}
 
-			FTransform SimulationTransform;
-			FVector SimulationLinearVelocity;
-			FVector SimulationAngularVelocity;
-			FVector SimulationLinearAcceleration;
-			FVector SimulationAngularAcceleration;
+			FSimulationSpaceData SimulationSpaceData;
 			CalculateSimulationSpace(
-				SimulationSpace, 
-				CompWorldSpaceTM, 
-				BaseBoneTM,
-				DeltaSeconds,
-				*UseSimSpaceSettings,
-				SimulationTransform,
-				SimulationLinearVelocity,
-				SimulationAngularVelocity,
-				SimulationLinearAcceleration,
-				SimulationAngularAcceleration);
-
-			UpdateWorldObjects(SimulationTransform, DeltaSeconds);
-			UpdateClothColliderObjects(SimulationTransform);
+				SimulationSpace, CompWorldSpaceTM, BaseBoneTM, DeltaSeconds, 
+				*SimSpaceSettingsToUse, SimulationSpaceData);
+
+			UpdateWorldObjects(SimulationSpaceData, DeltaSeconds);
+			UpdateClothColliderObjects(SimulationSpaceData);
 
 			PhysicsSimulation->UpdateSimulationSpace(
-				SimulationTransform, 
-				SimulationLinearVelocity,
-				SimulationAngularVelocity,
-				SimulationLinearAcceleration,
-				SimulationAngularAcceleration);
+				SimulationSpaceData.Transform, SimulationSpaceData.LinearVel, SimulationSpaceData.AngularVel,
+				SimulationSpaceData.LinearAcc, SimulationSpaceData.AngularAcc);
 
 			PhysicsSimulation->SetSimulationSpaceSettings(
-				UseSimSpaceSettings->WorldAlpha,
-				UseSimSpaceSettings->ExternalLinearDragV);
+				SimSpaceSettingsToUse->WorldAlpha, SimSpaceSettingsToUse->ExternalLinearDragV);
 
 			PhysicsSimulation->SetSolverSettings(
 				SolverSettings.FixedTimeStep,
@@ -962,7 +920,7 @@ void ComputeBodyInsertionOrderWithControl(TArray<FBoneIndexType>& InsertionOrder
 
 UPhysicsAsset* FAnimNode_RigidBodyWithControl::GetPhysicsAssetToBeUsed(const UAnimInstance* InAnimInstance) const
 {
-	if (OverridePhysicsAsset)
+	if (IsValid(OverridePhysicsAsset))
 	{
 		return ToRawPtr(OverridePhysicsAsset);
 	}
@@ -1110,30 +1068,33 @@ void FAnimNode_RigidBodyWithControl::InitPhysics(const UAnimInstance* InAnimInst
 				// as kinematic prevents them being subsequently made dynamic.
 				bool bSimulated = (BodySetup->PhysicsType == EPhysicsType::PhysType_Simulated);
 				ImmediatePhysics::EActorType ActorType = ImmediatePhysics::EActorType::DynamicActor;
-				ImmediatePhysics::FActorHandle* NewBodyHandle = PhysicsSimulation->CreateActor(ActorType, BodyInstance, BodyInstance->GetUnrealWorldTransform());
-				if (NewBodyHandle)
+				ImmediatePhysics::FActorHandle* ActorHandle = PhysicsSimulation->CreateActor(
+					ActorType, BodyInstance, BodyInstance->GetUnrealWorldTransform());
+				if (ActorHandle)
 				{
-					const float InvMass = (float) NewBodyHandle->GetInverseMass();
+					ActorHandle->InitWorldTransform(BodyInstance->GetUnrealWorldTransform());
+
+					const float InvMass = (float) ActorHandle->GetInverseMass();
 					TotalMass += InvMass > 0.f ? 1.f / InvMass : 0.f;
 					if (!bSimulated)
 					{
-						NewBodyHandle->SetIsKinematic(true);
+						ActorHandle->SetIsKinematic(true);
 					}
 
-					NewBodyHandle->SetName(BodySetup->BoneName);
-					const int32 BodyIndex = AddBody(NewBodyHandle);
+					ActorHandle->SetName(BodySetup->BoneName);
+					const int32 BodyIndex = AddBody(ActorHandle);
 					const int32 SkeletonBoneIndex = MeshToSkeletonBoneIndex[InsertBone];
 					if (ensure(SkeletonBoneIndex >= 0))
 					{
 						SkeletonBoneIndexToBodyIndex[SkeletonBoneIndex] = BodyIndex;
 					}
 					BodyAnimData[BodyIndex].bIsSimulated = bSimulated;
-					NamesToHandles.Add(BodySetup->BoneName, NewBodyHandle);
-					BodyIndexToActorHandle[BodyInstance->InstanceBodyIndex] = NewBodyHandle;
+					NamesToHandles.Add(BodySetup->BoneName, ActorHandle);
+					BodyIndexToActorHandle[BodyInstance->InstanceBodyIndex] = ActorHandle;
 
 					if (BodySetup->CollisionReponse == EBodyCollisionResponse::BodyCollision_Disabled)
 					{
-						IgnoreCollisionActors.Add(NewBodyHandle);
+						IgnoreCollisionActors.Add(ActorHandle);
 					}
 				}
 			}
@@ -1220,15 +1181,11 @@ void FAnimNode_RigidBodyWithControl::InitPhysics(const UAnimInstance* InAnimInst
 			ResetSimulatedTeleportType = ETeleportType::ResetPhysics;
 		}
 
-		//== WithControl ==================================
-
 		if (!bHaveSetupControls && bEnableControls)
 		{
 			SetupControls(InAnimInstance->GetSkelMeshComponent());
 		}
 
-		//== WithControl ==================================
-
 		// Terminate all the constraint instances
 		for (FConstraintInstance* CI : HighLevelConstraintInstances)
 		{
@@ -1297,33 +1254,25 @@ void FAnimNode_RigidBodyWithControl::UpdateWorldGeometry(const UWorld& World, co
 	
 	SCOPE_CYCLE_COUNTER(STAT_ImmediateUpdateWorldGeometry);
 	QueryParams = FCollisionQueryParams(SCENE_QUERY_STAT(RagdollNodeFindGeometry), /*bTraceComplex=*/false);
-#if WITH_EDITOR
-	if(!World.IsGameWorld())
-	{
-		QueryParams.MobilityType = EQueryMobilityType::Any;	//If we're in some preview world trace against everything because things like the preview floor are not static
-		QueryParams.AddIgnoredComponent(&SKC);
-	}
-	else
-#endif
-	{
-		QueryParams.MobilityType = EQueryMobilityType::Static;	//We only want static actors
-	}
+	// The Mobility type is ignored anyway - see UE-168341
+	// When that's fixed, we could expose this, as we might not want to collide with dynamic objects etc
+	QueryParams.MobilityType = EQueryMobilityType::Any;	
+	QueryParams.AddIgnoredComponent(&SKC);
 
 	// Check for deleted world objects and flag for removal (later in anim task)
 	ExpireWorldObjects();
 
 	// If we have moved outside of the bounds we checked for world objects we need to gather new world objects
 	FSphere Bounds = SKC.CalcBounds(SKC.GetComponentToWorld()).GetSphere();
-	if (!Bounds.IsInside(CachedBounds) || UpdateCacheEveryFrame)
+	if (!Bounds.IsInside(CachedBounds) || bUpdateCacheEveryFrame)
 	{
 		// Since the cached bounds are no longer valid, update them.
 		CachedBounds = Bounds;
 		CachedBounds.W *= CachedBoundsScale;
 
-		// Cache the PhysScene and World for use in UpdateWorldForces and CollectWorldObjects. When
+		// Cache the World for use in UpdateWorldForces and CollectWorldObjects. When
 		// these are non-null it is an indicator that we need to update the collected world objects
 		// list in CollectWorldObjects
-		PhysScene = World.GetPhysicsScene();
 		UnsafeWorld = &World;
 		UnsafeOwner = SKC.GetOwner();
 
@@ -1578,7 +1527,6 @@ void FAnimNode_RigidBodyWithControl::UpdateInternal(const FAnimationUpdateContex
 	// These get set again if our bounds change. Subsequent calls to CollectWorldObjects will early-out until then
 	UnsafeWorld = nullptr;
 	UnsafeOwner = nullptr;
-	PhysScene = nullptr;
 }
 
 void FAnimNode_RigidBodyWithControl::CollectClothColliderObjects(const USkeletalMeshComponent* SkeletalMeshComp)
@@ -1631,7 +1579,7 @@ void FAnimNode_RigidBodyWithControl::RemoveClothColliderObjects()
 	ClothColliders.Reset();
 }
 
-void FAnimNode_RigidBodyWithControl::UpdateClothColliderObjects(const FTransform& SpaceTransform)
+void FAnimNode_RigidBodyWithControl::UpdateClothColliderObjects(const FSimulationSpaceData& SimulationSpaceData)
 {
 	for (FClothCollider& ClothCollider : ClothColliders)
 	{
@@ -1640,8 +1588,10 @@ void FAnimNode_RigidBodyWithControl::UpdateClothColliderObjects(const FTransform
 			// Calculate the sim-space transform of this object
 			const FTransform CompWorldTransform = ClothCollider.SkeletalMeshComponent->GetBoneTransform(ClothCollider.BoneIndex);
 			FTransform CompSpaceTransform;
-			CompSpaceTransform.SetTranslation(SpaceTransform.InverseTransformPosition(CompWorldTransform.GetLocation()));
-			CompSpaceTransform.SetRotation(SpaceTransform.InverseTransformRotation(CompWorldTransform.GetRotation()));
+			CompSpaceTransform.SetTranslation(SimulationSpaceData.Transform.InverseTransformPosition(
+				CompWorldTransform.GetLocation()));
+			CompSpaceTransform.SetRotation(SimulationSpaceData.Transform.InverseTransformRotation(
+				CompWorldTransform.GetRotation()));
 			CompSpaceTransform.SetScale3D(FVector::OneVector);	// TODO - sort out scale for world objects in local sim
 
 			// Update the sim's copy of the world object
@@ -1650,13 +1600,38 @@ void FAnimNode_RigidBodyWithControl::UpdateClothColliderObjects(const FTransform
 	}
 }
 
+inline bool IsComponentDesiredInSim(const UPrimitiveComponent* Component)
+{
+	if (!IsValid(Component))
+	{
+		return false;
+	}
+	if (!Component->GetBodyInstance())
+	{
+		return false;
+	}
+	if (!Component->GetBodyInstance()->IsValidBodyInstance())
+	{
+		return false;
+	}
+	if (!Component->GetBodyInstance()->GetCollisionEnabled())
+	{
+		return false;
+	}
+	return true;
+}
+
+// TODO I think this can crash when running deferred, since it accesses UnsafeWorld from a worker thread
 void FAnimNode_RigidBodyWithControl::CollectWorldObjects()
 {
-	if ((UnsafeWorld != nullptr) && (PhysScene != nullptr))
+	if (UnsafeWorld)
 	{
 		// @todo(ccaulfield): should this use CachedBounds?
 		TArray<FOverlapResult> Overlaps;
-		UnsafeWorld->OverlapMultiByChannel(Overlaps, CachedBounds.Center, FQuat::Identity, OverlapChannel, FCollisionShape::MakeSphere((float) CachedBounds.W), QueryParams, FCollisionResponseParams(ECR_Overlap));
+		UnsafeWorld->OverlapMultiByChannel(
+			Overlaps, CachedBounds.Center, FQuat::Identity, OverlapChannel, 
+			FCollisionShape::MakeSphere((float) CachedBounds.W), 
+			QueryParams, FCollisionResponseParams(ECR_Overlap));
 
 		for (const FOverlapResult& Overlap : Overlaps)
 		{
@@ -1672,19 +1647,23 @@ void FAnimNode_RigidBodyWithControl::CollectWorldObjects()
 				{
 					// New object - add it to the sim
 					const bool bIsSelf = (UnsafeOwner == OverlapComp->GetOwner());
-					if (!bIsSelf)
+					if (!bIsSelf && IsComponentDesiredInSim(OverlapComp))
 					{
 						// Note that the TM here isn't correct unless we are actually using
 						// world-space simulation, but it will be updated.
 						FTransform TM = OverlapComp->GetComponentTransform();
 						// Create a kinematic actor. Not using Static as world-static objects may
-						// move in the simulation's frame of reference
-						ImmediatePhysics::FActorHandle* ActorHandle = PhysicsSimulation->CreateActor(ImmediatePhysics::EActorType::KinematicActor, &OverlapComp->BodyInstance, TM);
+						// move in the simulation's frame of reference. Note that TM may be in the wrong space here
+						ImmediatePhysics::FActorHandle* ActorHandle = PhysicsSimulation->CreateActor(
+							ImmediatePhysics::EActorType::KinematicActor, &OverlapComp->BodyInstance, TM);
+						ActorHandle->SetName(FName(OverlapComp->GetName()));
 						PhysicsSimulation->AddToCollidingPairs(ActorHandle);
-						// This seems necessary in order for subsequent calls to set the transform
-						// to update things correctly.
-						//ActorHandle->InitWorldTransform(TM);
 						ComponentsInSim.Add(OverlapComp, FWorldObject(ActorHandle, ComponentsInSimTick));
+						// We need this: 
+						// OverlapComp->BodyInstance.bUpdateKinematicFromSimulation = true;
+						// in order to get velocities from kinematic objects, but we can't as there's no way of 
+						// reliably undoing it (multiple RBWC nodes might be interacting with the same world object).
+						// So for now accept that we won't be able to pull velocities from kinematics.
 					}
 				}
 			}
@@ -1704,14 +1683,8 @@ void FAnimNode_RigidBodyWithControl::ExpireWorldObjects()
 
 		// Do we need to expire this object?
 		const int32 ExpireTickCount = RBANWithControl_WorldObjectExpiry;
-		bool bIsInvalid =
-			((ComponentsInSimTick - WorldObject.LastSeenTick) > ExpireTickCount)	// Haven't seen this object for a while
-			|| !IsValid(WorldComp)
-			|| (WorldComp->GetBodyInstance() == nullptr)
-			|| (!WorldComp->GetBodyInstance()->IsValidBodyInstance());
-
-		// Remove from sim if necessary
-		if (bIsInvalid)
+		if (((ComponentsInSimTick - WorldObject.LastSeenTick) > ExpireTickCount)	// Haven't seen this object for a while
+			|| !IsComponentDesiredInSim(WorldComp))
 		{
 			WorldObject.bExpired = true;
 		}
@@ -1745,47 +1718,95 @@ void FAnimNode_RigidBodyWithControl::PurgeExpiredWorldObjects()
 // Update the transforms of the world objects we added to the sim. This could be because we're a
 // component-based simulation and need to update stationary objects into our space, or simply
 // because the objects are moving.
-void FAnimNode_RigidBodyWithControl::UpdateWorldObjects(const FTransform& SpaceTransform, const float DeltaSeconds)
+void FAnimNode_RigidBodyWithControl::UpdateWorldObjects(
+	const FSimulationSpaceData& SimulationSpaceData, const float DeltaSeconds)
 {
-	LLM_SCOPE_BYNAME(TEXT("Animation/RigidBodyWithControl")); 
+	SCOPE_CYCLE_COUNTER(STAT_RigidBodyNodeWithControl_UpdateWorldObjects);
+	LLM_SCOPE_BYNAME(TEXT("Animation/RigidBodyWithControl"));
 
-	for (const auto& WorldEntry : ComponentsInSim)
-	{ 
-		const UPrimitiveComponent* OverlapComp = WorldEntry.Key;
-		if (OverlapComp != nullptr)
-		{
-			ImmediatePhysics::FActorHandle* ActorHandle = WorldEntry.Value.ActorHandle;
+	if (ComponentsInSim.IsEmpty())
+	{
+		return;
+	}
 
-			// Calculate the sim-space transform of this object
-			const FTransform CompWorldTransform = OverlapComp->BodyInstance.GetUnrealWorldTransform();
-			const FVector CompSpacePosition = SpaceTransform.InverseTransformPosition(CompWorldTransform.GetLocation());
-			const FQuat CompSpaceOrientation = SpaceTransform.InverseTransformRotation(CompWorldTransform.GetRotation());
+	if (bCalculateVelocitiesForWorldGeometry)
+	{
+		// We will want the "previous" simulation space to calculate velocities. Get it from the
+		// velocity, so we can take advantage of that being zero at initialisation/reset times etc.
+		const FVector PrevSimSpacePosition = 
+			SimulationSpaceData.Transform.GetLocation() - SimulationSpaceData.LinearVel * DeltaSeconds;
+		const FQuat SimSpaceDeltaOrientation = FQuat::MakeFromRotationVector(SimulationSpaceData.AngularVel * DeltaSeconds);
+		const FQuat PrevSimSpaceOrientation = 
+			SimSpaceDeltaOrientation.Inverse() * SimulationSpaceData.Transform.GetRotation();
+		const FTransform PrevSimSpaceTransform(PrevSimSpaceOrientation, PrevSimSpacePosition, FVector::OneVector);
+
+		for (const auto& WorldEntry : ComponentsInSim)
+		{ 
+			const UPrimitiveComponent* OverlapComp = WorldEntry.Key;
+			if (OverlapComp != nullptr)
+			{
+				ImmediatePhysics::FActorHandle* ActorHandle = WorldEntry.Value.ActorHandle;
 
-			// TODO - sort out scale for world objects in local sim
-			FTransform CompSpaceTransform(CompSpaceOrientation, CompSpacePosition, FVector::OneVector);
+				// Calculate the sim-space transform of this object
+				const FTransform CompWorldTransform = OverlapComp->BodyInstance.GetUnrealWorldTransform();
+				const FVector WorldPosition = CompWorldTransform.GetLocation();
+				const FQuat WorldOrientation = CompWorldTransform.GetRotation();
 
-			// This is what we don't do as it does not track velocity. Also, it only seems to
-			// work if InitWorldTransform has been called first (after actor creation).
-			// ActorHandle->SetWorldTransform(CompSpaceTransform);
+				const FVector CompSpacePosition = SimulationSpaceData.Transform.InverseTransformPosition(WorldPosition);
+				const FQuat CompSpaceOrientation = SimulationSpaceData.Transform.InverseTransformRotation(WorldOrientation);
 
-			// We also need to set the velocity - either because the object or the space is
-			// moving. Ideally would do this by setting the velocity (with transformations), but
-			// it's not possible to set the velocity of kinematics. Use the kinematic target to
-			// force the velocities. This should handle initialization and teleportation too.
-			const FVector Velocity = OverlapComp->BodyInstance.GetUnrealWorldVelocity();
-			const FVector AngularVelocity = OverlapComp->BodyInstance.GetUnrealWorldAngularVelocityInRadians();
+				// TODO - sort out scale for world objects in local sim
+				const FTransform CompSpaceTransform(CompSpaceOrientation, CompSpacePosition, FVector::OneVector);
 
-			const FVector CompSpaceVelocity = SpaceTransform.GetRotation().Inverse().RotateVector(Velocity);
-			const FVector CompSpaceAngularVelocity = SpaceTransform.GetRotation().Inverse().RotateVector(AngularVelocity);
+				// We need to set the velocity - either because the object or the space is
+				// moving. Ideally would do this by setting the velocity (with transformations), but
+				// it's not possible to set the velocity of kinematics. Use the kinematic target to
+				// force the velocities. This should handle initialization and teleportation too.
 
-			const FVector PrevCompSpacePosition = CompSpacePosition - CompSpaceVelocity * DeltaSeconds;
-			const FQuat CompSpaceDeltaOrientation = FQuat::MakeFromRotationVector(CompSpaceAngularVelocity * DeltaSeconds);
-			const FQuat PrevCompSpaceOrientation = CompSpaceDeltaOrientation.Inverse() * CompSpaceOrientation;
+				// We also have the problem that velocities are not reported from kinematics, unless
+				// OverlapComp->BodyInstance.bUpdateKinematicFromSimulation = true;
 
-			const FTransform PrevCompSpaceTransform(PrevCompSpaceOrientation, PrevCompSpacePosition, FVector::OneVector);
+				const FVector Velocity = OverlapComp->BodyInstance.GetUnrealWorldVelocity();
+				const FVector AngularVelocity = OverlapComp->BodyInstance.GetUnrealWorldAngularVelocityInRadians();
 
-			ActorHandle->InitWorldTransform(PrevCompSpaceTransform);
-			ActorHandle->SetKinematicTarget(CompSpaceTransform);
+				// To get the previous component space transform, we estimate the previous world space
+				// transform, and then put that into the previous sim space.
+
+				const FVector PrevWorldPosition = WorldPosition - Velocity * DeltaSeconds;
+				const FQuat DeltaWorldOrientation = FQuat::MakeFromRotationVector(AngularVelocity * DeltaSeconds);
+				const FQuat PrevWorldOrientation = DeltaWorldOrientation.Inverse() * WorldOrientation;
+
+				const FVector PrevCompSpacePosition = PrevSimSpaceTransform.InverseTransformPosition(WorldPosition);
+				const FQuat PrevCompSpaceOrientation = PrevSimSpaceTransform.InverseTransformRotation(WorldOrientation);
+				const FTransform PrevCompSpaceTransform(PrevCompSpaceOrientation, PrevCompSpacePosition, FVector::OneVector);
+
+				ActorHandle->InitWorldTransform(PrevCompSpaceTransform);
+				ActorHandle->SetKinematicTarget(CompSpaceTransform);
+			}
+		}
+	}
+	else
+	{
+		for (const auto& WorldEntry : ComponentsInSim)
+		{
+			const UPrimitiveComponent* OverlapComp = WorldEntry.Key;
+			if (OverlapComp != nullptr)
+			{
+				ImmediatePhysics::FActorHandle* ActorHandle = WorldEntry.Value.ActorHandle;
+
+				// Calculate the sim-space transform of this object
+				const FTransform CompWorldTransform = OverlapComp->BodyInstance.GetUnrealWorldTransform();
+				const FVector WorldPosition = CompWorldTransform.GetLocation();
+				const FQuat WorldOrientation = CompWorldTransform.GetRotation();
+
+				const FVector CompSpacePosition = SimulationSpaceData.Transform.InverseTransformPosition(WorldPosition);
+				const FQuat CompSpaceOrientation = SimulationSpaceData.Transform.InverseTransformRotation(WorldOrientation);
+
+				// TODO - sort out scale for world objects in local sim
+				const FTransform CompSpaceTransform(CompSpaceOrientation, CompSpacePosition, FVector::OneVector);
+
+				ActorHandle->InitWorldTransform(CompSpaceTransform);
+			}
 		}
 	}
 }
@@ -1818,7 +1839,7 @@ void FAnimNode_RigidBodyWithControl::InitializeBoneReferences(const FBoneContain
 	{
 		// If the user specified a simulation root that is not used by the skelmesh, issue a warning 
 		// (FAnimNode_RigidBodyWithControl::IsValidToEvaluate will return false and the simulation will not run)
-		UE_LOG(LogRigidBodyWithControl, Log, TEXT("FAnimNode_RigidBodyWithControl: RigidBodyWithControl Simulation Base Bone \'%s\' does not exist on SkeletalMesh %s."), *BaseBoneRef.BoneName.ToString(), *GetNameSafe(RequiredBones.GetSkeletalMeshAsset()));
+		UE_LOG(LogPhysicsControl, Log, TEXT("FAnimNode_RigidBodyWithControl: RigidBodyWithControl Simulation Base Bone \'%s\' does not exist on SkeletalMesh %s."), *BaseBoneRef.BoneName.ToString(), *GetNameSafe(RequiredBones.GetSkeletalMeshAsset()));
 	}
 
 	bool bHasInvalidBoneReference = false;
diff --git a/Engine/Plugins/Experimental/PhysicsControl/Source/PhysicsControl/Private/PhysicsControlBPLibrary.cpp b/Engine/Plugins/Experimental/PhysicsControl/Source/PhysicsControl/Private/PhysicsControlBPLibrary.cpp
--- a/Engine/Plugins/Experimental/PhysicsControl/Source/PhysicsControl/Private/PhysicsControlBPLibrary.cpp
+++ b/Engine/Plugins/Experimental/PhysicsControl/Source/PhysicsControl/Private/PhysicsControlBPLibrary.cpp
@@ -5,10 +5,6 @@
 #include "RigidBodyControlData.h"
 #include "AnimNode_RigidBodyWithControl.h"
 
-//UE_DISABLE_OPTIMIZATION
-
-#if 1
-
 //======================================================================================================================
 template<typename TParameterType> const TArray<FName>* FindNamesInSet(
 	const FPhysicsControlNameRecords& NameRecords, const FName SetName)
@@ -17,14 +13,14 @@ template<typename TParameterType> const TArray<FName>* FindNamesInSet(
 }
 
 //======================================================================================================================
-template<> const TArray<FName>* FindNamesInSet<FRigidBodyControl>(
+template<> const TArray<FName>* FindNamesInSet<FPhysicsControl>(
 	const FPhysicsControlNameRecords& NameRecords, const FName SetName)
 {
 	return NameRecords.ControlSets.Find(SetName);
 }
 
 //======================================================================================================================
-template<> const TArray<FName>* FindNamesInSet<FRigidBodyModifier>(
+template<> const TArray<FName>* FindNamesInSet<FPhysicsBodyModifier>(
 	const FPhysicsControlNameRecords& NameRecords, const FName SetName)
 {
 	return NameRecords.BodyModifierSets.Find(SetName);
@@ -214,7 +210,7 @@ TArray<FName> UPhysicsControlBPLibrary::GetControlNamesInSet(
 	const FRigidBodyWithControlReference& RigidBodyWithControl, 
 	const FName                           SetName)
 {
-	return GetNamesInSet<FRigidBodyControl>(RigidBodyWithControl, SetName);
+	return GetNamesInSet<FPhysicsControl>(RigidBodyWithControl, SetName);
 }
 
 //======================================================================================================================
@@ -223,7 +219,6 @@ TArray<FName> UPhysicsControlBPLibrary::GetBodyModifierNamesInSet(
 	const FRigidBodyWithControlReference& RigidBodyWithControl,
 	const FName                           SetName)
 {
-	return GetNamesInSet<FRigidBodyModifier>(RigidBodyWithControl, SetName);
+	return GetNamesInSet<FPhysicsBodyModifier>(RigidBodyWithControl, SetName);
 }
 
-#endif
diff --git a/Engine/Plugins/Experimental/PhysicsControl/Source/PhysicsControl/Private/PhysicsControlComponent.cpp b/Engine/Plugins/Experimental/PhysicsControl/Source/PhysicsControl/Private/PhysicsControlComponent.cpp
--- a/Engine/Plugins/Experimental/PhysicsControl/Source/PhysicsControl/Private/PhysicsControlComponent.cpp
+++ b/Engine/Plugins/Experimental/PhysicsControl/Source/PhysicsControl/Private/PhysicsControlComponent.cpp
@@ -4,6 +4,7 @@
 #include "PhysicsControlLog.h"
 #include "PhysicsControlRecord.h"
 #include "PhysicsControlComponentHelpers.h"
+#include "PhysicsControlOperatorNameGeneration.h"
 
 #include "PhysicsEngine/ConstraintInstance.h"
 #include "PhysicsEngine/BodyInstance.h"
@@ -70,19 +71,19 @@ void UPhysicsControlComponent::InitializeComponent()
 //======================================================================================================================
 void UPhysicsControlComponent::BeginDestroy()
 {
-	for (TPair<FName, FPhysicsControlRecord>& PhysicsControlRecordPair : PhysicsControlRecords)
+	for (TPair<FName, FPhysicsControlRecord>& PhysicsControlRecordPair : ControlRecords)
 	{
 		DestroyControl(
 			PhysicsControlRecordPair.Key, EDestroyBehavior::KeepRecord);
 	}
-	PhysicsControlRecords.Empty();
+	ControlRecords.Empty();
 
-	for (TPair<FName, FPhysicsBodyModifier>& PhysicsBodyModifierPair : PhysicsBodyModifiers)
+	for (TPair<FName, FPhysicsBodyModifierRecord>& PhysicsBodyModifierPair : BodyModifierRecords)
 	{
 		DestroyBodyModifier(
 			PhysicsBodyModifierPair.Key, EDestroyBehavior::KeepRecord);
 	}
-	PhysicsBodyModifiers.Empty();
+	BodyModifierRecords.Empty();
 
 	Super::BeginDestroy();
 }
@@ -103,10 +104,10 @@ void UPhysicsControlComponent::UpdateControls(float DeltaTime)
 	SCOPE_CYCLE_COUNTER(STAT_PhysicsControl_UpdateControls);
 	TRACE_CPUPROFILER_EVENT_SCOPE(UPhysicsControlComponent::UpdateControls);
 
-	PhysicsControlRecords.Compact();
-	PhysicsBodyModifiers.Compact();
+	ControlRecords.Compact();
+	BodyModifierRecords.Compact();
 
-	for (TPair<FName, FPhysicsControlRecord>& RecordPair : PhysicsControlRecords)
+	for (TPair<FName, FPhysicsControlRecord>& RecordPair : ControlRecords)
 	{
 		// New constraint requested when one doesn't exist
 		FName ControlName = RecordPair.Key;
@@ -119,9 +120,9 @@ void UPhysicsControlComponent::UpdateControls(float DeltaTime)
 	}
 
 	// Handle body modifiers
-	for (TPair<FName, FPhysicsBodyModifier>& BodyModifierPair : PhysicsBodyModifiers)
+	for (TPair<FName, FPhysicsBodyModifierRecord>& BodyModifierPair : BodyModifierRecords)
 	{
-		FPhysicsBodyModifier& BodyModifier = BodyModifierPair.Value;
+		FPhysicsBodyModifierRecord& BodyModifier = BodyModifierPair.Value;
 		ApplyBodyModifier(BodyModifier);
 	}
 }
@@ -157,7 +158,7 @@ TMap<FName, FPhysicsControlLimbBones> UPhysicsControlComponent::GetLimbBonesFrom
 	UPhysicsAsset* PhysicsAsset = SkeletalMeshComponent ? SkeletalMeshComponent->GetPhysicsAsset() : nullptr;
 	if (!PhysicsAsset)
 	{
-		UE_LOG(LogPhysicsControlComponent, Warning, TEXT("No physics asset in skeletal mesh"));
+		UE_LOG(LogPhysicsControl, Warning, TEXT("No physics asset in skeletal mesh"));
 		return Result;
 	}
 
@@ -176,7 +177,7 @@ TMap<FName, FPhysicsControlLimbBones> UPhysicsControlComponent::GetLimbBonesFrom
 		if (LimbSetupData.bIncludeParentBone)
 		{
 			LimbBones.bFirstBoneIsAdditional = true;
-			const FName ParentBoneName = UE::PhysicsControlComponent::GetPhysicalParentBone(
+			const FName ParentBoneName = UE::PhysicsControl::GetPhysicalParentBone(
 				SkeletalMeshComponent, LimbSetupData.StartBone);
 			if (!ParentBoneName.IsNone())
 			{
@@ -219,7 +220,7 @@ FName UPhysicsControlComponent::CreateControl(
 	const FName                   Set,
 	const FString                 NamePrefix)
 {
-	const FName Name = GetUniqueControlName(ParentBoneName, ChildBoneName, NamePrefix);
+	const FName Name = UE::PhysicsControl::GetUniqueControlName(ParentBoneName, ChildBoneName, ControlRecords, NamePrefix);
 	if (CreateNamedControl(
 		Name, ParentMeshComponent, ParentBoneName, ChildMeshComponent, ChildBoneName, ControlData, ControlTarget, Set))
 	{
@@ -242,14 +243,14 @@ bool UPhysicsControlComponent::CreateNamedControl(
 {
 	if (FindControlRecord(Name))
 	{
-		UE_LOG(LogPhysicsControlComponent, Warning,
+		UE_LOG(LogPhysicsControl, Warning,
 			TEXT("Unable to make a Control as one with the desired name already exists"), *Name.ToString());
 		return false;
 	}
 
 	if (!ChildMeshComponent)
 	{
-		UE_LOG(LogPhysicsControlComponent, Warning,
+		UE_LOG(LogPhysicsControl, Warning,
 			TEXT("Unable to make a Control as the child mesh component has not been set"));
 		return false;
 	}
@@ -263,9 +264,10 @@ bool UPhysicsControlComponent::CreateNamedControl(
 		AddSkeletalMeshReferenceForCaching(SkeletalMeshComponent);
 	}
 
-	FPhysicsControlRecord& NewRecord = PhysicsControlRecords.Add(
-		Name, FPhysicsControl(
-			ParentMeshComponent, ParentBoneName, ChildMeshComponent, ChildBoneName, ControlData, ControlTarget));
+	FPhysicsControlRecord& NewRecord = ControlRecords.Add(
+		Name, FPhysicsControlRecord(
+			FPhysicsControl(ParentBoneName, ChildBoneName, ControlData), 
+			ControlTarget, ParentMeshComponent, ChildMeshComponent));
 	NewRecord.ResetControlPoint();
 
 	NameRecords.AddControl(Name, Set);
@@ -287,7 +289,7 @@ TArray<FName> UPhysicsControlComponent::CreateControlsFromSkeletalMeshBelow(
 	UPhysicsAsset* PhysicsAsset = SkeletalMeshComponent ? SkeletalMeshComponent->GetPhysicsAsset() : nullptr;
 	if (!PhysicsAsset)
 	{
-		UE_LOG(LogPhysicsControlComponent, Warning, TEXT("No physics asset in skeletal mesh"));
+		UE_LOG(LogPhysicsControl, Warning, TEXT("No physics asset in skeletal mesh"));
 		return Result;
 	}
 
@@ -309,7 +311,7 @@ TArray<FName> UPhysicsControlComponent::CreateControlsFromSkeletalMeshBelow(
 				FName ParentBoneName;
 				if (ParentMeshComponent)
 				{
-					ParentBoneName = UE::PhysicsControlComponent::GetPhysicalParentBone(
+					ParentBoneName = UE::PhysicsControl::GetPhysicalParentBone(
 						SkeletalMeshComponent, ChildBoneName);
 					if (ParentBoneName.IsNone())
 					{
@@ -327,7 +329,7 @@ TArray<FName> UPhysicsControlComponent::CreateControlsFromSkeletalMeshBelow(
 				}
 				else
 				{
-					UE_LOG(LogPhysicsControlComponent, Warning, 
+					UE_LOG(LogPhysicsControl, Warning, 
 						TEXT("Failed to make control for %s"), *ChildBoneName.ToString());
 				}
 			}
@@ -350,7 +352,7 @@ TArray<FName> UPhysicsControlComponent::CreateControlsFromSkeletalMeshAndConstra
 	UPhysicsAsset* PhysicsAsset = SkeletalMeshComponent ? SkeletalMeshComponent->GetPhysicsAsset() : nullptr;
 	if (!PhysicsAsset)
 	{
-		UE_LOG(LogPhysicsControlComponent, Warning, TEXT("No physics asset in skeletal mesh"));
+		UE_LOG(LogPhysicsControl, Warning, TEXT("No physics asset in skeletal mesh"));
 		return Result;
 	}
 
@@ -366,7 +368,7 @@ TArray<FName> UPhysicsControlComponent::CreateControlsFromSkeletalMeshAndConstra
 				const FName ChildBoneName = PhysicsAsset->SkeletalBodySetups[BI->InstanceBodyIndex]->BoneName;
 
 				FName ParentBoneName;
-				ParentBoneName = UE::PhysicsControlComponent::GetPhysicalParentBone(
+				ParentBoneName = UE::PhysicsControl::GetPhysicalParentBone(
 					SkeletalMeshComponent, ChildBoneName);
 				if (ParentBoneName.IsNone())
 				{
@@ -381,12 +383,12 @@ TArray<FName> UPhysicsControlComponent::CreateControlsFromSkeletalMeshAndConstra
 				if (!SkeletalMeshComponent->GetConstraintProfilePropertiesOrDefault(
 					ProfileProperties, ChildBoneName, ConstraintProfile))
 				{
-					UE_LOG(LogPhysicsControlComponent, Warning, 
+					UE_LOG(LogPhysicsControl, Warning, 
 						TEXT("Failed get constraint profile for %s"), *ChildBoneName.ToString());
 					return;
 				}
 
-				UE::PhysicsControlComponent::ConvertConstraintProfileToControlData(ControlData, ProfileProperties);
+				UE::PhysicsControl::ConvertConstraintProfileToControlData(ControlData, ProfileProperties);
 				ControlData.bEnabled = bEnabled;
 
 				const FName ControlName = CreateControl(
@@ -401,7 +403,7 @@ TArray<FName> UPhysicsControlComponent::CreateControlsFromSkeletalMeshAndConstra
 				}
 				else
 				{
-					UE_LOG(LogPhysicsControlComponent, Warning,
+					UE_LOG(LogPhysicsControl, Warning,
 						TEXT("Failed to make control for %s"), *ChildBoneName.ToString());
 				}
 			}
@@ -423,7 +425,7 @@ TArray<FName> UPhysicsControlComponent::CreateControlsFromSkeletalMesh(
 	UPhysicsAsset* PhysicsAsset = SkeletalMeshComponent ? SkeletalMeshComponent->GetPhysicsAsset() : nullptr;
 	if (!PhysicsAsset)
 	{
-		UE_LOG(LogPhysicsControlComponent, Warning, TEXT("No physics asset in skeletal mesh"));
+		UE_LOG(LogPhysicsControl, Warning, TEXT("No physics asset in skeletal mesh"));
 		return Result;
 	}
 
@@ -435,7 +437,7 @@ TArray<FName> UPhysicsControlComponent::CreateControlsFromSkeletalMesh(
 		FName ParentBoneName;
 		if (ParentMeshComponent)
 		{
-			ParentBoneName = UE::PhysicsControlComponent::GetPhysicalParentBone(
+			ParentBoneName = UE::PhysicsControl::GetPhysicalParentBone(
 				SkeletalMeshComponent, ChildBoneName);
 			if (ParentBoneName.IsNone())
 			{
@@ -453,7 +455,7 @@ TArray<FName> UPhysicsControlComponent::CreateControlsFromSkeletalMesh(
 		}
 		else
 		{
-			UE_LOG(LogPhysicsControlComponent, Warning, TEXT("Failed to make control for %s"), *ChildBoneName.ToString());
+			UE_LOG(LogPhysicsControl, Warning, TEXT("Failed to make control for %s"), *ChildBoneName.ToString());
 		}
 	}
 
@@ -473,14 +475,14 @@ TArray<FName> UPhysicsControlComponent::CreateControlsFromSkeletalMeshAndConstra
 	UPhysicsAsset* PhysicsAsset = SkeletalMeshComponent ? SkeletalMeshComponent->GetPhysicsAsset() : nullptr;
 	if (!PhysicsAsset)
 	{
-		UE_LOG(LogPhysicsControlComponent, Warning, TEXT("No physics asset in skeletal mesh"));
+		UE_LOG(LogPhysicsControl, Warning, TEXT("No physics asset in skeletal mesh"));
 		return Result;
 	}
 
 	for (FName ChildBoneName : BoneNames)
 	{
 		const FName ParentBoneName = 
-			UE::PhysicsControlComponent::GetPhysicalParentBone(SkeletalMeshComponent, ChildBoneName);
+			UE::PhysicsControl::GetPhysicalParentBone(SkeletalMeshComponent, ChildBoneName);
 		if (ParentBoneName.IsNone())
 		{
 			continue;
@@ -494,12 +496,12 @@ TArray<FName> UPhysicsControlComponent::CreateControlsFromSkeletalMeshAndConstra
 		if (!SkeletalMeshComponent->GetConstraintProfilePropertiesOrDefault(
 			ProfileProperties, ChildBoneName, ConstraintProfile))
 		{
-			UE_LOG(LogPhysicsControlComponent, Warning, 
+			UE_LOG(LogPhysicsControl, Warning, 
 				TEXT("Failed get constraint profile for %s"), *ChildBoneName.ToString());
 			continue;
 		}
 
-		UE::PhysicsControlComponent::ConvertConstraintProfileToControlData(ControlData, ProfileProperties);
+		UE::PhysicsControl::ConvertConstraintProfileToControlData(ControlData, ProfileProperties);
 		ControlData.bEnabled = bEnabled;
 
 		const FName ControlName = CreateControl(
@@ -513,7 +515,7 @@ TArray<FName> UPhysicsControlComponent::CreateControlsFromSkeletalMeshAndConstra
 		}
 		else
 		{
-			UE_LOG(LogPhysicsControlComponent, Warning, TEXT("Failed to make control for %s"), *ChildBoneName.ToString());
+			UE_LOG(LogPhysicsControl, Warning, TEXT("Failed to make control for %s"), *ChildBoneName.ToString());
 		}
 	}
 
@@ -541,7 +543,7 @@ TMap<FName, FPhysicsControlNames> UPhysicsControlComponent::CreateControlsFromLi
 
 		if (!BonesInLimb.SkeletalMeshComponent.IsValid())
 		{
-			UE_LOG(LogPhysicsControlComponent, Warning, TEXT("No Skeletal mesh in limb %s"), *LimbName.ToString());
+			UE_LOG(LogPhysicsControl, Warning, TEXT("No Skeletal mesh in limb %s"), *LimbName.ToString());
 			continue;
 		}
 
@@ -576,7 +578,7 @@ TMap<FName, FPhysicsControlNames> UPhysicsControlComponent::CreateControlsFromLi
 			FName ParentBoneName;
 			if (ParentSkeletalMeshComponent)
 			{
-				ParentBoneName = UE::PhysicsControlComponent::GetPhysicalParentBone(
+				ParentBoneName = UE::PhysicsControl::GetPhysicalParentBone(
 					ParentSkeletalMeshComponent, ChildBoneName);
 				if (ParentBoneName.IsNone())
 				{
@@ -603,7 +605,7 @@ TMap<FName, FPhysicsControlNames> UPhysicsControlComponent::CreateControlsFromLi
 			}
 			else
 			{
-				UE_LOG(LogPhysicsControlComponent, Warning, 
+				UE_LOG(LogPhysicsControl, Warning, 
 					TEXT("Failed to make control for %s"), *ChildBoneName.ToString());
 			}
 		}
@@ -628,13 +630,13 @@ TMap<FName, FPhysicsControlNames> UPhysicsControlComponent::CreateControlsFromLi
 		USkeletalMeshComponent* SkeletalMeshComponent = BonesInLimb.SkeletalMeshComponent.Get();
 		if (!SkeletalMeshComponent)
 		{
-			UE_LOG(LogPhysicsControlComponent, Warning, TEXT("No Skeletal mesh in limb %s"), *LimbName.ToString());
+			UE_LOG(LogPhysicsControl, Warning, TEXT("No Skeletal mesh in limb %s"), *LimbName.ToString());
 			continue;
 		}
 		UPhysicsAsset* PhysicsAsset = SkeletalMeshComponent ? SkeletalMeshComponent->GetPhysicsAsset() : nullptr;
 		if (!PhysicsAsset)
 		{
-			UE_LOG(LogPhysicsControlComponent, Warning, TEXT("No physics asset in skeletal mesh"));
+			UE_LOG(LogPhysicsControl, Warning, TEXT("No physics asset in skeletal mesh"));
 			return Result;
 		}
 
@@ -654,7 +656,7 @@ TMap<FName, FPhysicsControlNames> UPhysicsControlComponent::CreateControlsFromLi
 
 			const FName ChildBoneName = BonesInLimb.BoneNames[BoneIndex];
 			const FName ParentBoneName = 
-				UE::PhysicsControlComponent::GetPhysicalParentBone(SkeletalMeshComponent, ChildBoneName);
+				UE::PhysicsControl::GetPhysicalParentBone(SkeletalMeshComponent, ChildBoneName);
 			if (ParentBoneName.IsNone())
 			{
 				continue;
@@ -670,12 +672,12 @@ TMap<FName, FPhysicsControlNames> UPhysicsControlComponent::CreateControlsFromLi
 				ProfileProperties, ChildBoneName, ConstraintProfile))
 			
 			{
-				UE_LOG(LogPhysicsControlComponent, Warning, 
+				UE_LOG(LogPhysicsControl, Warning, 
 					TEXT("Failed get constraint profile for %s"), *ChildBoneName.ToString());
 				continue;
 			}
 
-			UE::PhysicsControlComponent::ConvertConstraintProfileToControlData(ControlData, ProfileProperties);
+			UE::PhysicsControl::ConvertConstraintProfileToControlData(ControlData, ProfileProperties);
 			ControlData.bEnabled = bEnabled;
 
 			const FName ControlName = CreateControl(
@@ -690,7 +692,7 @@ TMap<FName, FPhysicsControlNames> UPhysicsControlComponent::CreateControlsFromLi
 			}
 			else
 			{
-				UE_LOG(LogPhysicsControlComponent, Warning, 
+				UE_LOG(LogPhysicsControl, Warning, 
 					TEXT("Failed to make control for %s"), *ChildBoneName.ToString());
 			}
 		}
@@ -732,7 +734,7 @@ bool UPhysicsControlComponent::SetControlEnabled(const FName Name, const bool bE
 	}
 	if (bWarnAboutInvalidNames)
 	{
-		UE_LOG(LogPhysicsControlComponent, Warning,
+		UE_LOG(LogPhysicsControl, Warning,
 			TEXT("SetControlEnabled - invalid name %s"), *Name.ToString());
 	}
 	return false;
@@ -763,7 +765,7 @@ bool UPhysicsControlComponent::SetControlParent(
 	if (Record)
 	{
 		if (USkeletalMeshComponent* SkeletalMeshComponent = 
-			Cast<USkeletalMeshComponent>(Record->PhysicsControl.ParentMeshComponent))
+			Cast<USkeletalMeshComponent>(Record->ParentMeshComponent.Get()))
 		{
 			RemoveSkeletalMeshReferenceForCaching(SkeletalMeshComponent);
 		}
@@ -773,13 +775,13 @@ bool UPhysicsControlComponent::SetControlParent(
 			AddSkeletalMeshReferenceForCaching(SkeletalMeshComponent);
 		}
 
-		Record->PhysicsControl.ParentMeshComponent = ParentMeshComponent;
+		Record->ParentMeshComponent = ParentMeshComponent;
 		Record->PhysicsControl.ParentBoneName = ParentBoneName;
 		return Record->InitConstraint(this, Name);
 	}
 	if (bWarnAboutInvalidNames)
 	{
-		UE_LOG(LogPhysicsControlComponent, Warning,
+		UE_LOG(LogPhysicsControl, Warning,
 			TEXT("SetControlParent - invalid name %s"), *Name.ToString());
 	}
 	return false;
@@ -820,7 +822,7 @@ bool UPhysicsControlComponent::SetControlData(
 	}
 	if (bWarnAboutInvalidNames)
 	{
-		UE_LOG(LogPhysicsControlComponent, Warning,
+		UE_LOG(LogPhysicsControl, Warning,
 			TEXT("SetControlData - invalid name %s"), *Name.ToString());
 	}
 	return false;
@@ -859,7 +861,7 @@ bool UPhysicsControlComponent::SetControlSparseData(
 	}
 	if (bWarnAboutInvalidNames)
 	{
-		UE_LOG(LogPhysicsControlComponent, Warning,
+		UE_LOG(LogPhysicsControl, Warning,
 			TEXT("SetControlData - invalid name %s"), *Name.ToString());
 	}
 	return false;
@@ -903,7 +905,7 @@ bool UPhysicsControlComponent::SetControlMultiplier(
 	}
 	if (bWarnAboutInvalidNames)
 	{
-		UE_LOG(LogPhysicsControlComponent, Warning,
+		UE_LOG(LogPhysicsControl, Warning,
 			TEXT("SetControlMultiplier - invalid name %s"), *Name.ToString());
 	}
 	return false;
@@ -949,7 +951,7 @@ bool UPhysicsControlComponent::SetControlSparseMultiplier(
 	}
 	if (bWarnAboutInvalidNames)
 	{
-		UE_LOG(LogPhysicsControlComponent, Warning,
+		UE_LOG(LogPhysicsControl, Warning,
 			TEXT("SetControlSparseMultiplier - invalid name %s"), *Name.ToString());
 	}
 	return false;
@@ -1000,7 +1002,7 @@ bool UPhysicsControlComponent::SetControlLinearData(
 	}
 	if (bWarnAboutInvalidNames)
 	{
-		UE_LOG(LogPhysicsControlComponent, Warning,
+		UE_LOG(LogPhysicsControl, Warning,
 			TEXT("SetControlLinearData - invalid name %s"), *Name.ToString());
 	}
 	return false;
@@ -1030,7 +1032,7 @@ bool UPhysicsControlComponent::SetControlAngularData(
 	}
 	if (bWarnAboutInvalidNames)
 	{
-		UE_LOG(LogPhysicsControlComponent, Warning,
+		UE_LOG(LogPhysicsControl, Warning,
 			TEXT("SetControlAngularData - invalid name %s"), *Name.ToString());
 	}
 	return false;
@@ -1049,7 +1051,7 @@ bool UPhysicsControlComponent::SetControlPoint(const FName Name, const FVector P
 	}
 	if (bWarnAboutInvalidNames)
 	{
-		UE_LOG(LogPhysicsControlComponent, Warning,
+		UE_LOG(LogPhysicsControl, Warning,
 			TEXT("SetControlPoint - invalid name %s"), *Name.ToString());
 	}
 	return false;
@@ -1066,7 +1068,7 @@ bool UPhysicsControlComponent::ResetControlPoint(const FName Name)
 	}
 	if (bWarnAboutInvalidNames)
 	{
-		UE_LOG(LogPhysicsControlComponent, Warning,
+		UE_LOG(LogPhysicsControl, Warning,
 			TEXT("ResetControlPoint - invalid name %s"), *Name.ToString());
 	}
 	return false;
@@ -1081,7 +1083,7 @@ bool UPhysicsControlComponent::SetControlTarget(
 	FPhysicsControlRecord* Record = FindControlRecord(Name);
 	if (Record)
 	{
-		Record->PhysicsControl.ControlTarget = ControlTarget;
+		Record->ControlTarget = ControlTarget;
 		if (bEnableControl)
 		{
 			Record->PhysicsControl.ControlData.bEnabled = true;
@@ -1090,7 +1092,7 @@ bool UPhysicsControlComponent::SetControlTarget(
 	}
 	if (bWarnAboutInvalidNames)
 	{
-		UE_LOG(LogPhysicsControlComponent, Warning,
+		UE_LOG(LogPhysicsControl, Warning,
 			TEXT("SetControlTarget - invalid name %s"), *Name.ToString());
 	}
 	return false;
@@ -1137,7 +1139,7 @@ bool UPhysicsControlComponent::SetControlTargetPositionAndOrientation(
 	}
 	if (bWarnAboutInvalidNames)
 	{
-		UE_LOG(LogPhysicsControlComponent, Warning,
+		UE_LOG(LogPhysicsControl, Warning,
 			TEXT("SetControlTargetPositionAndOrientation - invalid name %s"), *Name.ToString());
 	}
 	return false;
@@ -1185,15 +1187,15 @@ bool UPhysicsControlComponent::SetControlTargetPosition(
 	{
 		if (VelocityDeltaTime != 0)
 		{
-			Record->PhysicsControl.ControlTarget.TargetVelocity =
-				(Position - Record->PhysicsControl.ControlTarget.TargetPosition) / VelocityDeltaTime;
+			Record->ControlTarget.TargetVelocity =
+				(Position - Record->ControlTarget.TargetPosition) / VelocityDeltaTime;
 		}
 		else
 		{
-			Record->PhysicsControl.ControlTarget.TargetVelocity = FVector::ZeroVector;
+			Record->ControlTarget.TargetVelocity = FVector::ZeroVector;
 		}
-		Record->PhysicsControl.ControlTarget.TargetPosition = Position;
-		Record->PhysicsControl.ControlTarget.bApplyControlPointToTarget = bApplyControlPointToTarget;
+		Record->ControlTarget.TargetPosition = Position;
+		Record->ControlTarget.bApplyControlPointToTarget = bApplyControlPointToTarget;
 		if (bEnableControl)
 		{
 			Record->PhysicsControl.ControlData.bEnabled = true;
@@ -1202,7 +1204,7 @@ bool UPhysicsControlComponent::SetControlTargetPosition(
 	}
 	if (bWarnAboutInvalidNames)
 	{
-		UE_LOG(LogPhysicsControlComponent, Warning,
+		UE_LOG(LogPhysicsControl, Warning,
 			TEXT("SetControlTargetPosition - invalid name %s"), *Name.ToString());
 	}
 	return false;
@@ -1248,20 +1250,20 @@ bool UPhysicsControlComponent::SetControlTargetOrientation(
 	{
 		if (AngularVelocityDeltaTime != 0)
 		{
-			FQuat OldQ = Record->PhysicsControl.ControlTarget.TargetOrientation.Quaternion();
-			FQuat OrientationQ = Orientation.Quaternion();
+			FQuat OldQ = Record->ControlTarget.TargetOrientation.Quaternion();
+			const FQuat OrientationQ = Orientation.Quaternion();
 			OldQ.EnforceShortestArcWith(OrientationQ);
 			// Note that quats multiply in the opposite order to TMs
-			FQuat DeltaQ = OrientationQ * OldQ.Inverse();
-			Record->PhysicsControl.ControlTarget.TargetAngularVelocity =
+			const FQuat DeltaQ = OrientationQ * OldQ.Inverse();
+			Record->ControlTarget.TargetAngularVelocity =
 				DeltaQ.ToRotationVector() / (UE_TWO_PI * AngularVelocityDeltaTime);
 		}
 		else
 		{
-			Record->PhysicsControl.ControlTarget.TargetAngularVelocity = FVector::ZeroVector;
+			Record->ControlTarget.TargetAngularVelocity = FVector::ZeroVector;
 		}
-		Record->PhysicsControl.ControlTarget.TargetOrientation = Orientation;
-		Record->PhysicsControl.ControlTarget.bApplyControlPointToTarget = bApplyControlPointToTarget;
+		Record->ControlTarget.TargetOrientation = Orientation;
+		Record->ControlTarget.bApplyControlPointToTarget = bApplyControlPointToTarget;
 		if (bEnableControl)
 		{
 			Record->PhysicsControl.ControlData.bEnabled = true;
@@ -1270,7 +1272,7 @@ bool UPhysicsControlComponent::SetControlTargetOrientation(
 	}
 	if (bWarnAboutInvalidNames)
 	{
-		UE_LOG(LogPhysicsControlComponent, Warning,
+		UE_LOG(LogPhysicsControl, Warning,
 			TEXT("SetControlTargetOrientation - invalid name %s"), *Name.ToString());
 	}
 	return false;
@@ -1315,7 +1317,7 @@ bool UPhysicsControlComponent::SetControlTargetPositionsFromArray(
 	int32 NumPositions = Positions.Num();
 	if (NumControlNames != NumPositions)
 	{
-		UE_LOG(LogPhysicsControlComponent, Warning,
+		UE_LOG(LogPhysicsControl, Warning,
 			TEXT("SetControlTargetPositionsFromArray - names and positions arrays sizes do not match"));
 		return false;
 	}
@@ -1339,7 +1341,7 @@ bool UPhysicsControlComponent::SetControlTargetOrientationsFromArray(
 	int32 NumOrientations = Orientations.Num();
 	if (NumControlNames != NumOrientations)
 	{
-		UE_LOG(LogPhysicsControlComponent, Warning,
+		UE_LOG(LogPhysicsControl, Warning,
 			TEXT("SetControlTargetOrientationsFromArray - names and orientations arrays sizes do not match"));
 		return false;
 	}
@@ -1365,7 +1367,7 @@ bool UPhysicsControlComponent::SetControlTargetPositionsAndOrientationsFromArray
 	int32 NumOrientations = Orientations.Num();
 	if (NumControlNames != NumPositions || NumControlNames != NumOrientations)
 	{
-		UE_LOG(LogPhysicsControlComponent, Warning,
+		UE_LOG(LogPhysicsControl, Warning,
 			TEXT("SetControlTargetPositionsAndOrientationsFromArray - names and positions/orientation arrays sizes do not match"));
 		return false;
 	}
@@ -1390,33 +1392,33 @@ bool UPhysicsControlComponent::SetControlTargetPoses(
 	FPhysicsControlRecord* Record = FindControlRecord(Name);
 	if (Record)
 	{
-		FTransform ParentTM(ParentOrientation, ParentPosition, FVector::One());
-		FTransform ChildTM(ChildOrientation, ChildPosition, FVector::One());
+		const FTransform ParentTM(ParentOrientation, ParentPosition, FVector::One());
+		const FTransform ChildTM(ChildOrientation, ChildPosition, FVector::One());
 
-		FTransform OffsetTM = ChildTM * ParentTM.Inverse();
-		FVector Position = OffsetTM.GetTranslation();
-		FQuat OrientationQ = OffsetTM.GetRotation();
+		const FTransform OffsetTM = ChildTM * ParentTM.Inverse();
+		const FVector Position = OffsetTM.GetTranslation();
+		const FQuat OrientationQ = OffsetTM.GetRotation();
 
 		if (VelocityDeltaTime != 0)
 		{
-			FQuat OldQ = Record->PhysicsControl.ControlTarget.TargetOrientation.Quaternion();
+			FQuat OldQ = Record->ControlTarget.TargetOrientation.Quaternion();
 			OldQ.EnforceShortestArcWith(OrientationQ);
 			// Note that quats multiply in the opposite order to TMs
 			FQuat DeltaQ = OrientationQ * OldQ.Inverse();
-			Record->PhysicsControl.ControlTarget.TargetAngularVelocity =
+			Record->ControlTarget.TargetAngularVelocity =
 				DeltaQ.ToRotationVector() / (UE_TWO_PI * VelocityDeltaTime);
 
-			Record->PhysicsControl.ControlTarget.TargetVelocity =
-				(Position - Record->PhysicsControl.ControlTarget.TargetPosition) / VelocityDeltaTime;
+			Record->ControlTarget.TargetVelocity =
+				(Position - Record->ControlTarget.TargetPosition) / VelocityDeltaTime;
 		}
 		else
 		{
-			Record->PhysicsControl.ControlTarget.TargetAngularVelocity = FVector::ZeroVector;
-			Record->PhysicsControl.ControlTarget.TargetVelocity = FVector::ZeroVector;
+			Record->ControlTarget.TargetAngularVelocity = FVector::ZeroVector;
+			Record->ControlTarget.TargetVelocity = FVector::ZeroVector;
 		}
-		Record->PhysicsControl.ControlTarget.TargetOrientation = OrientationQ.Rotator();
-		Record->PhysicsControl.ControlTarget.TargetPosition = Position;
-		Record->PhysicsControl.ControlTarget.bApplyControlPointToTarget = true;
+		Record->ControlTarget.TargetOrientation = OrientationQ.Rotator();
+		Record->ControlTarget.TargetPosition = Position;
+		Record->ControlTarget.bApplyControlPointToTarget = true;
 		if (bEnableControl)
 		{
 			Record->PhysicsControl.ControlData.bEnabled = true;
@@ -1425,7 +1427,7 @@ bool UPhysicsControlComponent::SetControlTargetPoses(
 	}
 	if (bWarnAboutInvalidNames)
 	{
-		UE_LOG(LogPhysicsControlComponent, Warning,
+		UE_LOG(LogPhysicsControl, Warning,
 			TEXT("SetControlTargetPoses - invalid name %s"), *Name.ToString());
 	}
 	return false;
@@ -1446,7 +1448,7 @@ bool UPhysicsControlComponent::SetControlUseSkeletalAnimation(
 	}
 	if (bWarnAboutInvalidNames)
 	{
-		UE_LOG(LogPhysicsControlComponent, Warning,
+		UE_LOG(LogPhysicsControl, Warning,
 			TEXT("SetControlUseSkeletalAnimation - invalid name %s"), *Name.ToString());
 	}
 	return false;
@@ -1485,7 +1487,7 @@ bool UPhysicsControlComponent::SetControlDisableCollision(const FName Name, cons
 	}
 	if (bWarnAboutInvalidNames)
 	{
-		UE_LOG(LogPhysicsControlComponent, Warning,
+		UE_LOG(LogPhysicsControl, Warning,
 			TEXT("SetControlDisableCollision - invalid name %s"), *Name.ToString());
 	}
 	return false;
@@ -1517,7 +1519,7 @@ bool UPhysicsControlComponent::GetControlData(const FName Name, FPhysicsControlD
 	}
 	if (bWarnAboutInvalidNames)
 	{
-		UE_LOG(LogPhysicsControlComponent, Warning,
+		UE_LOG(LogPhysicsControl, Warning,
 			TEXT("GetControlData - invalid name %s"), *Name.ToString());
 	}
 	return false;
@@ -1534,7 +1536,7 @@ bool UPhysicsControlComponent::GetControlMultiplier(const FName Name, FPhysicsCo
 	}
 	if (bWarnAboutInvalidNames)
 	{
-		UE_LOG(LogPhysicsControlComponent, Warning,
+		UE_LOG(LogPhysicsControl, Warning,
 			TEXT("GetControlMultiplier - invalid name %s"), *Name.ToString());
 	}
 	return false;
@@ -1543,15 +1545,15 @@ bool UPhysicsControlComponent::GetControlMultiplier(const FName Name, FPhysicsCo
 //======================================================================================================================
 bool UPhysicsControlComponent::GetControlTarget(const FName Name, FPhysicsControlTarget& ControlTarget) const
 {
-	const FPhysicsControl* PhysicsControl = FindControl(Name);
-	if (PhysicsControl)
+	const FPhysicsControlRecord* Record = FindControlRecord(Name);
+	if (Record)
 	{
-		ControlTarget = PhysicsControl->ControlTarget;
+		ControlTarget = Record->ControlTarget;
 		return true;
 	}
 	if (bWarnAboutInvalidNames)
 	{
-		UE_LOG(LogPhysicsControlComponent, Warning,
+		UE_LOG(LogPhysicsControl, Warning,
 			TEXT("GetControlTarget - invalid name %s"), *Name.ToString());
 	}
 	return false;
@@ -1563,11 +1565,11 @@ bool UPhysicsControlComponent::GetControlEnabled(const FName Name) const
 	const FPhysicsControlRecord* Record = FindControlRecord(Name);
 	if (Record)
 	{
-		return Record->Enabled();
+		return Record->PhysicsControl.IsEnabled();
 	}
 	if (bWarnAboutInvalidNames)
 	{
-		UE_LOG(LogPhysicsControlComponent, Warning,
+		UE_LOG(LogPhysicsControl, Warning,
 			TEXT("GetControlEnabled - invalid name %s"), *Name.ToString());
 	}
 	return false;
@@ -1580,7 +1582,9 @@ FName UPhysicsControlComponent::CreateBodyModifier(
 	const FName                       Set,
 	const FPhysicsControlModifierData BodyModifierData)
 {
-	const FName Name = GetUniqueBodyModifierName(BoneName);
+	TSet<FName> Keys;
+	BodyModifierRecords.GetKeys(Keys);
+	const FName Name = UE::PhysicsControl::GetUniqueBodyModifierName(BoneName, Keys, TEXT(""));
 	if (CreateNamedBodyModifier(Name, MeshComponent, BoneName, Set, BodyModifierData))
 	{
 		return Name;
@@ -1596,22 +1600,22 @@ bool UPhysicsControlComponent::CreateNamedBodyModifier(
 	const FName                       Set,
 	const FPhysicsControlModifierData BodyModifierData)
 {
-	if (FindBodyModifier(Name))
+	if (FindBodyModifierRecord(Name))
 	{
-		UE_LOG(LogPhysicsControlComponent, Warning,
+		UE_LOG(LogPhysicsControl, Warning,
 			TEXT("CreateNamedBodyModifier - modifier with name %s already exists"), *Name.ToString());
 		return false;
 	}
 
 	if (!MeshComponent)
 	{
-		UE_LOG(LogPhysicsControlComponent, Warning,
+		UE_LOG(LogPhysicsControl, Warning,
 			TEXT("Unable to make a PhysicsBodyModifier as the mesh component has not been set"));
 		return false;
 	}
 
-	FPhysicsBodyModifier& Modifier = PhysicsBodyModifiers.Add(
-		Name, FPhysicsBodyModifier(MeshComponent, BoneName, BodyModifierData));
+	FPhysicsBodyModifierRecord& Modifier = BodyModifierRecords.Add(
+		Name, FPhysicsBodyModifierRecord(MeshComponent, BoneName, BodyModifierData));
 
 	USkeletalMeshComponent* SkeletalMeshComponent = Cast<USkeletalMeshComponent>(MeshComponent);
 	if (SkeletalMeshComponent)
@@ -1637,7 +1641,7 @@ TArray<FName> UPhysicsControlComponent::CreateBodyModifiersFromSkeletalMeshBelow
 	UPhysicsAsset* PhysicsAsset = SkeletalMeshComponent ? SkeletalMeshComponent->GetPhysicsAsset() : nullptr;
 	if (!PhysicsAsset)
 	{
-		UE_LOG(LogPhysicsControlComponent, Warning,
+		UE_LOG(LogPhysicsControl, Warning,
 			TEXT("CreateBodyModifiersFromSkeletalMeshBelow - No physics asset available"));
 		return Result;
 	}
@@ -1674,7 +1678,7 @@ TMap<FName, FPhysicsControlNames> UPhysicsControlComponent::CreateBodyModifiersF
 
 		if (!BonesInLimb.SkeletalMeshComponent.Get())
 		{
-			UE_LOG(LogPhysicsControlComponent, Warning, TEXT("No Skeletal mesh in limb %s"), *LimbName.ToString());
+			UE_LOG(LogPhysicsControl, Warning, TEXT("No Skeletal mesh in limb %s"), *LimbName.ToString());
 			continue;
 		}
 
@@ -1695,7 +1699,7 @@ TMap<FName, FPhysicsControlNames> UPhysicsControlComponent::CreateBodyModifiersF
 			}
 			else
 			{
-				UE_LOG(LogPhysicsControlComponent, Warning, 
+				UE_LOG(LogPhysicsControl, Warning, 
 					TEXT("Failed to make body modifier for %s"), *BoneName.ToString());
 			}
 		}
@@ -1731,15 +1735,15 @@ bool UPhysicsControlComponent::SetBodyModifierData(
 	const FName                       Name,
 	const FPhysicsControlModifierData ModifierData)
 {
-	FPhysicsBodyModifier* Modifier = FindBodyModifier(Name);
-	if (Modifier)
+	FPhysicsBodyModifierRecord* Record = FindBodyModifierRecord(Name);
+	if (Record)
 	{
-		Modifier->BodyModifierData = ModifierData;
+		Record->BodyModifier.ModifierData = ModifierData;
 		return true;
 	}
 	if (bWarnAboutInvalidNames)
 	{
-		UE_LOG(LogPhysicsControlComponent, Warning,
+		UE_LOG(LogPhysicsControl, Warning,
 			TEXT("SetBodyModifierData - invalid name %s"), *Name.ToString());
 	}
 	return false;
@@ -1770,15 +1774,15 @@ bool UPhysicsControlComponent::SetBodyModifierSparseData(
 	const FName                             Name,
 	const FPhysicsControlModifierSparseData ModifierData)
 {
-	FPhysicsBodyModifier* Modifier = FindBodyModifier(Name);
-	if (Modifier)
+	FPhysicsBodyModifierRecord* Record = FindBodyModifierRecord(Name);
+	if (Record)
 	{
-		Modifier->BodyModifierData.UpdateFromSparseData(ModifierData);
+		Record->BodyModifier.ModifierData.UpdateFromSparseData(ModifierData);
 		return true;
 	}
 	if (bWarnAboutInvalidNames)
 	{
-		UE_LOG(LogPhysicsControlComponent, Warning,
+		UE_LOG(LogPhysicsControl, Warning,
 			TEXT("SetBodyModifierData - invalid name %s"), *Name.ToString());
 	}
 	return false;
@@ -1811,20 +1815,20 @@ bool UPhysicsControlComponent::SetBodyModifierKinematicTarget(
 	const FRotator KinematicTargetOrienation,
 	const bool     bMakeKinematic)
 {
-	FPhysicsBodyModifier* PhysicsBodyModifier = FindBodyModifier(Name);
-	if (PhysicsBodyModifier)
+	FPhysicsBodyModifierRecord* Record = FindBodyModifierRecord(Name);
+	if (Record)
 	{
-		PhysicsBodyModifier->KinematicTargetPosition = KinematicTargetPosition;
-		PhysicsBodyModifier->KinematicTargetOrientation = KinematicTargetOrienation.Quaternion();
+		Record->KinematicTargetPosition = KinematicTargetPosition;
+		Record->KinematicTargetOrientation = KinematicTargetOrienation.Quaternion();
 		if (bMakeKinematic)
 		{
-			PhysicsBodyModifier->BodyModifierData.MovementType = EPhysicsMovementType::Kinematic;
+			Record->BodyModifier.ModifierData.MovementType = EPhysicsMovementType::Kinematic;
 		}
 		return true;
 	}
 	if (bWarnAboutInvalidNames)
 	{
-		UE_LOG(LogPhysicsControlComponent, Warning,
+		UE_LOG(LogPhysicsControl, Warning,
 			TEXT("SetBodyModifierKinematicTarget - invalid name %s"), *Name.ToString());
 	}
 	return false;
@@ -1835,15 +1839,15 @@ bool UPhysicsControlComponent::SetBodyModifierMovementType(
 	const FName                Name,
 	const EPhysicsMovementType MovementType)
 {
-	FPhysicsBodyModifier* PhysicsBodyModifier = FindBodyModifier(Name);
-	if (PhysicsBodyModifier)
+	FPhysicsBodyModifierRecord* Record = FindBodyModifierRecord(Name);
+	if (Record)
 	{
-		PhysicsBodyModifier->BodyModifierData.MovementType = MovementType;
+		Record->BodyModifier.ModifierData.MovementType = MovementType;
 		return true;
 	}
 	if (bWarnAboutInvalidNames)
 	{
-		UE_LOG(LogPhysicsControlComponent, Warning,
+		UE_LOG(LogPhysicsControl, Warning,
 			TEXT("SetBodyModifierMovementType - invalid name %s"), *Name.ToString());
 	}
 	return false;
@@ -1873,15 +1877,15 @@ bool UPhysicsControlComponent::SetBodyModifierCollisionType(
 	const FName                   Name,
 	const ECollisionEnabled::Type CollisionType)
 {
-	FPhysicsBodyModifier* PhysicsBodyModifier = FindBodyModifier(Name);
-	if (PhysicsBodyModifier)
+	FPhysicsBodyModifierRecord* Record = FindBodyModifierRecord(Name);
+	if (Record)
 	{
-		PhysicsBodyModifier->BodyModifierData.CollisionType = CollisionType;
+		Record->BodyModifier.ModifierData.CollisionType = CollisionType;
 		return true;
 	}
 	if (bWarnAboutInvalidNames)
 	{
-		UE_LOG(LogPhysicsControlComponent, Warning,
+		UE_LOG(LogPhysicsControl, Warning,
 			TEXT("SetBodyModifierCollisionType - invalid name %s"), *Name.ToString());
 	}
 	return false;
@@ -1911,15 +1915,15 @@ bool UPhysicsControlComponent::SetBodyModifierGravityMultiplier(
 	const FName Name,
 	const float GravityMultiplier)
 {
-	FPhysicsBodyModifier* PhysicsBodyModifier = FindBodyModifier(Name);
-	if (PhysicsBodyModifier)
+	FPhysicsBodyModifierRecord* Record = FindBodyModifierRecord(Name);
+	if (Record)
 	{
-		PhysicsBodyModifier->BodyModifierData.GravityMultiplier = GravityMultiplier;
+		Record->BodyModifier.ModifierData.GravityMultiplier = GravityMultiplier;
 		return true;
 	}
 	if (bWarnAboutInvalidNames)
 	{
-		UE_LOG(LogPhysicsControlComponent, Warning,
+		UE_LOG(LogPhysicsControl, Warning,
 			TEXT("SetBodyModifierGravityMultiplier - invalid name %s"), *Name.ToString());
 	}
 	return false;
@@ -1949,15 +1953,15 @@ bool UPhysicsControlComponent::SetBodyModifierPhysicsBlendWeight(
 	const FName Name,
 	const float PhysicsBlendWeight)
 {
-	FPhysicsBodyModifier* PhysicsBodyModifier = FindBodyModifier(Name);
-	if (PhysicsBodyModifier)
+	FPhysicsBodyModifierRecord* Record = FindBodyModifierRecord(Name);
+	if (Record)
 	{
-		PhysicsBodyModifier->BodyModifierData.PhysicsBlendWeight = PhysicsBlendWeight;
+		Record->BodyModifier.ModifierData.PhysicsBlendWeight = PhysicsBlendWeight;
 		return true;
 	}
 	if (bWarnAboutInvalidNames)
 	{
-		UE_LOG(LogPhysicsControlComponent, Warning,
+		UE_LOG(LogPhysicsControl, Warning,
 			TEXT("SetBodyModifierPhysicsBlendWeight - invalid name %s"), *Name.ToString());
 	}
 	return false;
@@ -1987,15 +1991,15 @@ bool UPhysicsControlComponent::SetBodyModifierUseSkeletalAnimation(
 	const FName Name,
 	const bool  bUseSkeletalAnimation)
 {
-	FPhysicsBodyModifier* PhysicsBodyModifier = FindBodyModifier(Name);
-	if (PhysicsBodyModifier)
+	FPhysicsBodyModifierRecord* Record = FindBodyModifierRecord(Name);
+	if (Record)
 	{
-		PhysicsBodyModifier->BodyModifierData.bUseSkeletalAnimation = bUseSkeletalAnimation;
+		Record->BodyModifier.ModifierData.bUseSkeletalAnimation = bUseSkeletalAnimation;
 		return true;
 	}
 	if (bWarnAboutInvalidNames)
 	{
-		UE_LOG(LogPhysicsControlComponent, Warning,
+		UE_LOG(LogPhysicsControl, Warning,
 			TEXT("SetBodyModifierUseSkeletalAnimation - invalid name %s"), *Name.ToString());
 	}
 	return false;
@@ -2025,15 +2029,15 @@ bool UPhysicsControlComponent::SetBodyModifierUpdateKinematicFromSimulation(
 	const FName Name,
 	const bool  bUpdateKinematicFromSimulation)
 {
-	FPhysicsBodyModifier* PhysicsBodyModifier = FindBodyModifier(Name);
-	if (PhysicsBodyModifier)
+	FPhysicsBodyModifierRecord* Record = FindBodyModifierRecord(Name);
+	if (Record)
 	{
-		PhysicsBodyModifier->BodyModifierData.bUpdateKinematicFromSimulation = bUpdateKinematicFromSimulation;
+		Record->BodyModifier.ModifierData.bUpdateKinematicFromSimulation = bUpdateKinematicFromSimulation;
 		return true;
 	}
 	if (bWarnAboutInvalidNames)
 	{
-		UE_LOG(LogPhysicsControlComponent, Warning,
+		UE_LOG(LogPhysicsControl, Warning,
 			TEXT("SetBodyModifierUpdateKinematicFromSimulation - invalid name %s"), *Name.ToString());
 	}
 	return false;
@@ -2094,6 +2098,180 @@ void UPhysicsControlComponent::CreateControlsAndBodyModifiersFromLimbBones(
 	LimbBodyModifiers = CreateBodyModifiersFromLimbBones(AllBodyModifiers, LimbBones, BodyModifierData);
 }
 
+//======================================================================================================================
+void UPhysicsControlComponent::CreateControlsAndBodyModifiersFromControlProfileAsset(
+	USkeletalMeshComponent* SkeletalMeshComponent,
+	UMeshComponent*         WorldComponent,
+	FName                   WorldBoneName)
+{
+	if (!PhysicsControlProfileAsset.IsValid())
+	{
+		UE_LOG(LogPhysicsControl, Warning,
+			TEXT("CreateControlsAndBodyModifiersFromControlProfile - unable to get/load the control profile asset"));
+		return;
+	}
+
+	FPhysicsControlNames AllWorldSpaceControls;
+	TMap<FName, FPhysicsControlNames> LimbWorldSpaceControls;
+	FPhysicsControlNames AllParentSpaceControls;
+	TMap<FName, FPhysicsControlNames> LimbParentSpaceControls;
+	FPhysicsControlNames AllBodyModifiers;
+	TMap<FName, FPhysicsControlNames> LimbBodyModifiers;
+
+	CreateControlsAndBodyModifiersFromLimbBones(
+		AllWorldSpaceControls, LimbWorldSpaceControls, AllParentSpaceControls, LimbParentSpaceControls, 
+		AllBodyModifiers, LimbBodyModifiers,
+		SkeletalMeshComponent,
+		PhysicsControlProfileAsset->CharacterSetupData.LimbSetupData,
+		PhysicsControlProfileAsset->CharacterSetupData.DefaultWorldSpaceControlData,
+		PhysicsControlProfileAsset->CharacterSetupData.DefaultParentSpaceControlData,
+		PhysicsControlProfileAsset->CharacterSetupData.DefaultBodyModifierData,
+		WorldComponent,
+		WorldBoneName);
+
+	// Create additional controls
+	for (const TPair<FName, FPhysicsControlCreationData>& ControlPair : 
+		PhysicsControlProfileAsset->AdditionalControlsAndModifiers.Controls)
+	{
+		FName ControlName = ControlPair.Key;
+		const FPhysicsControlCreationData& ControlCreationData = ControlPair.Value;
+		if (CreateNamedControl(ControlName,
+			!ControlCreationData.Control.ParentBoneName.IsNone() 
+			? SkeletalMeshComponent : nullptr, ControlCreationData.Control.ParentBoneName,
+			SkeletalMeshComponent, ControlCreationData.Control.ChildBoneName,
+			ControlCreationData.Control.ControlData, FPhysicsControlTarget(), FName()))
+		{
+			for (FName SetName : ControlCreationData.Sets)
+			{
+				NameRecords.AddControl(ControlName, SetName);
+			}
+		}
+	}
+
+	// Create additional modifiers
+	for (const TPair<FName, FPhysicsBodyModifierCreationData>& ModifierPair : 
+		PhysicsControlProfileAsset->AdditionalControlsAndModifiers.Modifiers)
+	{
+		FName ModifierName = ModifierPair.Key;
+		const FPhysicsBodyModifierCreationData& ModifierCreationData = ModifierPair.Value;
+		if (CreateNamedBodyModifier(ModifierName,
+			SkeletalMeshComponent, ModifierCreationData.Modifier.BoneName,
+			FName(), ModifierCreationData.Modifier.ModifierData))
+		{
+			for (FName SetName : ModifierCreationData.Sets)
+			{
+				NameRecords.AddBodyModifier(ModifierName, SetName);
+			}
+		}
+	}
+
+	// Create any additional sets that have been requested
+	UE::PhysicsControl::CreateAdditionalSets(
+		PhysicsControlProfileAsset->AdditionalSets, BodyModifierRecords, ControlRecords, NameRecords);
+
+	for (FPhysicsControlControlAndModifierUpdates& Updates : PhysicsControlProfileAsset->InitialControlAndModifierUpdates)
+	{
+		ApplyControlAndModifierUpdates(Updates);
+	}
+}
+
+//======================================================================================================================
+void UPhysicsControlComponent::InvokeControlProfile(FName ProfileName)
+{
+	if (!PhysicsControlProfileAsset.IsValid())
+	{
+		if (bWarnAboutInvalidNames)
+		{
+			UE_LOG(LogPhysicsControl, Warning,
+				TEXT("InvokeControlProfile - control profile asset is invalid or missing"));
+		}
+		return;
+	}
+
+	const FPhysicsControlControlAndModifierUpdates* ControlAndModifierUpdates =
+		PhysicsControlProfileAsset->Profiles.Find(ProfileName);
+
+	if (!ControlAndModifierUpdates)
+	{
+		if (bWarnAboutInvalidNames)
+		{
+			UE_LOG(LogPhysicsControl, Warning,
+				TEXT("InvokeControlProfile - control profile %s not found"), *ProfileName.ToString());
+		}
+		return;
+	}
+
+	ApplyControlAndModifierUpdates(*ControlAndModifierUpdates);
+
+	return;
+}
+
+//======================================================================================================================
+void UPhysicsControlComponent::ApplyControlAndModifierUpdates(
+	const FPhysicsControlControlAndModifierUpdates& ControlAndModifierUpdates)
+{
+	for (const FPhysicsControlNamedControlParameters& ControlParameters : ControlAndModifierUpdates.ControlUpdates)
+	{
+		TArray<FName> Names = ExpandName(ControlParameters.Name, NameRecords.ControlSets);
+		for (FName Name : Names)
+		{
+			const FPhysicsControlSparseData& ControlData = ControlParameters.Data;
+			if (FPhysicsControlRecord* ControlRecord = ControlRecords.Find(Name))
+			{
+				ControlRecord->PhysicsControl.ControlData.UpdateFromSparseData(ControlData);
+			}
+			else
+			{
+				if (bWarnAboutInvalidNames)
+				{
+					UE_LOG(LogPhysicsControl, Warning,
+						TEXT("ApplyControlAndModifierUpdates: Failed to find control with name %s"), *Name.ToString());
+				}
+			}
+		}
+	}
+
+	for (const FPhysicsControlNamedControlMultiplierParameters& ControlMultiplierParameters :
+		ControlAndModifierUpdates.ControlMultiplierUpdates)
+	{
+		TArray<FName> Names = ExpandName(ControlMultiplierParameters.Name, NameRecords.ControlSets);
+		for (FName Name : Names)
+		{
+			const FPhysicsControlSparseMultiplier& Multiplier = ControlMultiplierParameters.Data;
+			if (FPhysicsControlRecord* ControlRecord = ControlRecords.Find(Name))
+			{
+				ControlRecord->PhysicsControl.ControlMultiplier.UpdateFromSparseData(Multiplier);
+			}
+			else
+			{
+				if (bWarnAboutInvalidNames)
+				{
+					UE_LOG(LogPhysicsControl, Warning,
+						TEXT("ApplyControlAndModifierUpdates: Failed to find control with name %s"), *Name.ToString());
+				}
+			}
+		}
+	}
+
+	for (const FPhysicsControlNamedModifierParameters& ModifierParameters : ControlAndModifierUpdates.ModifierUpdates)
+	{
+		TArray<FName> Names = ExpandName(ModifierParameters.Name, NameRecords.BodyModifierSets);
+		for (FName Name : Names)
+		{
+			const FPhysicsControlModifierSparseData& ModifierData = ModifierParameters.Data;
+			if (FPhysicsBodyModifierRecord* Record = BodyModifierRecords.Find(Name))
+			{
+				Record->BodyModifier.ModifierData.UpdateFromSparseData(ModifierData);
+			}
+			else
+			{
+				UE_LOG(LogPhysicsControl, Warning,
+					TEXT("InvokeControlProfile: Failed to find modifier with name %s"), *Name.ToString());
+			}
+		}
+	}
+}
+
 //======================================================================================================================
 void UPhysicsControlComponent::AddControlToSet(
 	FPhysicsControlNames& NewSet, 
@@ -2149,10 +2327,10 @@ void UPhysicsControlComponent::AddBodyModifierToSet(
 //======================================================================================================================
 void UPhysicsControlComponent::AddBodyModifiersToSet(
 	FPhysicsControlNames& NewSet, 
-	const TArray<FName>&  BodyModifiers, 
+	const TArray<FName>&  InBodyModifiers, 
 	const FName           SetName)
 {
-	for (FName BodyModifier : BodyModifiers)
+	for (FName BodyModifier : InBodyModifiers)
 	{
 		NameRecords.AddBodyModifier(BodyModifier, SetName);
 	}
@@ -2323,7 +2501,7 @@ FTransform UPhysicsControlComponent::GetCachedBoneTransform(
 	}
 	if (bWarnAboutInvalidNames)
 	{
-		UE_LOG(LogPhysicsControlComponent, Warning,
+		UE_LOG(LogPhysicsControl, Warning,
 			TEXT("GetCachedBoneTransform - invalid name %s"), *Name.ToString());
 	}
 	return FTransform();
@@ -2341,7 +2519,7 @@ FVector UPhysicsControlComponent::GetCachedBonePosition(
 	}
 	if (bWarnAboutInvalidNames)
 	{
-		UE_LOG(LogPhysicsControlComponent, Warning,
+		UE_LOG(LogPhysicsControl, Warning,
 			TEXT("GetCachedBonePosition - invalid name %s"), *Name.ToString());
 	}
 	return FVector::ZeroVector;
@@ -2359,7 +2537,7 @@ FRotator UPhysicsControlComponent::GetCachedBoneOrientation(
 	}
 	if (bWarnAboutInvalidNames)
 	{
-		UE_LOG(LogPhysicsControlComponent, Warning,
+		UE_LOG(LogPhysicsControl, Warning,
 			TEXT("GetCachedBoneOrientation - invalid name %s"), *Name.ToString());
 	}
 	return FRotator::ZeroRotator;
@@ -2377,7 +2555,7 @@ FVector UPhysicsControlComponent::GetCachedBoneVelocity(
 	}
 	if (bWarnAboutInvalidNames)
 	{
-		UE_LOG(LogPhysicsControlComponent, Warning,
+		UE_LOG(LogPhysicsControl, Warning,
 			TEXT("GetCachedBoneVelocity - invalid name %s"), *Name.ToString());
 	}
 	return FVector::Zero();
@@ -2395,7 +2573,7 @@ FVector UPhysicsControlComponent::GetCachedBoneAngularVelocity(
 	}
 	if (bWarnAboutInvalidNames)
 	{
-		UE_LOG(LogPhysicsControlComponent, Warning,
+		UE_LOG(LogPhysicsControl, Warning,
 			TEXT("GetCachedBoneAngularVelocity - invalid name %s"), *Name.ToString());
 	}
 	return FVector::Zero();
@@ -2420,7 +2598,7 @@ bool UPhysicsControlComponent::SetCachedBoneData(
 	}
 	if (bWarnAboutInvalidNames)
 	{
-		UE_LOG(LogPhysicsControlComponent, Warning,
+		UE_LOG(LogPhysicsControl, Warning,
 			TEXT("SetCachedBoneData - invalid name %s"), *Name.ToString());
 	}
 	return false;
@@ -2431,22 +2609,22 @@ bool UPhysicsControlComponent::ResetBodyModifierToCachedBoneTransform(
 	const FName                        Name,
 	const EResetToCachedTargetBehavior Behavior)
 {
-	FPhysicsBodyModifier* PhysicsBodyModifier = FindBodyModifier(Name);
-	if (PhysicsBodyModifier)
+	FPhysicsBodyModifierRecord* Record = FindBodyModifierRecord(Name);
+	if (Record)
 	{
 		if (Behavior == EResetToCachedTargetBehavior::ResetImmediately)
 		{
-			ResetToCachedTarget(*PhysicsBodyModifier);
+			ResetToCachedTarget(*Record);
 		}
 		else
 		{
-			PhysicsBodyModifier->bResetToCachedTarget = true;
+			Record->bResetToCachedTarget = true;
 		}
 		return true;
 	}
 	if (bWarnAboutInvalidNames)
 	{
-		UE_LOG(LogPhysicsControlComponent, Warning,
+		UE_LOG(LogPhysicsControl, Warning,
 			TEXT("ResetBodyModifierToCachedBoneTransform - invalid name %s"), *Name.ToString());
 	}
 	return false;
@@ -2480,7 +2658,7 @@ bool UPhysicsControlComponent::GetControlExists(const FName Name) const
 //======================================================================================================================
 bool UPhysicsControlComponent::GetBodyModifierExists(const FName Name) const
 {
-	return FindBodyModifier(Name) != nullptr;
+	return FindBodyModifierRecord(Name) != nullptr;
 }
 
 #if WITH_EDITOR
@@ -2505,7 +2683,7 @@ void UPhysicsControlComponent::DebugDraw(FPrimitiveDrawInterface* PDI) const
 	if (bShowDebugVisualization && VisualizationSizeScale > 0)
 	{
 		for (const TPair<FName, FPhysicsControlRecord>& PhysicsControlRecordPair : 
-			PhysicsControlRecords)
+			ControlRecords)
 		{
 			const FName Name = PhysicsControlRecordPair.Key;
 			const FPhysicsControlRecord& Record = PhysicsControlRecordPair.Value;
@@ -2517,7 +2695,7 @@ void UPhysicsControlComponent::DebugDraw(FPrimitiveDrawInterface* PDI) const
 	if (!DebugControlDetailFilter.IsEmpty())
 	{
 		for (const TPair<FName, FPhysicsControlRecord>& PhysicsControlRecordPair :
-			PhysicsControlRecords)
+			ControlRecords)
 		{
 			const FName Name = PhysicsControlRecordPair.Key;
 
@@ -2525,10 +2703,10 @@ void UPhysicsControlComponent::DebugDraw(FPrimitiveDrawInterface* PDI) const
 			{
 				const FPhysicsControlRecord& Record = PhysicsControlRecordPair.Value;
 
-				FString ParentComponentName = Record.PhysicsControl.ParentMeshComponent.IsValid() ?
-					Record.PhysicsControl.ParentMeshComponent->GetName() : TEXT("NoParent");
-				FString ChildComponentName = Record.PhysicsControl.ChildMeshComponent.IsValid() ?
-					Record.PhysicsControl.ChildMeshComponent->GetName() : TEXT("NoChild");
+				FString ParentComponentName = Record.ParentMeshComponent.IsValid() ?
+					Record.ParentMeshComponent->GetName() : TEXT("NoParent");
+				FString ChildComponentName = Record.ChildMeshComponent.IsValid() ?
+					Record.ChildMeshComponent->GetName() : TEXT("NoChild");
 
 				FString Text = FString::Printf(
 					TEXT("%s: Parent %s (%s) Child %s (%s): Linear strength %f Angular strength %f"),
@@ -2542,7 +2720,7 @@ void UPhysicsControlComponent::DebugDraw(FPrimitiveDrawInterface* PDI) const
 
 				GEngine->AddOnScreenDebugMessage(
 					-1, 0.0f,
-					Record.Enabled() ? FColor::Green : FColor::Red, Text);
+					Record.PhysicsControl.IsEnabled() ? FColor::Green : FColor::Red, Text);
 			}
 		}
 	}
@@ -2553,7 +2731,7 @@ void UPhysicsControlComponent::DebugDraw(FPrimitiveDrawInterface* PDI) const
 		FString AllNames;
 
 		for (const TPair<FName, FPhysicsControlRecord>& PhysicsControlRecordPair :
-			PhysicsControlRecords)
+			ControlRecords)
 		{
 			const FName Name = PhysicsControlRecordPair.Key;
 			AllNames += Name.ToString() + TEXT(" ");
@@ -2565,21 +2743,21 @@ void UPhysicsControlComponent::DebugDraw(FPrimitiveDrawInterface* PDI) const
 		}
 		GEngine->AddOnScreenDebugMessage(
 			-1, 0.0f, FColor::White,
-			FString::Printf(TEXT("%d Controls: %s"), PhysicsControlRecords.Num(), *AllNames));
+			FString::Printf(TEXT("%d Controls: %s"), ControlRecords.Num(), *AllNames));
 
 	}
 
 	// Detailed body modifiers - if there's a filter
 	if (!DebugBodyModifierDetailFilter.IsEmpty())
 	{
-		for (const TPair<FName, FPhysicsBodyModifier>& PhysicsBodyModifierPair :
-			PhysicsBodyModifiers)
+		for (const TPair<FName, FPhysicsBodyModifierRecord>& PhysicsBodyModifierPair :
+			BodyModifierRecords)
 		{
 			const FName Name = PhysicsBodyModifierPair.Key;
 
 			if (Name.ToString().Contains(DebugBodyModifierDetailFilter))
 			{
-				const FPhysicsBodyModifier& Record = PhysicsBodyModifierPair.Value;
+				const FPhysicsBodyModifierRecord& Record = PhysicsBodyModifierPair.Value;
 
 				FString ComponentName = Record.MeshComponent.IsValid() ? Record.MeshComponent->GetName() : TEXT("None");
 
@@ -2587,13 +2765,13 @@ void UPhysicsControlComponent::DebugDraw(FPrimitiveDrawInterface* PDI) const
 					TEXT("%s: %s: %s %s GravityMultiplier %f BlendWeight %f"),
 					*Name.ToString(),
 					*ComponentName,
-					*UEnum::GetValueAsString(Record.BodyModifierData.MovementType),
-					*UEnum::GetValueAsString(Record.BodyModifierData.CollisionType),
-					Record.BodyModifierData.GravityMultiplier,
-					Record.BodyModifierData.PhysicsBlendWeight);
+					*UEnum::GetValueAsString(Record.BodyModifier.ModifierData.MovementType),
+					*UEnum::GetValueAsString(Record.BodyModifier.ModifierData.CollisionType),
+					Record.BodyModifier.ModifierData.GravityMultiplier,
+					Record.BodyModifier.ModifierData.PhysicsBlendWeight);
 
 				GEngine->AddOnScreenDebugMessage(-1, 0.0f, 
-					Record.BodyModifierData.MovementType == EPhysicsMovementType::Simulated ? 
+					Record.BodyModifier.ModifierData.MovementType == EPhysicsMovementType::Simulated ?
 					FColor::Green : FColor::Red, Text);
 			}
 		}
@@ -2604,7 +2782,7 @@ void UPhysicsControlComponent::DebugDraw(FPrimitiveDrawInterface* PDI) const
 	{
 		FString AllNames;
 
-		for (const TPair<FName, FPhysicsBodyModifier>& PhysicsBodyModifierPair : PhysicsBodyModifiers)
+		for (const TPair<FName, FPhysicsBodyModifierRecord>& PhysicsBodyModifierPair : BodyModifierRecords)
 		{
 			const FName Name = PhysicsBodyModifierPair.Key;
 			AllNames += Name.ToString() + TEXT(" ");
@@ -2616,7 +2794,7 @@ void UPhysicsControlComponent::DebugDraw(FPrimitiveDrawInterface* PDI) const
 		}
 		GEngine->AddOnScreenDebugMessage(
 			-1, 0.0f, FColor::White,
-			FString::Printf(TEXT("%d Body modifiers: %s"), PhysicsBodyModifiers.Num(), *AllNames));
+			FString::Printf(TEXT("%d Body modifiers: %s"), BodyModifierRecords.Num(), *AllNames));
 
 	}
 }
@@ -2635,18 +2813,18 @@ void UPhysicsControlComponent::DebugDrawControl(
 	const bool bHaveLinear = Record.PhysicsControl.ControlData.LinearStrength > 0;
 	const bool bHaveAngular = Record.PhysicsControl.ControlData.AngularStrength > 0;
 
-	if (Record.Enabled() && ConstraintInstance)
+	if (Record.PhysicsControl.IsEnabled() && ConstraintInstance)
 	{
-		FBodyInstance* ChildBodyInstance = UE::PhysicsControlComponent::GetBodyInstance(
-			Record.PhysicsControl.ChildMeshComponent.Get(), Record.PhysicsControl.ChildBoneName);
+		FBodyInstance* ChildBodyInstance = UE::PhysicsControl::GetBodyInstance(
+			Record.ChildMeshComponent.Get(), Record.PhysicsControl.ChildBoneName);
 		if (!ChildBodyInstance)
 		{
 			return;
 		}
 		FTransform ChildBodyTM = ChildBodyInstance->GetUnrealWorldTransform();
 
-		FBodyInstance* ParentBodyInstance = UE::PhysicsControlComponent::GetBodyInstance(
-			Record.PhysicsControl.ParentMeshComponent.Get(), Record.PhysicsControl.ParentBoneName);
+		FBodyInstance* ParentBodyInstance = UE::PhysicsControl::GetBodyInstance(
+			Record.ParentMeshComponent.Get(), Record.PhysicsControl.ParentBoneName);
 		const FTransform ParentBodyTM = ParentBodyInstance ? ParentBodyInstance->GetUnrealWorldTransform() : FTransform();
 
 		FTransform TargetTM;
diff --git a/Engine/Plugins/Experimental/PhysicsControl/Source/PhysicsControl/Private/PhysicsControlComponentHelpers.cpp b/Engine/Plugins/Experimental/PhysicsControl/Source/PhysicsControl/Private/PhysicsControlComponentHelpers.cpp
--- a/Engine/Plugins/Experimental/PhysicsControl/Source/PhysicsControl/Private/PhysicsControlComponentHelpers.cpp
+++ b/Engine/Plugins/Experimental/PhysicsControl/Source/PhysicsControl/Private/PhysicsControlComponentHelpers.cpp
@@ -9,7 +9,7 @@
 
 namespace UE
 {
-namespace PhysicsControlComponent
+namespace PhysicsControl
 {
 
 //======================================================================================================================
diff --git a/Engine/Plugins/Experimental/PhysicsControl/Source/PhysicsControl/Private/PhysicsControlComponentHelpers.h b/Engine/Plugins/Experimental/PhysicsControl/Source/PhysicsControl/Private/PhysicsControlComponentHelpers.h
--- a/Engine/Plugins/Experimental/PhysicsControl/Source/PhysicsControl/Private/PhysicsControlComponentHelpers.h
+++ b/Engine/Plugins/Experimental/PhysicsControl/Source/PhysicsControl/Private/PhysicsControlComponentHelpers.h
@@ -14,7 +14,7 @@ class USkeletalMeshComponent;
 
 namespace UE
 {
-namespace PhysicsControlComponent
+namespace PhysicsControl
 {
 
 /** 
diff --git a/Engine/Plugins/Experimental/PhysicsControl/Source/PhysicsControl/Private/PhysicsControlComponentImpl.cpp b/Engine/Plugins/Experimental/PhysicsControl/Source/PhysicsControl/Private/PhysicsControlComponentImpl.cpp
--- a/Engine/Plugins/Experimental/PhysicsControl/Source/PhysicsControl/Private/PhysicsControlComponentImpl.cpp
+++ b/Engine/Plugins/Experimental/PhysicsControl/Source/PhysicsControl/Private/PhysicsControlComponentImpl.cpp
@@ -26,7 +26,7 @@
 static USkeletalMeshComponent* GetValidSkeletalMeshComponentFromControlParent(
 	const FPhysicsControlRecord& Record)
 {
-	return Cast<USkeletalMeshComponent>(Record.PhysicsControl.ParentMeshComponent.Get());
+	return Cast<USkeletalMeshComponent>(Record.ParentMeshComponent.Get());
 }
 
 //======================================================================================================================
@@ -34,13 +34,13 @@ static USkeletalMeshComponent* GetValidSkeletalMeshComponentFromControlParent(
 static USkeletalMeshComponent* GetValidSkeletalMeshComponentFromControlChild(
 	const FPhysicsControlRecord& Record)
 {
-	return Cast<USkeletalMeshComponent>(Record.PhysicsControl.ChildMeshComponent.Get());
+	return Cast<USkeletalMeshComponent>(Record.ChildMeshComponent.Get());
 }
 
 //======================================================================================================================
 // Helper to get a valid skeletal mesh component pointer from a record
 static USkeletalMeshComponent* GetValidSkeletalMeshComponentFromBodyModifier(
-	const FPhysicsBodyModifier& PhysicsBodyModifier)
+	const FPhysicsBodyModifierRecord& PhysicsBodyModifier)
 {
 	return Cast<USkeletalMeshComponent>(PhysicsBodyModifier.MeshComponent.Get());
 }
@@ -57,7 +57,7 @@ bool UPhysicsControlComponent::GetBoneData(
 
 	if (BoneIndex == INDEX_NONE)
 	{
-		UE_LOG(LogPhysicsControlComponent, Warning, TEXT("Failed to find BoneIndex for %s"), *InBoneName.ToString());
+		UE_LOG(LogPhysicsControl, Warning, TEXT("Failed to find BoneIndex for %s"), *InBoneName.ToString());
 		return false;
 	}
 
@@ -71,10 +71,10 @@ bool UPhysicsControlComponent::GetBoneData(
 			OutBoneData = CachedSkeletalMeshData->BoneData[BoneIndex];
 			return true;
 		}
-		UE_LOG(LogPhysicsControlComponent, Warning, TEXT("BoneIndex is out of range"));
+		UE_LOG(LogPhysicsControl, Warning, TEXT("BoneIndex is out of range"));
 
 	}
-	UE_LOG(LogPhysicsControlComponent, Warning, TEXT("Failed to find bone data for %s"), *InBoneName.ToString());
+	UE_LOG(LogPhysicsControl, Warning, TEXT("Failed to find bone data for %s"), *InBoneName.ToString());
 	return false;
 }
 
@@ -90,7 +90,7 @@ bool UPhysicsControlComponent::GetModifiableBoneData(
 
 	if (BoneIndex == INDEX_NONE)
 	{
-		UE_LOG(LogPhysicsControlComponent, Warning, TEXT("Failed to find BoneIndex for %s"), *InBoneName.ToString());
+		UE_LOG(LogPhysicsControl, Warning, TEXT("Failed to find BoneIndex for %s"), *InBoneName.ToString());
 		return false;
 	}
 
@@ -104,10 +104,10 @@ bool UPhysicsControlComponent::GetModifiableBoneData(
 			OutBoneData = &CachedSkeletalMeshData->BoneData[BoneIndex];
 			return true;
 		}
-		UE_LOG(LogPhysicsControlComponent, Warning, TEXT("BoneIndex is out of range"));
+		UE_LOG(LogPhysicsControl, Warning, TEXT("BoneIndex is out of range"));
 
 	}
-	UE_LOG(LogPhysicsControlComponent, Warning, TEXT("Failed to find bone data for %s"), *InBoneName.ToString());
+	UE_LOG(LogPhysicsControl, Warning, TEXT("Failed to find bone data for %s"), *InBoneName.ToString());
 	return false;
 }
 
@@ -115,7 +115,7 @@ bool UPhysicsControlComponent::GetModifiableBoneData(
 //======================================================================================================================
 FPhysicsControlRecord* UPhysicsControlComponent::FindControlRecord(const FName Name)
 {
-	if (FPhysicsControlRecord* Record = PhysicsControlRecords.Find(Name))
+	if (FPhysicsControlRecord* Record = ControlRecords.Find(Name))
 	{
 		return Record;
 	}
@@ -125,7 +125,7 @@ FPhysicsControlRecord* UPhysicsControlComponent::FindControlRecord(const FName N
 //======================================================================================================================
 const FPhysicsControlRecord* UPhysicsControlComponent::FindControlRecord(const FName Name) const
 {
-	if (const FPhysicsControlRecord* Record = PhysicsControlRecords.Find(Name))
+	if (const FPhysicsControlRecord* Record = ControlRecords.Find(Name))
 	{
 		return Record;
 	}
@@ -242,7 +242,7 @@ void UPhysicsControlComponent::UpdateCachedSkeletalBoneData(float Dt)
 //======================================================================================================================
 void UPhysicsControlComponent::ResetControls(bool bKeepControlRecords)
 {
-	for (TPair<FName, FPhysicsControlRecord>& PhysicsControlRecordPair : PhysicsControlRecords)
+	for (TPair<FName, FPhysicsControlRecord>& PhysicsControlRecordPair : ControlRecords)
 	{
 		FPhysicsControlRecord& Record = PhysicsControlRecordPair.Value;
 		Record.ResetConstraint();
@@ -250,18 +250,18 @@ void UPhysicsControlComponent::ResetControls(bool bKeepControlRecords)
 
 	if (!bKeepControlRecords)
 	{
-		PhysicsControlRecords.Empty();
+		ControlRecords.Empty();
 	}
 }
 
 //======================================================================================================================
-void UPhysicsControlComponent::ApplyKinematicTarget(const FPhysicsBodyModifier& BodyModifier) const
+void UPhysicsControlComponent::ApplyKinematicTarget(const FPhysicsBodyModifierRecord& Record) const
 {
 	// Seems like static and skeletal meshes need to be handled differently
-	if (USkeletalMeshComponent* SkeletalMeshComponent = GetValidSkeletalMeshComponentFromBodyModifier(BodyModifier))
+	if (USkeletalMeshComponent* SkeletalMeshComponent = GetValidSkeletalMeshComponentFromBodyModifier(Record))
 	{
-		FBodyInstance* BodyInstance = UE::PhysicsControlComponent::GetBodyInstance(
-			BodyModifier.MeshComponent.Get(), BodyModifier.BoneName);
+		FBodyInstance* BodyInstance = UE::PhysicsControl::GetBodyInstance(
+			Record.MeshComponent.Get(), Record.BodyModifier.BoneName);
 		if (!BodyInstance)
 		{
 			return;
@@ -269,12 +269,12 @@ void UPhysicsControlComponent::ApplyKinematicTarget(const FPhysicsBodyModifier&
 
 		const FTransform TM = BodyInstance->GetUnrealWorldTransform(); // Preserve scale
 		FTransform KinematicTarget = TM;
-		KinematicTarget.SetRotation(BodyModifier.KinematicTargetOrientation);
-		KinematicTarget.SetTranslation(BodyModifier.KinematicTargetPosition);
-		if (BodyModifier.BodyModifierData.bUseSkeletalAnimation)
+		KinematicTarget.SetRotation(Record.KinematicTargetOrientation);
+		KinematicTarget.SetTranslation(Record.KinematicTargetPosition);
+		if (Record.BodyModifier.ModifierData.bUseSkeletalAnimation)
 		{
 			FCachedSkeletalMeshData::FBoneData BoneData;
-			if (GetBoneData(BoneData, SkeletalMeshComponent, BodyModifier.BoneName))
+			if (GetBoneData(BoneData, SkeletalMeshComponent, Record.BodyModifier.BoneName))
 			{
 				FTransform BoneTM = BoneData.GetTM();
 				KinematicTarget = KinematicTarget * BoneTM;
@@ -285,31 +285,31 @@ void UPhysicsControlComponent::ApplyKinematicTarget(const FPhysicsBodyModifier&
 	}
 	else
 	{
-		const FTransform TM = BodyModifier.MeshComponent->GetComponentToWorld();
+		const FTransform TM = Record.MeshComponent->GetComponentToWorld();
 		const ETeleportType TT = DetectTeleport(
 			TM.GetTranslation(), TM.GetRotation(), 
-			BodyModifier.KinematicTargetPosition, BodyModifier.KinematicTargetOrientation) 
+			Record.KinematicTargetPosition, Record.KinematicTargetOrientation)
 			? ETeleportType::ResetPhysics : ETeleportType::None;
 		// Note that calling BodyInstance->SetBodyTransform moves the physics, but not the mesh
-		BodyModifier.MeshComponent->SetWorldLocationAndRotation(
-			BodyModifier.KinematicTargetPosition, BodyModifier.KinematicTargetOrientation, false, nullptr, TT);
+		Record.MeshComponent->SetWorldLocationAndRotation(
+			Record.KinematicTargetPosition, Record.KinematicTargetOrientation, false, nullptr, TT);
 	}
 }
 
 //======================================================================================================================
-void UPhysicsControlComponent::ResetToCachedTarget(const FPhysicsBodyModifier& BodyModifier) const
+void UPhysicsControlComponent::ResetToCachedTarget(const FPhysicsBodyModifierRecord& Record) const
 {
-	FBodyInstance* BodyInstance = UE::PhysicsControlComponent::GetBodyInstance(
-		BodyModifier.MeshComponent.Get(), BodyModifier.BoneName);
+	FBodyInstance* BodyInstance = UE::PhysicsControl::GetBodyInstance(
+		Record.MeshComponent.Get(), Record.BodyModifier.BoneName);
 	if (!BodyInstance)
 	{
 		return;
 	}
 
-	if (USkeletalMeshComponent* SkeletalMeshComponent = GetValidSkeletalMeshComponentFromBodyModifier(BodyModifier))
+	if (USkeletalMeshComponent* SkeletalMeshComponent = GetValidSkeletalMeshComponentFromBodyModifier(Record))
 	{
 		FCachedSkeletalMeshData::FBoneData BoneData;
-		if (GetBoneData(BoneData, SkeletalMeshComponent, BodyModifier.BoneName))
+		if (GetBoneData(BoneData, SkeletalMeshComponent, Record.BodyModifier.BoneName))
 		{
 			FTransform BoneTM = BodyInstance->GetUnrealWorldTransform(); // Preserve scale
 			BoneTM.SetLocation(BoneData.Position);
@@ -350,7 +350,7 @@ bool UPhysicsControlComponent::RemoveSkeletalMeshReferenceForCaching(
 	check(InSkeletalMeshComponent);
 	if (!InSkeletalMeshComponent)
 	{
-		UE_LOG(LogPhysicsControlComponent, Warning, TEXT("Invalid skeletal mesh component"));
+		UE_LOG(LogPhysicsControl, Warning, TEXT("Invalid skeletal mesh component"));
 		return false;
 	}
 
@@ -370,7 +370,7 @@ bool UPhysicsControlComponent::RemoveSkeletalMeshReferenceForCaching(
 			return false;
 		}
 	}
-	UE_LOG(LogPhysicsControlComponent, Warning, TEXT("Failed to remove skeletal mesh component reference for caching"));
+	UE_LOG(LogPhysicsControl, Warning, TEXT("Failed to remove skeletal mesh component reference for caching"));
 	return false;
 }
 
@@ -415,7 +415,7 @@ bool UPhysicsControlComponent::RemoveSkeletalMeshReferenceForModifier(
 	check(InSkeletalMeshComponent);
 	if (!InSkeletalMeshComponent)
 	{
-		UE_LOG(LogPhysicsControlComponent, Warning, TEXT("Invalid skeletal mesh component"));
+		UE_LOG(LogPhysicsControl, Warning, TEXT("Invalid skeletal mesh component"));
 		return false;
 	}
 
@@ -435,77 +435,10 @@ bool UPhysicsControlComponent::RemoveSkeletalMeshReferenceForModifier(
 			return false;
 		}
 	}
-	UE_LOG(LogPhysicsControlComponent, Warning, TEXT("Failed to remove skeletal mesh component reference for modifier"));
+	UE_LOG(LogPhysicsControl, Warning, TEXT("Failed to remove skeletal mesh component reference for modifier"));
 	return false;
 }
 
-//======================================================================================================================
-FName UPhysicsControlComponent::GetUniqueControlName(
-	const FName ParentBoneName, const FName ChildBoneName, const FString& NamePrefix) const
-{
-	FString NameBase = NamePrefix;
-	if (!ParentBoneName.IsNone())
-	{
-		NameBase += ParentBoneName.ToString() + TEXT("_");
-	}
-	if (!ChildBoneName.IsNone())
-	{
-		NameBase += ChildBoneName.ToString() + TEXT("_");
-	}
-
-	TSet<FName> Keys;
-	PhysicsControlRecords.GetKeys(Keys);
-	// If the number gets too large, almost certainly we're in some nasty situation where this is
-	// getting called in a loop. Better to quit and fail, rather than allow the constraint set to
-	// increase without bound. 
-	for (int32 Index = 0; Index < MaxNumControlsOrModifiersPerName; ++Index)
-	{
-		FString NameStr = FString::Format(TEXT("{0}{1}"), { NameBase, Index });
-		FName Name(NameStr);
-		if (!Keys.Find(Name))
-		{
-			return Name;
-		}
-	}
-	UE_LOG(LogPhysicsControlComponent, Warning,
-		TEXT("Unable to find a suitable Control name - the limit of MaxNumControlsOrModifiersPerName (%d) has been exceeded"),
-		MaxNumControlsOrModifiersPerName);
-	return FName();
-}
-
-//======================================================================================================================
-FName UPhysicsControlComponent::GetUniqueBodyModifierName(const FName BoneName) const
-{
-	FString NameBase = TEXT("");
-	if (!BoneName.IsNone())
-	{
-		NameBase += BoneName.ToString() + TEXT("_");
-	}
-	else
-	{
-		NameBase = TEXT("Body_");
-	}
-
-	TSet<FName> Keys;
-	PhysicsBodyModifiers.GetKeys(Keys);
-	// If the number gets too large, almost certainly we're in some nasty situation where this is
-	// getting called in a loop. Better to quit and fail, rather than allow the modifier set to
-	// increase without bound. 
-	for (int32 Index = 0; Index != MaxNumControlsOrModifiersPerName; ++Index)
-	{
-		FString NameStr = FString::Format(TEXT("{0}{1}"), { NameBase, Index });
-		FName Name(NameStr);
-		if (!Keys.Find(Name))
-		{
-			return Name;
-		}
-	}
-	UE_LOG(LogPhysicsControlComponent, Warning,
-		TEXT("Unable to find a suitable Body Modifier name - the limit of MaxNumControlsOrModifiersPerName (%d) has been exceeded"),
-		MaxNumControlsOrModifiersPerName);
-	return FName();
-}
-
 //======================================================================================================================
 void UPhysicsControlComponent::CalculateControlTargetData(
 	FTransform&                  OutTargetTM,
@@ -539,13 +472,13 @@ void UPhysicsControlComponent::CalculateControlTargetData(
 			bHaveParentBoneData = GetBoneData(
 				ParentBoneData, ParentSkeletalMeshComponent, Record.PhysicsControl.ParentBoneName);
 		}
-		else if (Record.PhysicsControl.ParentMeshComponent.IsValid())
+		else if (Record.ParentMeshComponent.IsValid())
 		{
-			const FTransform ParentTM = Record.PhysicsControl.ParentMeshComponent->GetComponentTransform();
+			const FTransform ParentTM = Record.ParentMeshComponent->GetComponentTransform();
 			ParentBoneData.Position = ParentTM.GetLocation();
 			ParentBoneData.Orientation = ParentTM.GetRotation();
-			ParentBoneData.Velocity = Record.PhysicsControl.ParentMeshComponent->GetPhysicsLinearVelocity();
-			ParentBoneData.AngularVelocity = Record.PhysicsControl.ParentMeshComponent->GetPhysicsAngularVelocityInRadians();
+			ParentBoneData.Velocity = Record.ParentMeshComponent->GetPhysicsLinearVelocity();
+			ParentBoneData.AngularVelocity = Record.ParentMeshComponent->GetPhysicsAngularVelocityInRadians();
 			bHaveParentBoneData = true;
 			bCalculateVelocity = false;
 		}
@@ -640,7 +573,7 @@ void UPhysicsControlComponent::CalculateControlTargetData(
 	// Now apply the explicit target specified in the record. It operates in the space of the target
 	// transform we (may have) just calculated.
 	{
-		const FPhysicsControlTarget& Target = Record.PhysicsControl.ControlTarget;
+		const FPhysicsControlTarget& Target = Record.ControlTarget;
 
 		// Calculate the authored target position/orientation - i.e. not using the skeletal animation
 		FQuat TargetOrientationQ = Target.TargetOrientation.Quaternion();
@@ -650,7 +583,7 @@ void UPhysicsControlComponent::CalculateControlTargetData(
 		// Incorporate the offset from the control point. If we used animation, then we don't need
 		// to do this.
 		if (!bUsedSkeletalAnimation && 
-			Record.PhysicsControl.ControlTarget.bApplyControlPointToTarget)
+			Record.ControlTarget.bApplyControlPointToTarget)
 		{
 			ExtraTargetPosition = TargetOrientationQ * Record.GetControlPoint();
 		}
@@ -684,12 +617,12 @@ bool UPhysicsControlComponent::ApplyControlStrengths(
 	FVector LinearDamping;
 	const FVector MaxForce = Data.MaxForce * Multiplier.MaxForceMultiplier;
 
-	UE::PhysicsControlComponent::ConvertStrengthToSpringParams(
+	UE::PhysicsControl::ConvertStrengthToSpringParams(
 		AngularSpring, AngularDamping,
 		Data.AngularStrength * Multiplier.AngularStrengthMultiplier,
 		Data.AngularDampingRatio * Multiplier.AngularDampingRatioMultiplier,
 		Data.AngularExtraDamping * Multiplier.AngularExtraDampingMultiplier);
-	UE::PhysicsControlComponent::ConvertStrengthToSpringParams(
+	UE::PhysicsControl::ConvertStrengthToSpringParams(
 		LinearSpring, LinearDamping,
 		Data.LinearStrength * Multiplier.LinearStrengthMultiplier,
 		Data.LinearDampingRatio * Multiplier.LinearDampingRatioMultiplier,
@@ -736,7 +669,7 @@ void UPhysicsControlComponent::ApplyControl(FPhysicsControlRecord& Record)
 		return;
 	}
 
-	if (!Record.Enabled())
+	if (!Record.PhysicsControl.IsEnabled())
 	{
 		// Note that this will disable the constraint elements when strength/damping are zero
 		ConstraintInstance->SetDriveParams(
@@ -750,11 +683,11 @@ void UPhysicsControlComponent::ApplyControl(FPhysicsControlRecord& Record)
 	// constraint-controlled collision doesn't interact nicely when there are multiple constraints.
 	ConstraintInstance->SetDisableCollision(Record.PhysicsControl.ControlData.bDisableCollision);
 
-	FBodyInstance* ParentBodyInstance = UE::PhysicsControlComponent::GetBodyInstance(
-		Record.PhysicsControl.ParentMeshComponent.Get(), Record.PhysicsControl.ParentBoneName);
+	FBodyInstance* ParentBodyInstance = UE::PhysicsControl::GetBodyInstance(
+		Record.ParentMeshComponent.Get(), Record.PhysicsControl.ParentBoneName);
 
-	FBodyInstance* ChildBodyInstance = UE::PhysicsControlComponent::GetBodyInstance(
-		Record.PhysicsControl.ChildMeshComponent.Get(), Record.PhysicsControl.ChildBoneName);
+	FBodyInstance* ChildBodyInstance = UE::PhysicsControl::GetBodyInstance(
+		Record.ChildMeshComponent.Get(), Record.PhysicsControl.ChildBoneName);
 
 	if (!ParentBodyInstance && !ChildBodyInstance)
 	{
@@ -787,36 +720,36 @@ void UPhysicsControlComponent::ApplyControl(FPhysicsControlRecord& Record)
 }
 
 //======================================================================================================================
-void UPhysicsControlComponent::ApplyBodyModifier(FPhysicsBodyModifier& BodyModifier)
+void UPhysicsControlComponent::ApplyBodyModifier(FPhysicsBodyModifierRecord& Record)
 {
-	USkeletalMeshComponent* SKM = GetValidSkeletalMeshComponentFromBodyModifier(BodyModifier);
-	FBodyInstance* BodyInstance = UE::PhysicsControlComponent::GetBodyInstance(
-		BodyModifier.MeshComponent.Get(), BodyModifier.BoneName);
+	USkeletalMeshComponent* SKM = GetValidSkeletalMeshComponentFromBodyModifier(Record);
+	FBodyInstance* BodyInstance = UE::PhysicsControl::GetBodyInstance(
+		Record.MeshComponent.Get(), Record.BodyModifier.BoneName);
 	if (BodyInstance)
 	{
-		switch (BodyModifier.BodyModifierData.MovementType)
+		switch (Record.BodyModifier.ModifierData.MovementType)
 		{
 		case EPhysicsMovementType::Static:
 			BodyInstance->SetInstanceSimulatePhysics(false, false, true);
 			break;
 		case EPhysicsMovementType::Kinematic:
 			BodyInstance->SetInstanceSimulatePhysics(false, false, true);
-			ApplyKinematicTarget(BodyModifier);
+			ApplyKinematicTarget(Record);
 			break;
 		case EPhysicsMovementType::Simulated:
 			BodyInstance->SetInstanceSimulatePhysics(true, false, true);
 			break;
 		default:
-			UE_LOG(LogPhysicsControlComponent, Warning, TEXT("Invalid movement type %d"),
-				int(BodyModifier.BodyModifierData.MovementType));
+			UE_LOG(LogPhysicsControl, Warning, TEXT("Invalid movement type %d"),
+				int(Record.BodyModifier.ModifierData.MovementType));
 			break;
 		}
 
 		// We always overwrite the physics blend weight, since the functions above can still modify
 		// it (even though they all use the "maintain physics blending" option), since there is an
 		// expectation that zero blend weight means to disable physics.
-		BodyInstance->PhysicsBlendWeight = BodyModifier.BodyModifierData.PhysicsBlendWeight;
-		BodyInstance->SetUpdateKinematicFromSimulation(BodyModifier.BodyModifierData.bUpdateKinematicFromSimulation);
+		BodyInstance->PhysicsBlendWeight = Record.BodyModifier.ModifierData.PhysicsBlendWeight;
+		BodyInstance->SetUpdateKinematicFromSimulation(Record.BodyModifier.ModifierData.bUpdateKinematicFromSimulation);
 
 		UBodySetup* BodySetup = BodyInstance->GetBodySetup();
 		if (BodySetup)
@@ -824,7 +757,7 @@ void UPhysicsControlComponent::ApplyBodyModifier(FPhysicsBodyModifier& BodyModif
 			int32 NumShapes = BodySetup->AggGeom.GetElementCount();
 			for (int32 ShapeIndex = 0; ShapeIndex != NumShapes; ++ShapeIndex)
 			{
-				BodyInstance->SetShapeCollisionEnabled(ShapeIndex, BodyModifier.BodyModifierData.CollisionType);
+				BodyInstance->SetShapeCollisionEnabled(ShapeIndex, Record.BodyModifier.ModifierData.CollisionType);
 			}
 		}
 
@@ -832,29 +765,28 @@ void UPhysicsControlComponent::ApplyBodyModifier(FPhysicsBodyModifier& BodyModif
 		{
 			const float GravityZ = BodyInstance->GetPhysicsScene()->GetOwningWorld()->GetGravityZ();
 			const float AppliedGravityZ = BodyInstance->bEnableGravity ? GravityZ : 0.0f;
-			const float DesiredGravityZ = GravityZ * BodyModifier.BodyModifierData.GravityMultiplier;
+			const float DesiredGravityZ = GravityZ * Record.BodyModifier.ModifierData.GravityMultiplier;
 			const float GravityZToApply = DesiredGravityZ - AppliedGravityZ;
 			BodyInstance->AddForce(FVector(0, 0, GravityZToApply), true, true);
 		}
 	}
-	if (BodyModifier.bResetToCachedTarget)
+	if (Record.bResetToCachedTarget)
 	{
-		BodyModifier.bResetToCachedTarget = false;
-		ResetToCachedTarget(BodyModifier);
+		Record.bResetToCachedTarget = false;
+		ResetToCachedTarget(Record);
 	}
 }
 
-
 //======================================================================================================================
-FPhysicsBodyModifier* UPhysicsControlComponent::FindBodyModifier(const FName Name)
+FPhysicsBodyModifierRecord* UPhysicsControlComponent::FindBodyModifierRecord(const FName Name)
 {
-	return PhysicsBodyModifiers.Find(Name);
+	return BodyModifierRecords.Find(Name);
 }
 
 //======================================================================================================================
-const FPhysicsBodyModifier* UPhysicsControlComponent::FindBodyModifier(const FName Name) const
+const FPhysicsBodyModifierRecord* UPhysicsControlComponent::FindBodyModifierRecord(const FName Name) const
 {
-	return PhysicsBodyModifiers.Find(Name);
+	return BodyModifierRecords.Find(Name);
 }
 
 //======================================================================================================================
@@ -876,13 +808,13 @@ bool UPhysicsControlComponent::DestroyControl(const FName Name, const EDestroyBe
 		NameRecords.RemoveControl(Name);
 		if (DestroyBehavior == EDestroyBehavior::RemoveRecord)
 		{
-			check(PhysicsControlRecords.Remove(Name) == 1);
+			check(ControlRecords.Remove(Name) == 1);
 		}
 		return true;
 	}
 	if (bWarnAboutInvalidNames)
 	{
-		UE_LOG(LogPhysicsControlComponent, Warning,
+		UE_LOG(LogPhysicsControl, Warning,
 			TEXT("DestroyControl - invalid name %s"), *Name.ToString());
 	}
 	return false;
@@ -891,7 +823,7 @@ bool UPhysicsControlComponent::DestroyControl(const FName Name, const EDestroyBe
 //======================================================================================================================
 bool UPhysicsControlComponent::DestroyBodyModifier(const FName Name, const EDestroyBehavior DestroyBehavior)
 {
-	FPhysicsBodyModifier* BodyModifier = FindBodyModifier(Name);
+	FPhysicsBodyModifierRecord* BodyModifier = FindBodyModifierRecord(Name);
 	if (BodyModifier)
 	{
 		if (USkeletalMeshComponent* SkeletalMeshComponent = GetValidSkeletalMeshComponentFromBodyModifier(*BodyModifier))
@@ -902,13 +834,13 @@ bool UPhysicsControlComponent::DestroyBodyModifier(const FName Name, const EDest
 		NameRecords.RemoveBodyModifier(Name);
 		if (DestroyBehavior == EDestroyBehavior::RemoveRecord)
 		{
-			check(PhysicsBodyModifiers.Remove(Name) == 1);
+			check(BodyModifierRecords.Remove(Name) == 1);
 		}
 		return true;
 	}
 	if (bWarnAboutInvalidNames)
 	{
-		UE_LOG(LogPhysicsControlComponent, Warning,
+		UE_LOG(LogPhysicsControl, Warning,
 			TEXT("DestroyBodyModifier - invalid name %s"), *Name.ToString());
 	}
 	return false;
diff --git a/Engine/Plugins/Experimental/PhysicsControl/Source/PhysicsControl/Private/PhysicsControlData.cpp b/Engine/Plugins/Experimental/PhysicsControl/Source/PhysicsControl/Private/PhysicsControlData.cpp
--- a/Engine/Plugins/Experimental/PhysicsControl/Source/PhysicsControl/Private/PhysicsControlData.cpp
+++ b/Engine/Plugins/Experimental/PhysicsControl/Source/PhysicsControl/Private/PhysicsControlData.cpp
@@ -8,6 +8,23 @@
 #define SELECT_PARAM(NAME) Output.NAME = (Weight < 0.5f) ? A.NAME : B.NAME
 #define SET_ENABLED_PARAM(NAME) Output.bEnable##NAME = A.bEnable##NAME && B.bEnable##NAME
 
+//======================================================================================================================
+FPhysicsControlSetUpdates& FPhysicsControlSetUpdates::operator+=(const FPhysicsControlSetUpdates& Other)
+{
+	ControlSetUpdates.Append(Other.ControlSetUpdates);
+	ModifierSetUpdates.Append(Other.ModifierSetUpdates);
+	return *this;
+}
+
+//======================================================================================================================
+FPhysicsControlAndBodyModifierCreationDatas& FPhysicsControlAndBodyModifierCreationDatas::operator+=(
+	const FPhysicsControlAndBodyModifierCreationDatas& Other)
+{
+	Controls.Append(Other.Controls);
+	Modifiers.Append(Other.Modifiers);
+	return *this;
+}
+
 //======================================================================================================================
 FPhysicsControlData Interpolate(
 	const FPhysicsControlData& A, const FPhysicsControlData& B, const float Weight)
diff --git a/Engine/Plugins/Experimental/PhysicsControl/Source/PhysicsControl/Private/PhysicsControlLimbData.cpp b/Engine/Plugins/Experimental/PhysicsControl/Source/PhysicsControl/Private/PhysicsControlLimbData.cpp
new file mode 100644
--- /dev/null
+++ b/Engine/Plugins/Experimental/PhysicsControl/Source/PhysicsControl/Private/PhysicsControlLimbData.cpp
@@ -0,0 +1,15 @@
+// Copyright Epic Games, Inc. All Rights Reserved.
+
+#include "PhysicsControlLimbData.h"
+
+//======================================================================================================================
+FPhysicsControlCharacterSetupData& FPhysicsControlCharacterSetupData::operator+=(
+	const FPhysicsControlCharacterSetupData& Other)
+{
+	LimbSetupData += Other.LimbSetupData;
+	DefaultWorldSpaceControlData = Other.DefaultParentSpaceControlData;
+	DefaultParentSpaceControlData = Other.DefaultParentSpaceControlData;
+	DefaultBodyModifierData = Other.DefaultBodyModifierData;
+	return *this;
+}
+
diff --git a/Engine/Plugins/Experimental/PhysicsControl/Source/PhysicsControl/Private/PhysicsControlLog.cpp b/Engine/Plugins/Experimental/PhysicsControl/Source/PhysicsControl/Private/PhysicsControlLog.cpp
--- a/Engine/Plugins/Experimental/PhysicsControl/Source/PhysicsControl/Private/PhysicsControlLog.cpp
+++ b/Engine/Plugins/Experimental/PhysicsControl/Source/PhysicsControl/Private/PhysicsControlLog.cpp
@@ -2,5 +2,4 @@
 
 #include "PhysicsControlLog.h"
 
-PHYSICSCONTROL_API DEFINE_LOG_CATEGORY(LogPhysicsControlComponent);
-PHYSICSCONTROL_API DEFINE_LOG_CATEGORY(LogRigidBodyWithControl);
+PHYSICSCONTROL_API DEFINE_LOG_CATEGORY(LogPhysicsControl);
diff --git a/Engine/Plugins/Experimental/PhysicsControl/Source/PhysicsControl/Private/PhysicsControlOperatorNameGeneration.cpp b/Engine/Plugins/Experimental/PhysicsControl/Source/PhysicsControl/Private/PhysicsControlOperatorNameGeneration.cpp
--- a/Engine/Plugins/Experimental/PhysicsControl/Source/PhysicsControl/Private/PhysicsControlOperatorNameGeneration.cpp
+++ b/Engine/Plugins/Experimental/PhysicsControl/Source/PhysicsControl/Private/PhysicsControlOperatorNameGeneration.cpp
@@ -2,28 +2,17 @@
 
 #include "PhysicsControlOperatorNameGeneration.h"
 
-#include "ReferenceSkeleton.h"
-#include "PhysicsEngine/PhysicsAsset.h"
 #include "AnimNode_RigidBodyWithControl.h"
 #include "PhysicsControlLimbData.h"
+#include "PhysicsControlRecord.h"
+#include "PhysicsEngine/PhysicsAsset.h"
 
-//======================================================================================================================
-template<typename TOperatorCreationType> FName FindUniqueOperatorName(const TOperatorCreationType& OperatorCreationSpecifier, const TSet<FName>& ExistingNames)
-{
-	return NAME_None;
-}
+#include "ReferenceSkeleton.h"
 
-//======================================================================================================================
-template<> FName FindUniqueOperatorName<FRigidBodyModifierCreation>(const FRigidBodyModifierCreation& OperatorCreationSpecifier, const TSet<FName>& ExistingNames)
+namespace UE
 {
-	return FindUniqueBodyModifierName(OperatorCreationSpecifier.Modifier.BoneName, ExistingNames);;
-}
-
-//======================================================================================================================
-template<> FName FindUniqueOperatorName<FRigidBodyControlCreation>(const FRigidBodyControlCreation& OperatorCreationSpecifier, const TSet<FName>& ExistingNames)
+namespace PhysicsControl
 {
-	return FindUniqueControlName(OperatorCreationSpecifier.Control.ParentBoneName, OperatorCreationSpecifier.Control.ChildBoneName, ExistingNames);
-}
 
 //======================================================================================================================
 FName FindParentBodyBoneName(const FName BoneName, const FReferenceSkeleton& RefSkeleton, UPhysicsAsset* const PhysicsAsset)
@@ -49,19 +38,23 @@ FName FindParentBodyBoneName(const FName BoneName, const FReferenceSkeleton& Ref
 }
 
 //======================================================================================================================
-template<typename TOperatorFunctorType> void CreateAdditionalBodyModifiers(const TArray<FRigidBodyModifierCreation> CreationSpecifiers, FPhysicsControlNameRecords& NameRecords, TOperatorFunctorType& OperatorFunctor)
+template<typename TOperatorFunctorType> void CreateAdditionalBodyModifiers(
+	const TMap<FName, FPhysicsBodyModifierCreationData>& CreationSpecifiers, 
+	FPhysicsControlNameRecords&                          NameRecords, 
+	TOperatorFunctorType&                                OperatorFunctor)
 {
-	for (const FRigidBodyModifierCreation& Specifier : CreationSpecifiers)
+	for (TMap<FName, FPhysicsBodyModifierCreationData>::ElementType CreationPair : CreationSpecifiers)
 	{
-		const FName Name = OperatorFunctor(Specifier.Modifier.BoneName, Specifier.Modifier.ModifierData);
+		const FName Name = CreationPair.Key;
+		const FPhysicsBodyModifierCreationData& Specifier = CreationPair.Value;
 		if (Name.IsNone())
 		{
-			UE_LOG(LogRigidBodyWithControl, Warning,
+			UE_LOG(LogPhysicsControl, Warning,
 				TEXT("CreateAdditionalControls: Failed to make body modifier for %s"), *Specifier.Modifier.BoneName.ToString());
 		}
 		else
 		{
-			UE_LOG(LogRigidBodyWithControl, Verbose,
+			UE_LOG(LogPhysicsControl, Verbose,
 				TEXT("Made modifier %s for %s"),
 				*Name.ToString(),
 				*Specifier.Modifier.BoneName.ToString());
@@ -72,20 +65,24 @@ template<typename TOperatorFunctorType> void CreateAdditionalBodyModifiers(const
 }
 
 //======================================================================================================================
-template<typename TOperatorFunctorType> void CreateAdditionalControls(const TArray<FRigidBodyControlCreation> CreationSpecifiers, FPhysicsControlNameRecords& NameRecords, TOperatorFunctorType& OperatorFunctor)
+template<typename TOperatorFunctorType> void CreateAdditionalControls(
+	const TMap<FName, FPhysicsControlCreationData>& CreationSpecifiers, 
+	FPhysicsControlNameRecords&                     NameRecords, 
+	TOperatorFunctorType&                           OperatorFunctor)
 {
-	for (const FRigidBodyControlCreation& Specifier : CreationSpecifiers)
+	for (TMap<FName, FPhysicsControlCreationData>::ElementType CreationPair : CreationSpecifiers)
 	{
-		const FName Name = OperatorFunctor(Specifier.Control.ParentBoneName, Specifier.Control.ChildBoneName, Specifier.Control.ControlData);
+		const FName Name = CreationPair.Key;
+		const FPhysicsControlCreationData& Specifier = CreationPair.Value;
 		if (Name.IsNone())
 		{
-			UE_LOG(LogRigidBodyWithControl, Warning,
+			UE_LOG(LogPhysicsControl, Warning,
 				TEXT("CreateAdditionalControls: Failed to make control between %s and %s"),
 				*Specifier.Control.ParentBoneName.ToString(), *Specifier.Control.ChildBoneName.ToString());
 		}
 		else
 		{
-			UE_LOG(LogRigidBodyWithControl, Verbose,
+			UE_LOG(LogPhysicsControl, Verbose,
 				TEXT("Made control %s between %s and %s"),
 				*Name.ToString(),
 				*Specifier.Control.ParentBoneName.ToString(),
@@ -95,17 +92,16 @@ template<typename TOperatorFunctorType> void CreateAdditionalControls(const TArr
 	}
 }
 
-
 //======================================================================================================================
 template<typename TFunctor> void CreateControlsFromLimbBones(
-	const FName                   LimbName,
+	const FName                     LimbName,
 	const FPhysicsControlLimbBones& LimbBones,
-	const EPhysicsControlType   ControlType,
-	const FReferenceSkeleton& RefSkeleton,
-	UPhysicsAsset* const PhysicsAsset,
-	const FPhysicsControlData& ControlData,
-	FPhysicsControlNameRecords& NameRecords,
-	TFunctor& CreateOperationLambda)
+	const EPhysicsControlType       ControlType,
+	const FReferenceSkeleton&       RefSkeleton,
+	UPhysicsAsset*                  PhysicsAsset,
+	const FPhysicsControlData&      ControlData,
+	FPhysicsControlNameRecords&     NameRecords,
+	TFunctor&                       CreateOperationLambda)
 {
 	for (const FName ChildBoneName : LimbBones.BoneNames)
 	{
@@ -134,12 +130,12 @@ template<typename TFunctor> void CreateControlsFromLimbBones(
 		}
 		else if (!ParentBoneName.IsNone())
 		{
-			UE_LOG(LogRigidBodyWithControl, Warning,
+			UE_LOG(LogPhysicsControl, Warning,
 				TEXT("Failed to find unique control name for %s and %s"), *ParentBoneName.ToString(), *ChildBoneName.ToString());
 		}
 		else
 		{
-			UE_LOG(LogRigidBodyWithControl, Warning,
+			UE_LOG(LogPhysicsControl, Warning,
 				TEXT("Failed to find unique control name for %s"), *ChildBoneName.ToString());
 		}
 	}
@@ -163,44 +159,60 @@ template<typename TFunctor> void CreateBodyModifiersFromLimbBones(
 		}
 		else
 		{
-			UE_LOG(LogRigidBodyWithControl, Warning,
+			UE_LOG(LogPhysicsControl, Warning,
 				TEXT("Failed to find unique body modifier name for %s"), *BoneName.ToString());
 		}
 	}
 }
 
 //======================================================================================================================
-template<typename TControlFunctorType, typename TBodyModifierFunctorType> void ForEachPotentialOperator(const FAnimNode_RigidBodyWithControl* const Node, TMap<FName, FPhysicsControlLimbBones> AllLimbBones, const FReferenceSkeleton& RefSkeleton, UPhysicsAsset* const PhysicsAsset, FPhysicsControlNameRecords& NameRecords, TControlFunctorType& ControlFunctor, TBodyModifierFunctorType& BodyModifierFunctor)
+template<typename TControlFunctorType, typename TBodyModifierFunctorType>
+void ForEachPotentialOperator(
+	const FPhysicsControlCharacterSetupData&           CharacterSetupData,
+	const FPhysicsControlAndBodyModifierCreationDatas& AdditionalControlsAndBodyModifiers,
+	TMap<FName, FPhysicsControlLimbBones>              AllLimbBones, 
+	const FReferenceSkeleton&                          RefSkeleton, 
+	UPhysicsAsset*                                     PhysicsAsset, 
+	FPhysicsControlNameRecords&                        NameRecords, 
+	TControlFunctorType&                               ControlFunctor, 
+	TBodyModifierFunctorType&                          BodyModifierFunctor)
 {
-	if (Node != nullptr)
+	for (const TMap<FName, FPhysicsControlLimbBones>::ElementType& LimbBoneEntry : AllLimbBones)
 	{
-		for (const TMap<FName, FPhysicsControlLimbBones>::ElementType& LimbBoneEntry : AllLimbBones)
+		const FName LimbName = LimbBoneEntry.Key;
+		const FPhysicsControlLimbBones& LimbBones = LimbBoneEntry.Value;
+		if (LimbBones.bCreateWorldSpaceControls)
 		{
-			const FName LimbName = LimbBoneEntry.Key;
-			const FPhysicsControlLimbBones& LimbBones = LimbBoneEntry.Value;
-			if (LimbBones.bCreateWorldSpaceControls)
-			{
-				CreateControlsFromLimbBones(LimbName, LimbBones, EPhysicsControlType::WorldSpace, RefSkeleton, PhysicsAsset, Node->SetupData.DefaultWorldSpaceControlData, NameRecords, ControlFunctor);
-			}
-			if (LimbBones.bCreateParentSpaceControls)
-			{
-				CreateControlsFromLimbBones(LimbName, LimbBones, EPhysicsControlType::ParentSpace, RefSkeleton, PhysicsAsset, Node->SetupData.DefaultParentSpaceControlData, NameRecords, ControlFunctor);
-			}
-			if (LimbBones.bCreateBodyModifiers)
-			{
-				CreateBodyModifiersFromLimbBones(LimbName, LimbBones, Node->SetupData.DefaultBodyModifierData, NameRecords, BodyModifierFunctor);
-			}
+			CreateControlsFromLimbBones(
+				LimbName, LimbBones, EPhysicsControlType::WorldSpace, RefSkeleton, PhysicsAsset, 
+				CharacterSetupData.DefaultWorldSpaceControlData, NameRecords, ControlFunctor);
+		}
+		if (LimbBones.bCreateParentSpaceControls)
+		{
+			CreateControlsFromLimbBones(
+				LimbName, LimbBones, EPhysicsControlType::ParentSpace, RefSkeleton, PhysicsAsset, 
+				CharacterSetupData.DefaultParentSpaceControlData, NameRecords, ControlFunctor);
+		}
+		if (LimbBones.bCreateBodyModifiers)
+		{
+			CreateBodyModifiersFromLimbBones(
+				LimbName, LimbBones, CharacterSetupData.DefaultBodyModifierData, NameRecords, BodyModifierFunctor);
 		}
-
-		// Find names for any additional controls that have been requested	
-		CreateAdditionalBodyModifiers(Node->AdditionalControlsAndBodyModifiers.Modifiers, NameRecords, BodyModifierFunctor);
-		CreateAdditionalControls(Node->AdditionalControlsAndBodyModifiers.Controls, NameRecords, ControlFunctor);
 	}
+
+	// Find names for any additional controls that have been requested	
+	CreateAdditionalBodyModifiers(AdditionalControlsAndBodyModifiers.Modifiers, NameRecords, BodyModifierFunctor);
+	CreateAdditionalControls(AdditionalControlsAndBodyModifiers.Controls, NameRecords, ControlFunctor);
 }
 
 //======================================================================================================================
 // Slightly annoying to have to add the names individually, but we want to check they exist
-template<typename TBodyModifierNameContainerType, typename TControlNameContainerType> void CreateAdditionalSets_Implementation(const FPhysicsControlSetUpdates& AdditionalSets, const TBodyModifierNameContainerType& BodyModifierNames, const TControlNameContainerType& ControlNames, FPhysicsControlNameRecords& NameRecords)
+template<typename TBodyModifierNameContainerType, typename TControlNameContainerType>
+void CreateAdditionalSets_Implementation(
+	const FPhysicsControlSetUpdates&      AdditionalSets, 
+	const TBodyModifierNameContainerType& BodyModifierNames, 
+	const TControlNameContainerType&      ControlNames, 
+	FPhysicsControlNameRecords&           NameRecords)
 {
 	for (const FPhysicsControlSetUpdate& Set : AdditionalSets.ControlSetUpdates)
 	{
@@ -214,7 +226,7 @@ template<typename TBodyModifierNameContainerType, typename TControlNameContainer
 			}
 			else
 			{
-				UE_LOG(LogRigidBodyWithControl, Warning,
+				UE_LOG(LogPhysicsControl, Warning,
 					TEXT("CreateAdditionalSets: Failed to find control with name %s to add to set %s"),
 					*Name.ToString(), *Set.SetName.ToString());
 			}
@@ -233,7 +245,7 @@ template<typename TBodyModifierNameContainerType, typename TControlNameContainer
 			}
 			else
 			{
-				UE_LOG(LogRigidBodyWithControl, Warning,
+				UE_LOG(LogPhysicsControl, Warning,
 					TEXT("CreateAdditionalSets: Failed to find body modifier with name %s to add to set %s"),
 					*Name.ToString(), *Set.SetName.ToString());
 			}
@@ -241,59 +253,6 @@ template<typename TBodyModifierNameContainerType, typename TControlNameContainer
 	}
 }
 
-//======================================================================================================================
-FName FindUniqueName(const FString& NameBase, const TSet<FName>& Keys, const int32 MaxNmeaIndex)
-{
-	FString NameStr = NameBase;
-
-	// If the number gets too large, almost certainly we're in some nasty situation where this is
-	// getting called in a loop. Better to quit and fail, rather than allow the constraint set to
-	// increase without bound. 
-	for (int32 Index = 0; Index < MaxNmeaIndex; ++Index)
-	{
-		const FName Name(NameStr);
-		if (!Keys.Find(Name))
-		{
-			return Name;
-		}
-		NameStr = FString::Format(TEXT("{0}_{1}"), { NameBase, Index });
-	}
-
-	return NAME_None;
-}
-
-//======================================================================================================================
-FName FindUniqueBodyModifierName(const FName BodyName, const TSet<FName>& ExistingNames)
-{
-	FString NameBase = TEXT("");
-	if (!BodyName.IsNone())
-	{
-		NameBase += BodyName.ToString();
-	}
-	else
-	{
-		NameBase = TEXT("Body");
-	}
-
-	return FindUniqueName(NameBase, ExistingNames, MaxNumControlsOrModifiersPerName);
-}
-
-//======================================================================================================================
-FName FindUniqueControlName(const FName ParentBodyName, const FName ChildBodyName, const TSet<FName>& ExistingNames)
-{
-	FString NameBase = TEXT("");
-	if (!ParentBodyName.IsNone())
-	{
-		NameBase += ParentBodyName.ToString() + TEXT("_");
-	}
-	if (!ChildBodyName.IsNone())
-	{
-		NameBase += ChildBodyName.ToString();
-	}
-
-	return FindUniqueName(NameBase, ExistingNames, MaxNumControlsOrModifiersPerName);
-}
-
 //======================================================================================================================
 TMap<FName, FPhysicsControlLimbBones> GetLimbBones(
 	const TArray<FPhysicsControlLimbSetupData>& LimbSetupData,
@@ -306,7 +265,7 @@ TMap<FName, FPhysicsControlLimbBones> GetLimbBones(
 
 	if (!PhysicsAsset)
 	{
-		UE_LOG(LogRigidBodyWithControl, Warning, TEXT("Physics asset missing"));
+		UE_LOG(LogPhysicsControl, Warning, TEXT("Physics asset missing"));
 		return Result;
 	}
 
@@ -351,28 +310,87 @@ TMap<FName, FPhysicsControlLimbBones> GetLimbBones(
 	return Result;
 }
 
-void CollectOperatorNames(const FAnimNode_RigidBodyWithControl* const Node, TMap<FName, FPhysicsControlLimbBones> AllLimbBones, const FReferenceSkeleton& RefSkeleton, UPhysicsAsset* const PhysicsAsset, TSet<FName>& BodyModifierNames, TSet<FName>& ControlNames, FPhysicsControlNameRecords& NameRecords)
+//======================================================================================================================
+void CollectOperatorNames(
+	const FAnimNode_RigidBodyWithControl*            Node, 
+	const FPhysicsControlCharacterSetupData&         CharacterSetupData,
+	const FPhysicsControlAndBodyModifierCreationDatas& AdditionalControlsAndBodyModifiers,
+	TMap<FName, FPhysicsControlLimbBones>            AllLimbBones,
+	const FReferenceSkeleton&                        RefSkeleton, 
+	UPhysicsAsset*                                   PhysicsAsset, 
+	TSet<FName>&                                     BodyModifierNames, 
+	TSet<FName>&                                     ControlNames, 
+	FPhysicsControlNameRecords&                      NameRecords)
 {
-	auto CollectControlName = [&ControlNames](const FName ParentBoneName, const FName ChildBoneName, const FPhysicsControlData& Data) -> FName { const FName Name = FindUniqueControlName(ParentBoneName, ChildBoneName, ControlNames); ControlNames.Add(Name); return Name; };
-	auto CollectBodyModifierName = [&BodyModifierNames](const FName BoneName, const FPhysicsControlModifierData& Data) -> FName { const FName Name = FindUniqueBodyModifierName(BoneName, BodyModifierNames); BodyModifierNames.Add(Name); return Name; };
-
-	ForEachPotentialOperator(Node, AllLimbBones, RefSkeleton, PhysicsAsset, NameRecords, CollectControlName, CollectBodyModifierName);
+	auto CollectControlName = [&ControlNames](const FName ParentBoneName, const FName ChildBoneName, const FPhysicsControlData& Data) -> FName 
+		{
+			const FName Name = GetUniqueControlName(ParentBoneName, ChildBoneName, ControlNames, TEXT(""));
+			ControlNames.Add(Name); 
+			return Name; 
+		};
+	auto CollectBodyModifierName = [&BodyModifierNames](const FName BoneName, const FPhysicsControlModifierData& Data) -> FName { 
+		const FName Name = GetUniqueBodyModifierName(BoneName, BodyModifierNames, TEXT(""));
+		BodyModifierNames.Add(Name); 
+		return Name; };
+
+	ForEachPotentialOperator(
+		CharacterSetupData, AdditionalControlsAndBodyModifiers, AllLimbBones, RefSkeleton,
+		PhysicsAsset, NameRecords, CollectControlName, CollectBodyModifierName);
 }
 
-void CreateOperatorsForNode(FAnimNode_RigidBodyWithControl* const Node, TMap<FName, FPhysicsControlLimbBones> AllLimbBones, const FReferenceSkeleton& RefSkeleton, UPhysicsAsset* const PhysicsAsset, FPhysicsControlNameRecords& NameRecords)
+//======================================================================================================================
+void CreateOperatorsForNode(
+	FAnimNode_RigidBodyWithControl*                  Node,
+	const FPhysicsControlCharacterSetupData&         CharacterSetupData,
+	const FPhysicsControlAndBodyModifierCreationDatas& AdditionalControlsAndBodyModifiers,
+	TMap<FName, FPhysicsControlLimbBones>            AllLimbBones,
+	const FReferenceSkeleton&                        RefSkeleton, 
+	UPhysicsAsset*                                   PhysicsAsset, 
+	FPhysicsControlNameRecords&                      NameRecords)
 {
-	auto CreateControl = [Node](const FName ParentBoneName, const FName ChildBoneName, const FPhysicsControlData& Data) -> FName { return Node->CreateControl(ParentBoneName, ChildBoneName, Data); };
-	auto CreateBodyModifier = [Node](const FName BoneName, const FPhysicsControlModifierData& Data) -> FName { return Node->CreateBodyModifier(BoneName, Data); };
-
-	ForEachPotentialOperator(Node, AllLimbBones, RefSkeleton, PhysicsAsset, NameRecords, CreateControl, CreateBodyModifier);
+	auto CreateControl = [Node](const FName ParentBoneName, const FName ChildBoneName, const FPhysicsControlData& Data) -> FName 
+		{ 
+			return Node->CreateControl(ParentBoneName, ChildBoneName, Data); 
+		};
+	auto CreateBodyModifier = [Node](const FName BoneName, const FPhysicsControlModifierData& Data) -> FName 
+		{ 
+			return Node->CreateBodyModifier(BoneName, Data); 
+		};
+
+	ForEachPotentialOperator(
+		CharacterSetupData, AdditionalControlsAndBodyModifiers, AllLimbBones, RefSkeleton,
+		PhysicsAsset, NameRecords, CreateControl, CreateBodyModifier);
 }
 
-void CreateAdditionalSets(const FPhysicsControlSetUpdates& AdditionalSets, const TSet<FName>& BodyModifierNames, const TSet<FName>& ControlNames, FPhysicsControlNameRecords& NameRecords)
+//======================================================================================================================
+void CreateAdditionalSets(
+	const FPhysicsControlSetUpdates& AdditionalSets, 
+	const TSet<FName>&               BodyModifierNames, 
+	const TSet<FName>&               ControlNames, 
+	FPhysicsControlNameRecords&      NameRecords)
 {
 	CreateAdditionalSets_Implementation(AdditionalSets, BodyModifierNames, ControlNames, NameRecords);
 }
 
-void CreateAdditionalSets(const FPhysicsControlSetUpdates& AdditionalSets, const TMap<FName, FRigidBodyModifierRecord>& BodyModifierNames, const TMap<FName, FRigidBodyControlRecord>& ControlNames, FPhysicsControlNameRecords& NameRecords)
+//======================================================================================================================
+void CreateAdditionalSets(
+	const FPhysicsControlSetUpdates&             AdditionalSets, 
+	const TMap<FName, FRigidBodyModifierRecord>& BodyModifierRecords, 
+	const TMap<FName, FRigidBodyControlRecord>&  Controls, 
+	FPhysicsControlNameRecords&                  NameRecords)
 {
-	CreateAdditionalSets_Implementation(AdditionalSets, BodyModifierNames, ControlNames, NameRecords);
+	CreateAdditionalSets_Implementation(AdditionalSets, BodyModifierRecords, Controls, NameRecords);
+}
+
+//======================================================================================================================
+void CreateAdditionalSets(
+	const FPhysicsControlSetUpdates&          AdditionalSets,
+	const TMap<FName, FPhysicsBodyModifierRecord>&  BodyModifierRecords,
+	const TMap<FName, FPhysicsControlRecord>& Controls,
+	FPhysicsControlNameRecords&               NameRecords)
+{
+	CreateAdditionalSets_Implementation(AdditionalSets, BodyModifierRecords, Controls, NameRecords);
 }
+
+} // namespace PhysicsControlComponent
+} // namespace UE
diff --git a/Engine/Plugins/Experimental/PhysicsControl/Source/PhysicsControl/Private/PhysicsControlProfileAsset.cpp b/Engine/Plugins/Experimental/PhysicsControl/Source/PhysicsControl/Private/PhysicsControlProfileAsset.cpp
--- a/Engine/Plugins/Experimental/PhysicsControl/Source/PhysicsControl/Private/PhysicsControlProfileAsset.cpp
+++ b/Engine/Plugins/Experimental/PhysicsControl/Source/PhysicsControl/Private/PhysicsControlProfileAsset.cpp
@@ -3,23 +3,139 @@
 #include "PhysicsControlProfileAsset.h"
 #include "Engine/SkeletalMesh.h"
 
+//#ifdef WITH_EDITOR
+//#include "Editor.h"
+//#endif // WITH_EDITOR
+
+
 //======================================================================================================================
 UPhysicsControlProfileAsset::UPhysicsControlProfileAsset()
 {
+	// This needs to be explored further - a possible way to hook into the BP compilation. However, we would also
+	// need to make sure we don't compile (and get marked as dirty) if there are no changes.
+//#ifdef WITH_EDITOR
+//	GEditor->OnBlueprintCompiled().AddUObject(this, &UPhysicsControlProfileAsset::Compile);
+//#endif
+}
+
+#if WITH_EDITOR
+//======================================================================================================================
+void UPhysicsControlProfileAsset::ShowCompiledData() const
+{
+	UE_LOG(LogTemp, Log, TEXT("Character setup data:"));
+	for (const FPhysicsControlLimbSetupData& LimbSetupData : CharacterSetupData.LimbSetupData)
+	{
+		UE_LOG(LogTemp, Log, TEXT("Limb %s"), *LimbSetupData.LimbName.ToString());
+		UE_LOG(LogTemp, Log, TEXT("  Start bone %s"), *LimbSetupData.StartBone.ToString());
+		UE_LOG(LogTemp, Log, TEXT("  Include parent bone %d"), LimbSetupData.bIncludeParentBone);
+		UE_LOG(LogTemp, Log, TEXT("  Create world space controls %d"), LimbSetupData.bCreateWorldSpaceControls);
+		UE_LOG(LogTemp, Log, TEXT("  Create parent space controls %d"), LimbSetupData.bCreateParentSpaceControls);
+		UE_LOG(LogTemp, Log, TEXT("  Create body modifiers %d"), LimbSetupData.bCreateBodyModifiers);
+	}
+	UE_LOG(LogTemp, Log, TEXT("Additional controls and modifiers:"));
+	UE_LOG(LogTemp, Log, TEXT("  Additional controls:"));
+	for (TMap<FName, FPhysicsControlCreationData>::ElementType ControlPair : AdditionalControlsAndModifiers.Controls)
+	{
+		UE_LOG(LogTemp, Log, TEXT("    %s:"), *ControlPair.Key.ToString());
+		UE_LOG(LogTemp, Log, TEXT("      Parent bone %s:"), *ControlPair.Value.Control.ParentBoneName.ToString());
+		UE_LOG(LogTemp, Log, TEXT("      Child bone %s:"), *ControlPair.Value.Control.ParentBoneName.ToString());
+	}
+	UE_LOG(LogTemp, Log, TEXT("  Additional modifiers:"));
+	for (TMap<FName, FPhysicsBodyModifierCreationData>::ElementType ModifierPair : AdditionalControlsAndModifiers.Modifiers)
+	{
+		UE_LOG(LogTemp, Log, TEXT("    %s:"), *ModifierPair.Key.ToString());
+		UE_LOG(LogTemp, Log, TEXT("      Bone %s:"), *ModifierPair.Value.Modifier.BoneName.ToString());
+	}
+	UE_LOG(LogTemp, Log, TEXT("Profiles:"));
+	for (TMap<FName, FPhysicsControlControlAndModifierUpdates>::ElementType ProfilePair : Profiles)
+	{
+		UE_LOG(LogTemp, Log, TEXT("  %s:"), *ProfilePair.Key.ToString());
+	}
+
+}
+
+//======================================================================================================================
+void UPhysicsControlProfileAsset::Compile()
+{
+	CharacterSetupData = GetCharacterSetupData();
+	AdditionalControlsAndModifiers = GetAdditionalControlsAndModifiers();
+	AdditionalSets = GetAdditionalSets();
+	InitialControlAndModifierUpdates = GetInitialControlAndModifierUpdates();
+	Profiles = GetProfiles();
+
+	Modify();
 }
 
 //======================================================================================================================
-void UPhysicsControlProfileAsset::Log()
+FPhysicsControlCharacterSetupData UPhysicsControlProfileAsset::GetCharacterSetupData() const
 {
-	UE_LOG(LogTemp, Log, TEXT("Number of initial controls = %d"), InitialControls.Num());
-	UE_LOG(LogTemp, Log, TEXT("Number of initial body modifiers = %d"), InitialBodyModifiers.Num());
+	FPhysicsControlCharacterSetupData CompiledCharacterSetupData;
+	if (ParentAsset.LoadSynchronous())
+	{
+		CompiledCharacterSetupData = ParentAsset->GetCharacterSetupData();
+	}
+	CompiledCharacterSetupData += MyCharacterSetupData;
+	return CompiledCharacterSetupData;
 }
 
 //======================================================================================================================
-void UPhysicsControlProfileAsset::MakeControlsAndModifiersFromCharacterSetupData()
+FPhysicsControlAndBodyModifierCreationDatas UPhysicsControlProfileAsset::GetAdditionalControlsAndModifiers() const
 {
+	FPhysicsControlAndBodyModifierCreationDatas CompiledAdditionalControlsAndModifiers;
+	if (ParentAsset.LoadSynchronous())
+	{
+		CompiledAdditionalControlsAndModifiers = ParentAsset->GetAdditionalControlsAndModifiers();
+	}
+	// this will overwrite duplicates with our value
+	CompiledAdditionalControlsAndModifiers += MyAdditionalControlsAndModifiers; 
+	return CompiledAdditionalControlsAndModifiers;
 }
 
+//======================================================================================================================
+FPhysicsControlSetUpdates UPhysicsControlProfileAsset::GetAdditionalSets() const
+{
+	FPhysicsControlSetUpdates CompiledAdditionalSets;
+	if (ParentAsset.LoadSynchronous())
+	{
+		CompiledAdditionalSets = ParentAsset->GetAdditionalSets();
+	}
+	CompiledAdditionalSets += MyAdditionalSets;
+	return CompiledAdditionalSets;
+}
+
+//======================================================================================================================
+TArray<FPhysicsControlControlAndModifierUpdates> UPhysicsControlProfileAsset::GetInitialControlAndModifierUpdates() const
+{
+	TArray<FPhysicsControlControlAndModifierUpdates> CompiledInitialControlAndModifierUpdates;
+	if (ParentAsset.LoadSynchronous())
+	{
+		CompiledInitialControlAndModifierUpdates = ParentAsset->GetInitialControlAndModifierUpdates();
+	}
+	CompiledInitialControlAndModifierUpdates.Append(MyInitialControlAndModifierUpdates);
+	return CompiledInitialControlAndModifierUpdates;
+}
+
+//======================================================================================================================
+TMap<FName, FPhysicsControlControlAndModifierUpdates> UPhysicsControlProfileAsset::GetProfiles() const
+{
+	TMap<FName, FPhysicsControlControlAndModifierUpdates> CompiledProfiles;
+	if (ParentAsset.LoadSynchronous())
+	{
+		CompiledProfiles = ParentAsset->GetProfiles();
+	}
+	for (const TSoftObjectPtr<UPhysicsControlProfileAsset>& AdditionalProfileAsset : AdditionalProfileAssets)
+	{
+		if (AdditionalProfileAsset.LoadSynchronous())
+		{
+			CompiledProfiles.Append(AdditionalProfileAsset->GetProfiles());
+		}
+	}
+	CompiledProfiles.Append(MyProfiles); // this will overwrite duplicates with our value
+	return CompiledProfiles;
+}
+
+#endif
+
 #if WITH_EDITOR
 //======================================================================================================================
 const FName UPhysicsControlProfileAsset::GetPreviewMeshPropertyName()
@@ -31,11 +147,17 @@ const FName UPhysicsControlProfileAsset::GetPreviewMeshPropertyName()
 //======================================================================================================================
 void UPhysicsControlProfileAsset::SetPreviewMesh(USkeletalMesh* PreviewMesh, bool bMarkAsDirty)
 {
+#if WITH_EDITOR
 	PreviewSkeletalMesh = PreviewMesh;
+#endif
 }
 
 //======================================================================================================================
 USkeletalMesh* UPhysicsControlProfileAsset::GetPreviewMesh() const
 {
+#if WITH_EDITOR
 	return PreviewSkeletalMesh.LoadSynchronous();
+#else
+	return nullptr;
+#endif
 }
diff --git a/Engine/Plugins/Experimental/PhysicsControl/Source/PhysicsControl/Private/PhysicsControlRecord.cpp b/Engine/Plugins/Experimental/PhysicsControl/Source/PhysicsControl/Private/PhysicsControlRecord.cpp
--- a/Engine/Plugins/Experimental/PhysicsControl/Source/PhysicsControl/Private/PhysicsControlRecord.cpp
+++ b/Engine/Plugins/Experimental/PhysicsControl/Source/PhysicsControl/Private/PhysicsControlRecord.cpp
@@ -25,8 +25,8 @@ FVector FPhysicsControlRecord::GetControlPoint() const
 		return PhysicsControl.ControlData.CustomControlPoint;
 	}
 
-	FBodyInstance* ChildBodyInstance = UE::PhysicsControlComponent::GetBodyInstance(
-		PhysicsControl.ChildMeshComponent.Get(), PhysicsControl.ChildBoneName);
+	FBodyInstance* ChildBodyInstance = UE::PhysicsControl::GetBodyInstance(
+		ChildMeshComponent.Get(), PhysicsControl.ChildBoneName);
 
 	return ChildBodyInstance ? ChildBodyInstance->GetMassSpaceLocal().GetTranslation() : FVector::ZeroVector;
 }
@@ -40,22 +40,22 @@ bool FPhysicsControlRecord::InitConstraint(UObject* ConstraintDebugOwner, FName
 	}
 	check(ConstraintInstance);
 
-	FBodyInstance* ParentBody = UE::PhysicsControlComponent::GetBodyInstance(
-		PhysicsControl.ParentMeshComponent.Get(), PhysicsControl.ParentBoneName);
-	FBodyInstance* ChildBody = UE::PhysicsControlComponent::GetBodyInstance(
-		PhysicsControl.ChildMeshComponent.Get(), PhysicsControl.ChildBoneName);
+	FBodyInstance* ParentBody = UE::PhysicsControl::GetBodyInstance(
+		ParentMeshComponent.Get(), PhysicsControl.ParentBoneName);
+	FBodyInstance* ChildBody = UE::PhysicsControl::GetBodyInstance(
+		ChildMeshComponent.Get(), PhysicsControl.ChildBoneName);
 
-	if (PhysicsControl.ParentMeshComponent.IsValid() && !PhysicsControl.ParentBoneName.IsNone() && !ParentBody)
+	if (ParentMeshComponent.IsValid() && !PhysicsControl.ParentBoneName.IsNone() && !ParentBody)
 	{
-		UE_LOG(LogPhysicsControlComponent, Warning,
+		UE_LOG(LogPhysicsControl, Warning,
 			TEXT("Failed to find expected parent body %s when making constraint for control %s"),
 			*PhysicsControl.ParentBoneName.ToString(),
 			*ControlName.ToString());
 		return false;
 	}
-	if (PhysicsControl.ChildMeshComponent.IsValid() && !PhysicsControl.ChildBoneName.IsNone() && !ChildBody)
+	if (ChildMeshComponent.IsValid() && !PhysicsControl.ChildBoneName.IsNone() && !ChildBody)
 	{
-		UE_LOG(LogPhysicsControlComponent, Warning,
+		UE_LOG(LogPhysicsControl, Warning,
 			TEXT("Failed to find expected child body %s when making constraint for control %s"),
 			*PhysicsControl.ChildBoneName.ToString(),
 			*ControlName.ToString());
diff --git a/Engine/Plugins/Experimental/PhysicsControl/Source/PhysicsControl/Private/RigidBodyControlData.cpp b/Engine/Plugins/Experimental/PhysicsControl/Source/PhysicsControl/Private/RigidBodyControlData.cpp
--- a/Engine/Plugins/Experimental/PhysicsControl/Source/PhysicsControl/Private/RigidBodyControlData.cpp
+++ b/Engine/Plugins/Experimental/PhysicsControl/Source/PhysicsControl/Private/RigidBodyControlData.cpp
@@ -5,12 +5,12 @@
 
 //======================================================================================================================
 FRigidBodyControlRecord::FRigidBodyControlRecord(
-	const FRigidBodyControl& InControl, ImmediatePhysics::FJointHandle* InJointHandle)
+	const FPhysicsControl& InControl, ImmediatePhysics::FJointHandle* InJointHandle)
 	: Control(InControl)
-	, JointHandle(InJointHandle)
 	, ControlData(Control.ControlData)
 	, ChildBodyIndex(-1)
 	, ParentBodyIndex(-1)
+	, JointHandle(InJointHandle)
 {
 }
 
@@ -18,6 +18,7 @@ FRigidBodyControlRecord::FRigidBodyControlRecord(
 void FRigidBodyControlRecord::ResetCurrent(bool bResetTarget)
 {
 	ControlData = Control.ControlData;
+	ControlMultiplier = Control.ControlMultiplier;
 
 	if (bResetTarget)
 	{
@@ -37,10 +38,10 @@ FVector FRigidBodyControlRecord::GetControlPoint(const ImmediatePhysics::FActorH
 
 //======================================================================================================================
 FRigidBodyModifierRecord::FRigidBodyModifierRecord(
-	const FRigidBodyModifier& InModifier, ImmediatePhysics::FActorHandle* InActorHandle)
+	const FPhysicsBodyModifier& InModifier, ImmediatePhysics::FActorHandle* InActorHandle)
 	: Modifier(InModifier)
-	, ActorHandle(InActorHandle)
 	, ModifierData(Modifier.ModifierData)
+	, ActorHandle(InActorHandle)
 {
 }
 
diff --git a/Engine/Plugins/Experimental/PhysicsControl/Source/PhysicsControl/Private/RigidBodyWithControl.cpp b/Engine/Plugins/Experimental/PhysicsControl/Source/PhysicsControl/Private/RigidBodyWithControl.cpp
--- a/Engine/Plugins/Experimental/PhysicsControl/Source/PhysicsControl/Private/RigidBodyWithControl.cpp
+++ b/Engine/Plugins/Experimental/PhysicsControl/Source/PhysicsControl/Private/RigidBodyWithControl.cpp
@@ -2,6 +2,11 @@
 
 #include "AnimNode_RigidBodyWithControl.h"
 
+#include "PhysicsControlOperatorNameGeneration.h"
+#include "PhysicsControlLog.h"
+#include "PhysicsControlProfileAsset.h"
+#include "PhysicsControlComponentHelpers.h"
+
 #include "Animation/AnimInstance.h"
 #include "Animation/AnimInstanceProxy.h"
 #include "Chaos/PBDJointConstraints.h"
@@ -13,11 +18,9 @@
 #include "PhysicsEngine/PhysicsConstraintTemplate.h"
 #include "PhysicsEngine/PhysicsSettings.h"
 #include "Chaos/ChaosConstraintSettings.h"
-#include "PhysicsControlLog.h"
 #include "Logging/StructuredLog.h"
-#include "PhysicsControlOperatorNameGeneration.h"
 
-//UE_DISABLE_OPTIMIZATION;
+//UE_DISABLE_OPTIMIZATION
 
 DECLARE_CYCLE_STAT(TEXT("RigidBodyNodeWithControl_InitControlsAndBodyModifiers"), STAT_RigidBodyNodeWithControl_InitControlsAndBodyModifiers, STATGROUP_Anim);
 DECLARE_CYCLE_STAT(TEXT("RigidBodyNodeWithControl_LogControlsModifiersAndSets"), STAT_RigidBodyNodeWithControl_LogControlsModifiersAndSets, STATGROUP_Anim);
@@ -45,6 +48,16 @@ static void SetConstraintEnabled(ImmediatePhysics::FJointHandle* const JointHand
 		if (ImmediatePhysics::FJointHandle::FChaosConstraintHandle* ConstraintHandle = JointHandle->GetConstraint())
 		{
 			ConstraintHandle->SetConstraintEnabled(bIsEnabled);
+
+			if (!bIsEnabled)
+			{
+				// The call above disables the constraint, but if any actor in the simulation is
+				// flagged as dirty, it gets re-enabled! See UE-204006 TODO remove this workaround
+				// when the bug has been fixed.
+				ConstraintHandle->SetDriveParams(
+					Chaos::FVec3::ZeroVector, Chaos::FVec3::ZeroVector, Chaos::FVec3::ZeroVector, 
+					Chaos::FVec3::ZeroVector, Chaos::FVec3::ZeroVector, Chaos::FVec3::ZeroVector);
+			}
 		}
 	}
 }
@@ -59,7 +72,7 @@ static void SetRecordParameters(
 	}
 	else
 	{
-		UE_LOG(LogRigidBodyWithControl, Warning,
+		UE_LOG(LogPhysicsControl, Warning,
 			TEXT("SetRecordParameters: Failed to find modifier with name %s"), *Name.ToString());
 	}
 }
@@ -74,7 +87,22 @@ static void SetRecordParameters(
 	}
 	else
 	{
-		UE_LOG(LogRigidBodyWithControl, Warning,
+		UE_LOG(LogPhysicsControl, Warning,
+			TEXT("SetRecordParameters: Failed to find control with name %s"), *Name.ToString());
+	}
+}
+
+//======================================================================================================================
+static void SetRecordParameters(
+	const FName Name, const FPhysicsControlSparseMultiplier& Data, TMap<FName, FRigidBodyControlRecord>& Records)
+{
+	if (FRigidBodyControlRecord* const RecordSearchResult = Records.Find(Name))
+	{
+		RecordSearchResult->ControlMultiplier.UpdateFromSparseData(Data);
+	}
+	else
+	{
+		UE_LOG(LogPhysicsControl, Warning,
 			TEXT("SetRecordParameters: Failed to find control with name %s"), *Name.ToString());
 	}
 }
@@ -113,28 +141,20 @@ const FTransform FAnimNode_RigidBodyWithControl::GetBodyTransform(const int32 Bo
 	return FTransform::Identity;
 }
 
-//======================================================================================================================
-const FTransform FAnimNode_RigidBodyWithControl::GetWorldSpaceControlRootTransform() const
-{
-	if (WorldSpaceControlActorHandle)
-	{
-		return WorldSpaceControlActorHandle->GetWorldTransform();
-	}
-	return FTransform::Identity;
-}
-
 //======================================================================================================================
 void FAnimNode_RigidBodyWithControl::CreateWorldSpaceControlRootBody(UPhysicsAsset* const PhysicsAsset)
 {
+	// Note that this never moves - it's just defined as being an actor at the root of the world
 	WorldSpaceControlActorHandle = PhysicsSimulation->CreateActor(
-		ImmediatePhysics::EActorType::KinematicActor, nullptr, FTransform());
+		ImmediatePhysics::EActorType::StaticActor, nullptr, FTransform());
 	if (WorldSpaceControlActorHandle)
 	{
 		WorldSpaceControlActorHandle->SetName(FName("WorldSpaceControlHandle"));
+		PhysicsSimulation->SetHasCollision(WorldSpaceControlActorHandle, false);
 	}
 	else
 	{
-		UE_LOG(LogRigidBodyWithControl, Error, TEXT("Failed to create world space control root actor"));
+		UE_LOG(LogPhysicsControl, Error, TEXT("Failed to create world space control root actor"));
 	}
 }
 
@@ -230,16 +250,17 @@ void FAnimNode_RigidBodyWithControl::UpdateBodyIndicesInControlRecord(FRigidBody
 FName FAnimNode_RigidBodyWithControl::CreateControl(
 	const FName ParentBoneName, const FName ChildBoneName, const FPhysicsControlData& ControlData)
 {
-	FRigidBodyControl Control;
+	FPhysicsControl Control;
 	Control.ParentBoneName = ParentBoneName;
 	Control.ChildBoneName = ChildBoneName;
 	Control.ControlData = ControlData;
 
 	ImmediatePhysics::FJointHandle* JointHandle = nullptr;
-	ImmediatePhysics::FActorHandle* ParentBodyHandle;
+	ImmediatePhysics::FActorHandle* ParentBodyHandle = nullptr;
 
 	if (ParentBoneName.IsNone())
 	{
+		// A parent actor is needed. Without it, the constraint doesn't work (though there's no error)
 		ParentBodyHandle = WorldSpaceControlActorHandle;
 	}
 	else
@@ -251,7 +272,7 @@ FName FAnimNode_RigidBodyWithControl::CreateControl(
 
 	if (!JointHandle)
 	{
-		UE_LOG(LogRigidBodyWithControl, Warning,
+		UE_LOG(LogPhysicsControl, Warning,
 			TEXT("Unable to create world space control constraint for bone %s"), *ChildBoneName.ToString());
 		return FName();
 	}
@@ -275,15 +296,12 @@ FName FAnimNode_RigidBodyWithControl::GetBodyFromBoneName(const FName BoneName)
 //======================================================================================================================
 FName FAnimNode_RigidBodyWithControl::GetUniqueBodyModifierName(const FName BoneName) const
 {
-	TSet<FName> Keys;
-	ModifierRecords.GetKeys(Keys);
-	const FName UniqueName = FindUniqueBodyModifierName(GetBodyFromBoneName(BoneName), Keys);
+	const FName UniqueName = UE::PhysicsControl::GetUniqueBodyModifierName(
+		GetBodyFromBoneName(BoneName), ModifierRecords, TEXT(""));
 
 	if (UniqueName.IsNone())
 	{
-		UE_LOG(LogRigidBodyWithControl, Warning,
-			TEXT("Unable to find a suitable Body Modifier name - the limit of MaxNumControlsOrModifiersPerName (%d) has been exceeded"),
-			MaxNumControlsOrModifiersPerName);
+		UE_LOG(LogPhysicsControl, Warning, TEXT("Unable to find a suitable Body Modifier name"));
 	}
 
 	return UniqueName;
@@ -292,15 +310,12 @@ FName FAnimNode_RigidBodyWithControl::GetUniqueBodyModifierName(const FName Bone
 //======================================================================================================================
 FName FAnimNode_RigidBodyWithControl::GetUniqueControlName(const FName ParentBoneName, const FName ChildBoneName) const
 {
-	TSet<FName> Keys;
-	ControlRecords.GetKeys(Keys);
-	const FName UniqueName = FindUniqueControlName(GetBodyFromBoneName(ParentBoneName), GetBodyFromBoneName(ChildBoneName), Keys);
+	const FName UniqueName = UE::PhysicsControl::GetUniqueControlName(
+		GetBodyFromBoneName(ParentBoneName), GetBodyFromBoneName(ChildBoneName), ControlRecords, TEXT(""));
 
 	if (UniqueName.IsNone())
 	{
-		UE_LOG(LogRigidBodyWithControl, Warning,
-			TEXT("Unable to find a suitable Control name - the limit of MaxNumControlsOrModifiersPerName (%d) has been exceeded"),
-			MaxNumControlsOrModifiersPerName);
+		UE_LOG(LogPhysicsControl, Warning, TEXT("Unable to find a suitable Control name"));
 	}
 
 	return UniqueName;
@@ -310,15 +325,11 @@ FName FAnimNode_RigidBodyWithControl::GetUniqueControlName(const FName ParentBon
 FName FAnimNode_RigidBodyWithControl::CreateBodyModifier(FName BoneName, const FPhysicsControlModifierData& ModifierData)
 {
 	FName Name;
-
 	ImmediatePhysics::FActorHandle* const ActorHandle = FindBodyFromBoneName(BoneName);
 	if (ActorHandle)
 	{
 		Name = GetUniqueBodyModifierName(BoneName);
-
-		FRigidBodyModifier BodyModifier;
-		BodyModifier.BoneName = BoneName;
-		BodyModifier.ModifierData = ModifierData;
+		FPhysicsBodyModifier BodyModifier(BoneName, ModifierData);
 		FRigidBodyModifierRecord& Modifier = ModifierRecords.Add(Name, FRigidBodyModifierRecord(BodyModifier, ActorHandle));
 	}
 
@@ -332,11 +343,30 @@ void FAnimNode_RigidBodyWithControl::InitControlsAndBodyModifiers(const FReferen
 
 	check(ControlRecords.IsEmpty()); // Controls should not exist when this function is called.
 
+	FPhysicsControlCharacterSetupData SetupData;
+	if (IsValid(PhysicsControlProfileAsset))
+	{
+		SetupData = PhysicsControlProfileAsset->CharacterSetupData;
+	}
+	if (bEnableCharacterSetupData)
+	{
+		SetupData += CharacterSetupData;
+	}
+
 	// These functions will create the base set of controls and modifiers from SetupData
 	TMap<FName, FPhysicsControlLimbBones> AllLimbBones =
-		GetLimbBones(SetupData.LimbSetupData, RefSkeleton, GetPhysicsAsset());
+		UE::PhysicsControl::GetLimbBones(SetupData.LimbSetupData, RefSkeleton, GetPhysicsAsset());
+
+	FPhysicsControlAndBodyModifierCreationDatas AdditionalOperators;
+	if (IsValid(PhysicsControlProfileAsset))
+	{
+		AdditionalOperators = PhysicsControlProfileAsset->AdditionalControlsAndModifiers;
+	}
+	AdditionalOperators += AdditionalControlsAndBodyModifiers;
 
-	CreateOperatorsForNode(this, AllLimbBones, RefSkeleton, GetPhysicsAsset(), NameRecords);
+	// An "operator" is a control or a body modifier. This will also add them to sets etc.
+	UE::PhysicsControl::CreateOperatorsForNode(
+		this, SetupData, AdditionalControlsAndBodyModifiers, AllLimbBones, RefSkeleton, GetPhysicsAsset(), NameRecords);
 
 	for (TMap<FName, FRigidBodyControlRecord>::ElementType& NameRecordPair : ControlRecords)
 	{
@@ -345,11 +375,17 @@ void FAnimNode_RigidBodyWithControl::InitControlsAndBodyModifiers(const FReferen
 	}
 
 	// Create any additional sets that have been requested
-	CreateAdditionalSets(AdditionalSets, ModifierRecords, ControlRecords, NameRecords);
+	if (IsValid(PhysicsControlProfileAsset))
+	{
+		UE::PhysicsControl::CreateAdditionalSets(
+			PhysicsControlProfileAsset->AdditionalSets, ModifierRecords, ControlRecords, NameRecords);
+	}
+	UE::PhysicsControl::CreateAdditionalSets(AdditionalSets, ModifierRecords, ControlRecords, NameRecords);
 
 	// Apply control and modifier parameters on a single-use basis
 	ApplyControlAndBodyModifierDatas(
 		InitialControlAndBodyModifierUpdates.ControlParameters,
+		InitialControlAndBodyModifierUpdates.ControlMultiplierParameters,
 		InitialControlAndBodyModifierUpdates.ModifierParameters);
 
 	// Tell the poor user what we've done
@@ -360,11 +396,9 @@ void FAnimNode_RigidBodyWithControl::InitControlsAndBodyModifiers(const FReferen
 void FAnimNode_RigidBodyWithControl::DestroyControlsAndBodyModifiers()
 {
 	// TODO - make sure this fn is complete
-
 	ControlRecords.Reset();
 	ModifierRecords.Reset();
 	NameRecords.Reset();
-
 	bHaveSetupControls = false;
 }
 
@@ -376,56 +410,56 @@ void FAnimNode_RigidBodyWithControl::LogControlsModifiersAndSets()
 
 #define RBWC_LOG_LEVEL Display
 
-	UE_LOG(LogRigidBodyWithControl, RBWC_LOG_LEVEL, TEXT("Controls:"));
+	UE_LOG(LogPhysicsControl, RBWC_LOG_LEVEL, TEXT("Controls:"));
 	for (TMap<FName, FRigidBodyControlRecord>::ElementType& NameRecordPair : ControlRecords)
 	{
-		UE_LOG(LogRigidBodyWithControl, RBWC_LOG_LEVEL, TEXT("  %s:"), *NameRecordPair.Key.ToString());
+		UE_LOG(LogPhysicsControl, RBWC_LOG_LEVEL, TEXT("  %s:"), *NameRecordPair.Key.ToString());
 		const FRigidBodyControlRecord& Record = NameRecordPair.Value;
-		UE_LOG(LogRigidBodyWithControl, RBWC_LOG_LEVEL, TEXT("    Parent bone: %s Child bone: %s"), 
+		UE_LOG(LogPhysicsControl, RBWC_LOG_LEVEL, TEXT("    Parent bone: %s Child bone: %s"), 
 			*Record.Control.ParentBoneName.ToString(), *Record.Control.ChildBoneName.ToString());
-		UE_LOG(LogRigidBodyWithControl, RBWC_LOG_LEVEL, TEXT("    Enabled %d"),
+		UE_LOG(LogPhysicsControl, RBWC_LOG_LEVEL, TEXT("    Enabled %d"),
 			Record.IsEnabled());
-		UE_LOG(LogRigidBodyWithControl, RBWC_LOG_LEVEL, TEXT("    Linear: Strength %f DampingRatio %f ExtraDamping %f"),
+		UE_LOG(LogPhysicsControl, RBWC_LOG_LEVEL, TEXT("    Linear: Strength %f DampingRatio %f ExtraDamping %f"),
 			Record.Control.ControlData.LinearStrength, 
 			Record.Control.ControlData.LinearDampingRatio, 
 			Record.Control.ControlData.LinearExtraDamping);
-		UE_LOG(LogRigidBodyWithControl, RBWC_LOG_LEVEL, TEXT("    Angular: Strength %f DampingRatio %f ExtraDamping %f"),
+		UE_LOG(LogPhysicsControl, RBWC_LOG_LEVEL, TEXT("    Angular: Strength %f DampingRatio %f ExtraDamping %f"),
 			Record.Control.ControlData.AngularStrength, 
 			Record.Control.ControlData.AngularDampingRatio, 
 			Record.Control.ControlData.AngularExtraDamping);
 	}
 
-	UE_LOG(LogRigidBodyWithControl, RBWC_LOG_LEVEL, TEXT("Body Modifiers:"));
+	UE_LOG(LogPhysicsControl, RBWC_LOG_LEVEL, TEXT("Body Modifiers:"));
 	for (TMap<FName, FRigidBodyModifierRecord>::ElementType& NameRecordPair : ModifierRecords)
 	{
-		UE_LOG(LogRigidBodyWithControl, RBWC_LOG_LEVEL, TEXT("  %s:"), *NameRecordPair.Key.ToString());
+		UE_LOG(LogPhysicsControl, RBWC_LOG_LEVEL, TEXT("  %s:"), *NameRecordPair.Key.ToString());
 		const FRigidBodyModifierRecord& Record = NameRecordPair.Value;
-		UE_LOG(LogRigidBodyWithControl, RBWC_LOG_LEVEL, TEXT("    Bone: %s Body: %s"),
+		UE_LOG(LogPhysicsControl, RBWC_LOG_LEVEL, TEXT("    Bone: %s Body: %s"),
 			*Record.Modifier.BoneName.ToString(), *GetBodyFromBoneName(Record.Modifier.BoneName).ToString());
-		UE_LOG(LogRigidBodyWithControl, RBWC_LOG_LEVEL, TEXT("    Movement: %s GravityMultiplier: %f"),
+		UE_LOG(LogPhysicsControl, RBWC_LOG_LEVEL, TEXT("    Movement: %s GravityMultiplier: %f"),
 			*GetPhysicsMovementTypeName(Record.Modifier.ModifierData.MovementType).ToString(),
 			Record.Modifier.ModifierData.GravityMultiplier);
 	}
 
-	UE_LOG(LogRigidBodyWithControl, RBWC_LOG_LEVEL, TEXT("Control sets:"));
+	UE_LOG(LogPhysicsControl, RBWC_LOG_LEVEL, TEXT("Control sets:"));
 	for (TMap<FName, TArray<FName>>::ElementType& Pair : NameRecords.ControlSets)
 	{
-		UE_LOG(LogRigidBodyWithControl, RBWC_LOG_LEVEL, TEXT("  %s:"), *Pair.Key.ToString());
+		UE_LOG(LogPhysicsControl, RBWC_LOG_LEVEL, TEXT("  %s:"), *Pair.Key.ToString());
 		const TArray<FName>& Names = Pair.Value;
 		for (FName Name : Names)
 		{
-			UE_LOG(LogRigidBodyWithControl, RBWC_LOG_LEVEL, TEXT("    %s:"), *Name.ToString());
+			UE_LOG(LogPhysicsControl, RBWC_LOG_LEVEL, TEXT("    %s:"), *Name.ToString());
 		}
 	}
 
-	UE_LOG(LogRigidBodyWithControl, RBWC_LOG_LEVEL, TEXT("Body Modifier sets:"));
+	UE_LOG(LogPhysicsControl, RBWC_LOG_LEVEL, TEXT("Body Modifier sets:"));
 	for (TMap<FName, TArray<FName>>::ElementType& Pair : NameRecords.BodyModifierSets)
 	{
-		UE_LOG(LogRigidBodyWithControl, RBWC_LOG_LEVEL, TEXT("  %s:"), *Pair.Key.ToString());
+		UE_LOG(LogPhysicsControl, RBWC_LOG_LEVEL, TEXT("  %s:"), *Pair.Key.ToString());
 		const TArray<FName>& Names = Pair.Value;
 		for (FName Name : Names)
 		{
-			UE_LOG(LogRigidBodyWithControl, RBWC_LOG_LEVEL, TEXT("    %s:"), *Name.ToString());
+			UE_LOG(LogPhysicsControl, RBWC_LOG_LEVEL, TEXT("    %s:"), *Name.ToString());
 		}
 	}
 #undef RBWC_LOG_LEVEL
@@ -439,7 +473,6 @@ static void ConvertStrengthToSpringParams(
 {
 	TOut AngularFrequency = TOut(InStrength * UE_DOUBLE_TWO_PI);
 	TOut Stiffness = AngularFrequency * AngularFrequency;
-
 	OutSpring = Stiffness;
 	OutDamping = TOut(InExtraDamping + 2 * InDampingRatio * AngularFrequency);
 }
@@ -450,24 +483,55 @@ static void ConvertStrengthToSpringParams(
 static bool UpdateDriveSpringDamperSettings(
 	Chaos::FPBDJointConstraintHandle* Constraint,
 	const Chaos::FPBDJointSettings&   Settings, 
-	const FPhysicsControlData&        ControlData)
+	const FPhysicsControlData&        Data,
+	const FPhysicsControlMultiplier&  Multiplier)
 {
-	float LinearSpring;
-	float LinearDamping;
 	float AngularSpring;
 	float AngularDamping;
-
-	ConvertStrengthToSpringParams(LinearSpring, LinearDamping, 
-		ControlData.LinearStrength, ControlData.LinearDampingRatio, ControlData.LinearExtraDamping);
-	ConvertStrengthToSpringParams(AngularSpring, AngularDamping, 
-		ControlData.AngularStrength, ControlData.AngularDampingRatio, ControlData.AngularExtraDamping);
+	const float MaxTorque = Data.MaxTorque * Multiplier.MaxTorqueMultiplier;
+
+	FVector LinearSpring;
+	FVector LinearDamping;
+	const FVector MaxForce = Data.MaxForce * Multiplier.MaxForceMultiplier;
+
+	UE::PhysicsControl::ConvertStrengthToSpringParams(
+		AngularSpring, AngularDamping,
+		Data.AngularStrength * Multiplier.AngularStrengthMultiplier,
+		Data.AngularDampingRatio * Multiplier.AngularDampingRatioMultiplier,
+		Data.AngularExtraDamping * Multiplier.AngularExtraDampingMultiplier);
+	UE::PhysicsControl::ConvertStrengthToSpringParams(
+		LinearSpring, LinearDamping,
+		Data.LinearStrength * Multiplier.LinearStrengthMultiplier,
+		Data.LinearDampingRatio * Multiplier.LinearDampingRatioMultiplier,
+		Data.LinearExtraDamping * Multiplier.LinearExtraDampingMultiplier);
+
+	if (Multiplier.MaxTorqueMultiplier <= 0)
+	{
+		AngularSpring = 0;
+		AngularDamping = 0;
+	}
+	if (Multiplier.MaxForceMultiplier.X <= 0)
+	{
+		LinearSpring.X = 0;
+		LinearDamping.X = 0;
+	}
+	if (Multiplier.MaxForceMultiplier.Y <= 0)
+	{
+		LinearSpring.Y = 0;
+		LinearDamping.Y = 0;
+	}
+	if (Multiplier.MaxForceMultiplier.Z <= 0)
+	{
+		LinearSpring.Z = 0;
+		LinearDamping.Z = 0;
+	}
 
 	Constraint->SetDriveParams(
-		Chaos::FVec3(LinearSpring), Chaos::FVec3(LinearDamping), Chaos::FVec3(ControlData.MaxForce),
-		Chaos::FVec3(AngularSpring), Chaos::FVec3(AngularDamping), Chaos::FVec3(ControlData.MaxTorque));
+		LinearSpring, LinearDamping, MaxForce,
+		Chaos::FVec3(AngularSpring), Chaos::FVec3(AngularDamping), Chaos::FVec3(Data.MaxTorque));
 
 	bool bHaveAngular = (AngularSpring + AngularDamping) > 0;
-	bool bHaveLinear = (LinearSpring + LinearDamping) > 0;
+	bool bHaveLinear = (LinearSpring + LinearDamping).GetMax() > 0;
 	return bHaveLinear || bHaveAngular;
 
 }
@@ -515,7 +579,8 @@ void FAnimNode_RigidBodyWithControl::ApplyControl(FRigidBodyControlRecord& Contr
 			Constraint->SetParentInvMassScale(ControlRecord.Control.ControlData.bOnlyControlChildObject ? 0 : 1);
 
 			const Chaos::FPBDJointSettings& JointSettings = Constraint->GetSettings();
-			if (UpdateDriveSpringDamperSettings(Constraint, JointSettings, ControlRecord.ControlData))
+			if (UpdateDriveSpringDamperSettings(
+				Constraint, JointSettings, ControlRecord.ControlData, ControlRecord.ControlMultiplier))
 			{
 				const FActorHandle* ChildActorHandle = JointHandle->GetActorHandles()[ConstraintChildIndex];
 				const FActorHandle* ParentActorHandle = JointHandle->GetActorHandles()[ConstraintParentIndex];
@@ -595,12 +660,13 @@ void FAnimNode_RigidBodyWithControl::ApplyControl(FRigidBodyControlRecord& Contr
 //======================================================================================================================
 void FAnimNode_RigidBodyWithControl::ApplyModifier(
 	const FRigidBodyModifierRecord& BodyModifierRecord, 
-	const FVector&             SimSpaceGravity)
+	const FVector&                  SimSpaceGravity)
 {
 	if (BodyModifierRecord.ActorHandle)
 	{
 		// Note that there's an early out if there's no change needed, so this should be OK.
-		BodyModifierRecord.ActorHandle->SetIsKinematic(
+		PhysicsSimulation->SetIsKinematic(
+			BodyModifierRecord.ActorHandle,
 			BodyModifierRecord.ModifierData.MovementType != EPhysicsMovementType::Simulated);
 
 		// Note that the actual kinematic targets will be set separately, since they need to be set
@@ -617,7 +683,7 @@ void FAnimNode_RigidBodyWithControl::ApplyModifier(
 
 		// Set collision
 		Chaos::FGeometryParticleHandle* ParticleHandle = BodyModifierRecord.ActorHandle->GetParticle();
-		ParticleHandle->SetHasCollision(CollisionEnabledHasPhysics(BodyModifierRecord.ModifierData.CollisionType));
+		//ParticleHandle->SetHasCollision(CollisionEnabledHasPhysics(BodyModifierRecord.ModifierData.CollisionType));
 	}
 }
 
@@ -629,7 +695,7 @@ void FAnimNode_RigidBodyWithControl::ApplyControlAndModifierUpdatesAndParameters
 	SCOPE_CYCLE_COUNTER(STAT_RigidBodyNodeWithControl_ApplyControlAndModifierUpdatesAndParametersToRecords);
 
 	// Apply control and modifier parameters on a single-use basis
-	ApplyControlAndBodyModifierDatas(Updates.ControlUpdates, Updates.ModifierUpdates);
+	ApplyControlAndBodyModifierDatas(Updates.ControlUpdates, Updates.ControlMultiplierUpdates, Updates.ModifierUpdates);
 
 	// This goes through the records, resetting the update parts.
 	// Then the update structures get adjusted based on the parameters.
@@ -655,6 +721,8 @@ void FAnimNode_RigidBodyWithControl::ApplyControlAndModifierUpdatesAndParameters
 
 	::ApplyControlAndModifierParametersToRecords(
 		ControlRecords, Parameters.ControlParameters, NameRecords.ControlSets);
+	::ApplyControlAndModifierParametersToRecords(
+		ControlRecords, Parameters.ControlMultiplierParameters, NameRecords.ControlSets);
 	::ApplyControlAndModifierParametersToRecords(
 		ModifierRecords, Parameters.ModifierParameters, NameRecords.BodyModifierSets);
 }
@@ -662,6 +730,7 @@ void FAnimNode_RigidBodyWithControl::ApplyControlAndModifierUpdatesAndParameters
 //======================================================================================================================
 void FAnimNode_RigidBodyWithControl::ApplyControlAndBodyModifierDatas(
 	const TArray<FPhysicsControlNamedControlParameters>&           InControlParameters,
+	const TArray<FPhysicsControlNamedControlMultiplierParameters>& InControlMultiplierParameters,
 	const TArray<FPhysicsControlNamedModifierParameters>&          InModifierParameters)
 {
 	SCOPE_CYCLE_COUNTER(STAT_RigidBodyNodeWithControl_ApplyControlsAndModifierDatas);
@@ -669,17 +738,35 @@ void FAnimNode_RigidBodyWithControl::ApplyControlAndBodyModifierDatas(
 	// This updates the "original" controls and modifiers based on the parameters.
 	for (const FPhysicsControlNamedControlParameters& ControlParameters : InControlParameters)
 	{
+		const FPhysicsControlSparseData& ControlData = ControlParameters.Data;
 		TArray<FName> Names = ExpandName(ControlParameters.Name, NameRecords.ControlSets);
 		for (FName Name : Names)
 		{
-			const FPhysicsControlSparseData& ControlData = ControlParameters.Data;
 			if (FRigidBodyControlRecord* ControlRecord = ControlRecords.Find(Name))
 			{
 				ControlRecord->Control.ControlData.UpdateFromSparseData(ControlData);
 			}
 			else
 			{
-				UE_LOG(LogRigidBodyWithControl, Warning,
+				UE_LOG(LogPhysicsControl, Warning,
+					TEXT("ApplyControlAndBodyModifierDatas: Failed to find control with name %s"), *Name.ToString());
+			}
+		}
+	}
+
+	for (const FPhysicsControlNamedControlMultiplierParameters& MultiplierParameters : InControlMultiplierParameters)
+	{
+		const FPhysicsControlSparseMultiplier& MultiplierData = MultiplierParameters.Data;
+		TArray<FName> Names = ExpandName(MultiplierParameters.Name, NameRecords.ControlSets);
+		for (FName Name : Names)
+		{
+			if (FRigidBodyControlRecord* ControlRecord = ControlRecords.Find(Name))
+			{
+				ControlRecord->Control.ControlMultiplier.UpdateFromSparseData(MultiplierData);
+			}
+			else
+			{
+				UE_LOG(LogPhysicsControl, Warning,
 					TEXT("ApplyControlAndBodyModifierDatas: Failed to find control with name %s"), *Name.ToString());
 			}
 		}
@@ -687,17 +774,17 @@ void FAnimNode_RigidBodyWithControl::ApplyControlAndBodyModifierDatas(
 
 	for (const FPhysicsControlNamedModifierParameters& ModifierParameters : InModifierParameters)
 	{
+		const FPhysicsControlModifierSparseData& ModifierData = ModifierParameters.Data;
 		TArray<FName> Names = ExpandName(ModifierParameters.Name, NameRecords.BodyModifierSets);
 		for (FName Name : Names)
 		{
-			const FPhysicsControlModifierSparseData& ModifierData = ModifierParameters.Data;
 			if (FRigidBodyModifierRecord* ModifierRecord = ModifierRecords.Find(Name))
 			{
 				ModifierRecord->Modifier.ModifierData.UpdateFromSparseData(ModifierData);
 			}
 			else
 			{
-				UE_LOG(LogRigidBodyWithControl, Warning,
+				UE_LOG(LogPhysicsControl, Warning,
 					TEXT("ApplyControlAndBodyModifierDatas: Failed to find modifier with name %s"), *Name.ToString());
 			}
 		}
@@ -719,13 +806,6 @@ void FAnimNode_RigidBodyWithControl::ApplyControlsAndModifiers(const FVector& Si
 		}
 	}
 
-	// Update world space control root object transform.
-	if (WorldSpaceControlActorHandle)
-	{
-		// This should be ok as we have no collisions on our root body
-		WorldSpaceControlActorHandle->SetWorldTransform(PoseData.WorldCoordinateFrame); 
-	}
-
 	if (!PoseData.IsEmpty())
 	{
 		// Apply Controls.
@@ -739,7 +819,6 @@ void FAnimNode_RigidBodyWithControl::ApplyControlsAndModifiers(const FVector& Si
 				{
 					ApplyControl(ControlRecord, DeltaTime);
 				}
-
 				SetConstraintEnabled(ControlRecord.JointHandle, ControlRecord.IsEnabled());
 			}
 		}
@@ -858,8 +937,9 @@ void FAnimNode_RigidBodyWithControl::ApplyCurrentConstraintProfile()
 
 					FPBDJointSettings JointSettings = ConstraintHandle->GetSettings();
 					
-					// See ImmediatePhysics_Chaos::TransferJointSettings
-					// TODO avoid code duplication with that function.
+					// See ImmediatePhysics_Chaos::TransferJointSettings TODO avoid code duplication
+					// with that function. However, it uses ConstraintInstance, but we don't have
+					// one of those.
 
 					JointSettings.Stiffness = ConstraintSettings::JointStiffness();
 					JointSettings.LinearProjection = Profile.bEnableProjection ? Profile.ProjectionLinearAlpha : 0.0f;
@@ -899,12 +979,18 @@ void FAnimNode_RigidBodyWithControl::ApplyCurrentConstraintProfile()
 					JointSettings.AngularSoftForceMode = (ConstraintSettings::SoftAngularForceMode() == 0) ? 
 						EJointForceMode::Acceleration : EJointForceMode::Force;
 
-					JointSettings.SoftLinearStiffness = Profile.LinearLimit.Stiffness;
-					JointSettings.SoftLinearDamping = Profile.LinearLimit.Damping;
-					JointSettings.SoftTwistStiffness = Profile.TwistLimit.Stiffness;
-					JointSettings.SoftTwistDamping = Profile.TwistLimit.Damping;
-					JointSettings.SoftSwingStiffness = Profile.ConeLimit.Stiffness;
-					JointSettings.SoftSwingDamping = Profile.ConeLimit.Damping;
+					JointSettings.SoftLinearStiffness = 
+						Chaos::ConstraintSettings::SoftLinearStiffnessScale() * Profile.LinearLimit.Stiffness;
+					JointSettings.SoftLinearDamping = 
+						Chaos::ConstraintSettings::SoftLinearDampingScale() * Profile.LinearLimit.Damping;
+					JointSettings.SoftTwistStiffness = 
+						Chaos::ConstraintSettings::SoftAngularStiffnessScale() * Profile.TwistLimit.Stiffness;
+					JointSettings.SoftTwistDamping = 
+						Chaos::ConstraintSettings::SoftAngularDampingScale() * Profile.TwistLimit.Damping;
+					JointSettings.SoftSwingStiffness = 
+						Chaos::ConstraintSettings::SoftAngularStiffnessScale() * Profile.ConeLimit.Stiffness;
+					JointSettings.SoftSwingDamping = 
+						Chaos::ConstraintSettings::SoftAngularDampingScale() * Profile.ConeLimit.Damping;
 
 					JointSettings.LinearRestitution = Profile.LinearLimit.Restitution;
 					JointSettings.TwistRestitution = Profile.TwistLimit.Restitution;
@@ -923,18 +1009,14 @@ void FAnimNode_RigidBodyWithControl::ApplyCurrentConstraintProfile()
 					JointSettings.bLinearVelocityDriveEnabled[1] = Profile.LinearDrive.YDrive.bEnableVelocityDrive;
 					JointSettings.bLinearVelocityDriveEnabled[2] = Profile.LinearDrive.ZDrive.bEnableVelocityDrive;
 					JointSettings.LinearDriveForceMode = EJointForceMode::Acceleration; // hardcoded!
-					JointSettings.LinearDriveStiffness[0] = 
-						ConstraintSettings::LinearDriveStiffnessScale() * Profile.LinearDrive.XDrive.Stiffness;
-					JointSettings.LinearDriveStiffness[1] = 
-						ConstraintSettings::LinearDriveStiffnessScale() * Profile.LinearDrive.YDrive.Stiffness;
-					JointSettings.LinearDriveStiffness[2] = 
-						ConstraintSettings::LinearDriveStiffnessScale() * Profile.LinearDrive.ZDrive.Stiffness;
-					JointSettings.LinearDriveDamping[0] = 
-						ConstraintSettings::LinearDriveDampingScale() * Profile.LinearDrive.XDrive.Stiffness;
-					JointSettings.LinearDriveDamping[1] = 
-						ConstraintSettings::LinearDriveDampingScale() * Profile.LinearDrive.YDrive.Stiffness;
-					JointSettings.LinearDriveDamping[2] = 
-						ConstraintSettings::LinearDriveDampingScale() * Profile.LinearDrive.ZDrive.Stiffness;
+					JointSettings.LinearDriveStiffness = Chaos::ConstraintSettings::LinearDriveStiffnessScale() * FVec3(
+							Profile.LinearDrive.XDrive.Stiffness, 
+							Profile.LinearDrive.YDrive.Stiffness, 
+							Profile.LinearDrive.ZDrive.Stiffness);
+					JointSettings.LinearDriveDamping = Chaos::ConstraintSettings::LinearDriveDampingScale() * FVec3(
+							Profile.LinearDrive.XDrive.Damping, 
+							Profile.LinearDrive.YDrive.Damping, 
+							Profile.LinearDrive.ZDrive.Damping);
 					JointSettings.LinearDriveMaxForce[0] = Profile.LinearDrive.XDrive.MaxForce;
 					JointSettings.LinearDriveMaxForce[1] = Profile.LinearDrive.YDrive.MaxForce;
 					JointSettings.LinearDriveMaxForce[2] = Profile.LinearDrive.ZDrive.MaxForce;
@@ -946,9 +1028,9 @@ void FAnimNode_RigidBodyWithControl::ApplyCurrentConstraintProfile()
 					if (Profile.AngularDrive.AngularDriveMode == EAngularDriveMode::SLERP)
 					{
 						JointSettings.AngularDriveStiffness = FVec3(
-							ConstraintSettings::LinearDriveStiffnessScale() * Profile.AngularDrive.SlerpDrive.Stiffness);
+							ConstraintSettings::AngularDriveStiffnessScale() * Profile.AngularDrive.SlerpDrive.Stiffness);
 						JointSettings.AngularDriveDamping = FVec3(
-							ConstraintSettings::LinearDriveDampingScale() * Profile.AngularDrive.SlerpDrive.Damping);
+							ConstraintSettings::AngularDriveDampingScale() * Profile.AngularDrive.SlerpDrive.Damping);
 						JointSettings.AngularDriveMaxTorque = FVec3(
 							Profile.AngularDrive.SlerpDrive.MaxForce);
 						JointSettings.bAngularSLerpPositionDriveEnabled = Profile.AngularDrive.SlerpDrive.bEnablePositionDrive;
@@ -960,18 +1042,14 @@ void FAnimNode_RigidBodyWithControl::ApplyCurrentConstraintProfile()
 					}
 					else
 					{
-						JointSettings.AngularDriveStiffness[0] = 
-							ConstraintSettings::LinearDriveStiffnessScale() * Profile.AngularDrive.TwistDrive.Stiffness;
-						JointSettings.AngularDriveStiffness[1] = 
-							ConstraintSettings::LinearDriveStiffnessScale() * Profile.AngularDrive.SwingDrive.Stiffness;
-						JointSettings.AngularDriveStiffness[2] = 
-							ConstraintSettings::LinearDriveStiffnessScale() * Profile.AngularDrive.SwingDrive.Stiffness;
-						JointSettings.AngularDriveDamping[0] = 
-							ConstraintSettings::LinearDriveDampingScale() * Profile.AngularDrive.TwistDrive.Damping;
-						JointSettings.AngularDriveDamping[1] = 
-							ConstraintSettings::LinearDriveDampingScale() * Profile.AngularDrive.SwingDrive.Damping;
-						JointSettings.AngularDriveDamping[2] = 
-							ConstraintSettings::LinearDriveDampingScale() * Profile.AngularDrive.SwingDrive.Damping;
+						JointSettings.AngularDriveStiffness = ConstraintSettings::AngularDriveStiffnessScale() * FVec3(
+							Profile.AngularDrive.TwistDrive.Stiffness, 
+							Profile.AngularDrive.SwingDrive.Stiffness, 
+							Profile.AngularDrive.SwingDrive.Stiffness);
+						JointSettings.AngularDriveDamping = ConstraintSettings::AngularDriveDampingScale() * FVec3(
+							Profile.AngularDrive.TwistDrive.Damping, 
+							Profile.AngularDrive.SwingDrive.Damping, 
+							Profile.AngularDrive.SwingDrive.Damping);
 						JointSettings.AngularDriveMaxTorque[0] = Profile.AngularDrive.TwistDrive.MaxForce;
 						JointSettings.AngularDriveMaxTorque[1] = Profile.AngularDrive.SwingDrive.MaxForce;
 						JointSettings.AngularDriveMaxTorque[2] = Profile.AngularDrive.SwingDrive.MaxForce;
@@ -1004,6 +1082,39 @@ void FAnimNode_RigidBodyWithControl::ApplyCurrentConstraintProfile()
 	}
 }
 
+//======================================================================================================================
+void FAnimNode_RigidBodyWithControl::InvokeControlProfile(FName ControlProfileName)
+{
+	CurrentControlProfile = ControlProfileName;
+	ApplyCurrentControlProfile();
+}
+
+//======================================================================================================================
+void FAnimNode_RigidBodyWithControl::ApplyCurrentControlProfile()
+{
+	// We shouldn't have a hitch here, since the asset (if set) will already have been loaded 
+	if (IsValid(PhysicsControlProfileAsset))
+	{
+		const FPhysicsControlControlAndModifierUpdates* Updates = 
+			PhysicsControlProfileAsset->Profiles.Find(CurrentControlProfile);
+		if (Updates)
+		{
+			ApplyControlAndBodyModifierDatas(
+				Updates->ControlUpdates, Updates->ControlMultiplierUpdates, Updates->ModifierUpdates);
+		}
+		else
+		{
+			UE_LOG(LogPhysicsControl, Warning,
+				TEXT("ApplyCurrentControlProfile: Profile %s not found"), *CurrentControlProfile.ToString());
+		}
+	}
+	else
+	{
+		UE_LOG(LogPhysicsControl, Warning,
+			TEXT("ApplyCurrentControlProfile: No control profile asset loaded"));
+	}
+}
+
 //======================================================================================================================
 void FAnimNode_RigidBodyWithControl::TransformConstraintsToMatchSkeletalMesh(
 	const USkeletalMesh*          const SkeletalMeshAsset, 
@@ -1016,7 +1127,7 @@ void FAnimNode_RigidBodyWithControl::TransformConstraintsToMatchSkeletalMesh(
 	{
 		if (SkeletalMeshAsset != PhysicsAssetAuthoredSkeletalMesh)
 		{
-			UE_LOGFMT(LogRigidBodyWithControl, Log, 
+			UE_LOGFMT(LogPhysicsControl, Log, 
 				"Modify Constraint parent transforms to correct for the difference between the Skeleton used to create the Physics asset \"{0}\" and the current skeleton \"{1}\".", 
 				PhysicsAssetAuthoredSkeletalMesh->GetName(), SkeletalMeshAsset->GetName());
 
@@ -1046,7 +1157,7 @@ void FAnimNode_RigidBodyWithControl::TransformConstraintsToMatchSkeletalMesh(
 				ConstraintInstance->SetRefFrame(EConstraintFrame::Frame2, CurrentRefFrame);
 
 #if !NO_LOGGING
-				UE_LOGFMT(LogRigidBodyWithControl, Log, 
+				UE_LOGFMT(LogPhysicsControl, Log, 
 					"Matched Constraint {0} - {1} Parent Transform - position was {2} now {3}.", 
 					ConstraintInstance->ConstraintBone1.ToString(), 
 					ConstraintInstance->ConstraintBone2.ToString(), 
@@ -1058,7 +1169,7 @@ void FAnimNode_RigidBodyWithControl::TransformConstraintsToMatchSkeletalMesh(
 #if !NO_LOGGING
 		else
 		{
-			UE_LOGFMT(LogRigidBodyWithControl, Log, 
+			UE_LOGFMT(LogPhysicsControl, Log, 
 				"Do not modify constraint parent transforms as the Skeleton used to create the Physics asset \"{0}\" matches the current skeleton \"{1}\".", 
 				PhysicsAssetAuthoredSkeletalMesh->GetName(), SkeletalMeshAsset->GetName());
 		}
@@ -1066,7 +1177,7 @@ void FAnimNode_RigidBodyWithControl::TransformConstraintsToMatchSkeletalMesh(
 	}
 	else if (SkeletalMeshAsset)
 	{
-		UE_LOGFMT(LogRigidBodyWithControl, Log, 
+		UE_LOGFMT(LogPhysicsControl, Log, 
 			"Snap Constraint parent transforms to the current skeleton \"{0}\" (Authored Skeleton is undefined in node details).", 
 			SkeletalMeshAsset->GetName());
 
@@ -1084,7 +1195,7 @@ void FAnimNode_RigidBodyWithControl::TransformConstraintsToMatchSkeletalMesh(
 			ConstraintInstance->SetRefFrame(EConstraintFrame::Frame2, ParentRelChildTM);
 
 #if !NO_LOGGING
-			UE_LOGFMT(LogRigidBodyWithControl, Log, 
+			UE_LOGFMT(LogPhysicsControl, Log, 
 				"Snapped Constraint {0} - {1} Parent Transform - position was {2} now {3}.", 
 				ConstraintInstance->ConstraintBone1.ToString(), 
 				ConstraintInstance->ConstraintBone2.ToString(), 
diff --git a/Engine/Plugins/Experimental/PhysicsControl/Source/PhysicsControl/Public/AnimNode_RigidBodyWithControl.h b/Engine/Plugins/Experimental/PhysicsControl/Source/PhysicsControl/Public/AnimNode_RigidBodyWithControl.h
--- a/Engine/Plugins/Experimental/PhysicsControl/Source/PhysicsControl/Public/AnimNode_RigidBodyWithControl.h
+++ b/Engine/Plugins/Experimental/PhysicsControl/Source/PhysicsControl/Public/AnimNode_RigidBodyWithControl.h
@@ -12,7 +12,6 @@
 #include "Physics/ImmediatePhysics/ImmediatePhysicsDeclares.h"
 #include "PhysicsControlLimbData.h"
 #include "PhysicsControlNameRecords.h"
-#include "PhysicsEngine/PhysicsAsset.h"
 #include "PhysicsProxy/PerSolverFieldSystem.h"
 #include "RigidBodyControlData.h"
 #include "RigidBodyPoseData.h"
@@ -25,12 +24,24 @@
 
 struct FBodyInstance;
 struct FConstraintInstance;
+
 class FEvent;
+class UPhysicsControlProfileAsset;
+class UPhysicsAsset;
 
 extern TAutoConsoleVariable<int32> CVarEnableRigidBodyNodeWithControl;
 extern TAutoConsoleVariable<int32> CVarEnableRigidBodyNodeWithControlSimulation;
 extern TAutoConsoleVariable<int32> CVarRigidBodyNodeWithControlLODThreshold;
 
+struct FSimulationSpaceData
+{
+	FTransform Transform;
+	FVector    LinearVel;
+	FVector    AngularVel;
+	FVector    LinearAcc;
+	FVector    AngularAcc;
+};
+
 /**
  * Controller that simulates physics based on the physics asset of the skeletal mesh component
  */
@@ -74,10 +85,19 @@ struct PHYSICSCONTROL_API FAnimNode_RigidBodyWithControl : public FAnimNode_Skel
 	 */
 	void SetOverridePhysicsAsset(UPhysicsAsset* PhysicsAsset);
 
+	/**
+	 * Invoke the profile stored in the control profile asset, and sets the current profile. Note that 
+	 * this produces permanent changes to the control and modifier values.
+	 */
+	void InvokeControlProfile(FName ControlProfileName);
+
 	UPhysicsAsset* GetPhysicsAsset() const { return PhysicsAssetToUse; }
 
 public:
-	/** Physics asset to use. If empty use the skeletal mesh's default physics asset */
+	/**
+	 * Physics asset to use. If empty use the skeletal mesh's default physics asset. Note that this can be
+	 * changed at runtime using SetOverridePhysicsAsset in the Animation Node Functions
+	 */
 	UPROPERTY(EditAnywhere, Category = Settings)
 	TObjectPtr<UPhysicsAsset> OverridePhysicsAsset;
 
@@ -98,7 +118,7 @@ private:
 	FTransform PreviousCompWorldSpaceTM;
 	FTransform CurrentTransform;
 	FTransform PreviousTransform;
-
+	
 	UPhysicsAsset* PhysicsAssetToUse;
 public:
 
@@ -135,7 +155,6 @@ public:
 	UPROPERTY(EditAnywhere, Category = Settings, meta = (PinHiddenByDefault))
 	FSimSpaceSettings SimSpaceSettings;
 
-
 	/**
 	 * Scale of cached bounds (vs. actual bounds) used for obtaining the list of objects we might collide with.
 	 * Increasing this may improve performance, but overlaps may not work as well.
@@ -150,7 +169,7 @@ public:
 	 * moving, otherwise the cache will only be updated if/when we move.
 	 */
 	UPROPERTY(EditAnywhere, Category = Settings, meta = (PinHiddenByDefault))
-	bool UpdateCacheEveryFrame;
+	bool bUpdateCacheEveryFrame;
 
 	/** Matters if SimulationSpace is BaseBone */
 	UPROPERTY(EditAnywhere, Category = Settings)
@@ -160,6 +179,15 @@ public:
 	UPROPERTY(EditAnywhere, Category = Settings, meta = (editcondition = "bEnableWorldGeometry"))
 	TEnumAsByte<ECollisionChannel> OverlapChannel;
 
+	/** 
+	 * Whether or not to calculate velocities for world geometry. Note that if the simulation space is
+	 * not set to world, then even static objects "should" have velocities calculated. There is a cost 
+	 * associated with calculating velocities, but without them there may be more penetration between
+	 * the simulated objects and the world. 
+	 */
+	UPROPERTY(EditAnywhere, Category = Settings, meta = (editcondition = "bEnableWorldGeometry"))
+	uint8 bCalculateVelocitiesForWorldGeometry : 1;
+
 	/** What space to simulate the bodies in. This affects how velocities are generated */
 	UPROPERTY(EditAnywhere, Category = Settings)
 	ESimulationSpace SimulationSpace;
@@ -226,6 +254,14 @@ public:
 	UPROPERTY(EditAnywhere, Category = ControlSetup, meta = (PinHiddenByDefault))
 	bool bEnableControls;
 
+	/** 
+	 * Physics control profile asset to use. This is optional, but if it has been set, then it will be used 
+	 * during creation of controls/body modifiers. Note that this can be modified at runtime in the Animation 
+	 * Blueprint via the Animation Node functions.
+	 */
+	UPROPERTY(EditAnywhere, Category = ControlSetup)
+	TObjectPtr<UPhysicsControlProfileAsset> PhysicsControlProfileAsset;
+
 	/**
 	 * A map of bone names to "body" names, the latter being used to assign names to controls/modifiers. 
 	 * This is optional - so if there is no mapping for a bone, then its name will be used directly when 
@@ -240,14 +276,17 @@ public:
 	 * Setup data for creating the main controls (world- and parent-space) and modifiers, based on splitting the 
 	 * skeleton up into limbs.
 	 */
-	UPROPERTY(EditAnywhere, Category = ControlSetup, meta = (PinHiddenByDefault))
-	FRigidBodySetupData SetupData;
+	UPROPERTY(EditAnywhere, Category = ControlSetup, meta = (PinHiddenByDefault, editcondition = "bEnableCharacterSetupData"))
+	FPhysicsControlCharacterSetupData CharacterSetupData;
+
+	UPROPERTY(EditAnywhere, Category = ControlSetup, meta = (InlineEditConditionToggle))
+	bool bEnableCharacterSetupData = true;
 
 	/**
 	 * Controls and modifiers that should be created, in addition to those made as part of the limb setup.
 	 */
 	UPROPERTY(EditAnywhere, Category = ControlSetup, meta = (PinHiddenByDefault))
-	FRigidBodyControlAndBodyModifierCreations AdditionalControlsAndBodyModifiers;
+	FPhysicsControlAndBodyModifierCreationDatas AdditionalControlsAndBodyModifiers;
 
 	/**
 	 * Allows additional sets of controls or modifiers to be created, and existing sets to be modified
@@ -302,11 +341,19 @@ public:
 	USkeletalMesh* PhysicsAssetAuthoredSkeletalMesh;
 
 	/**
-	 * The constraint profile to use on all the joints in the physics asset
+	 * The constraint profile to use on all the joints in the physics asset. Will be applied when 
+	 * the profile name changes.
 	 */
 	UPROPERTY(EditAnywhere, Category = Controls, meta = (PinHiddenByDefault))
 	FName ConstraintProfile;
 
+	/**
+	 * The control profile to use. Will be applied when the profile name changes. To force it to be invoked again, 
+	 * either change the profile to something else for a frame, or call InvokeControlProfile.
+	 */
+	UPROPERTY(EditAnywhere, Category = Controls, meta = (PinHiddenByDefault))
+	FName ControlProfile;
+
 private:
 	uint8 bEnabled : 1;
 	uint8 bSimulationStarted : 1;
@@ -332,7 +379,6 @@ private:
 
 public:
 	const FTransform GetBodyTransform(const int32 BodyIndex) const;
-	const FTransform GetWorldSpaceControlRootTransform() const;
 
 	FName CreateControl(const FName ParentBoneName, const FName ChildBoneName, const FPhysicsControlData& ControlData);
 
@@ -361,8 +407,8 @@ private:
 	// This applies the desired constraint profile, if necessary
 	void ApplyCurrentConstraintProfile();
 
-	TMap<FName, int32> BodyNameToIndexMap;
-	ImmediatePhysics::FActorHandle* WorldSpaceControlActorHandle;
+	// This applies the desired control profile, if necessary
+	void ApplyCurrentControlProfile();
 
 	FName GetBodyFromBoneName(const FName BoneName) const;
 
@@ -382,8 +428,9 @@ private:
 
 	// Applies the overrides to the underlying controls and modifiers
 	void ApplyControlAndBodyModifierDatas(
-		const TArray<FPhysicsControlNamedControlParameters>& ControlParameters,
-		const TArray<FPhysicsControlNamedModifierParameters>& ModifierParameters);
+		const TArray<FPhysicsControlNamedControlParameters>& InControlParameters,
+		const TArray<FPhysicsControlNamedControlMultiplierParameters>& InControlMultiplierParameters,
+		const TArray<FPhysicsControlNamedModifierParameters>& InModifierParameters);
 
 	// Applies the parameters to control and modifier records
 	void ApplyControlAndModifierUpdatesAndParametersToRecords(
@@ -397,16 +444,13 @@ private:
 
 	// Calculate simulation space transform, velocity etc to pass into the solver
 	void CalculateSimulationSpace(
-		ESimulationSpace Space,
-		const FTransform& ComponentToWorld,
-		const FTransform& BoneToComponent,
-		const float Dt,
+		ESimulationSpace         Space,
+		const FTransform&        ComponentToWorld,
+		const FTransform&        BoneToComponent,
+		const float              Dt,
 		const FSimSpaceSettings& Settings,
-		FTransform& SpaceTransform,
-		FVector& SpaceLinearVel,
-		FVector& SpaceAngularVel,
-		FVector& SpaceLinearAcc,
-		FVector& SpaceAngularAcc);
+		FSimulationSpaceData&    OutSimulationSpaceData
+	);
 
 	// Modify Constraint transforms relative to the parent bone to correct for the difference
 	// between the Skeleton used to create the Physics asset and the current skeleton.
@@ -421,7 +465,7 @@ private:
 	void RemoveClothColliderObjects();
 
 	// Update the sim-space transforms of all cloth collider objects.
-	void UpdateClothColliderObjects(const FTransform& SpaceTransform);
+	void UpdateClothColliderObjects(const FSimulationSpaceData& SimulationSpaceData);
 
 	// Gather nearby world objects and add them to the sim
 	void CollectWorldObjects();
@@ -433,7 +477,7 @@ private:
 	void PurgeExpiredWorldObjects();
 
 	// Update sim-space transforms of world objects
-	void UpdateWorldObjects(const FTransform& SpaceTransform, const float DeltaSeconds);
+	void UpdateWorldObjects(const FSimulationSpaceData& SimulationSpaceData, const float DeltaSeconds);
 
 	// Advances the simulation by a given timestep
 	void RunPhysicsSimulation(float DeltaSeconds, const FVector& SimSpaceGravity);
@@ -471,6 +515,9 @@ private:
 	friend class FRigidBodyNodeWithControlSimulationTask;
 	UE::Tasks::FTask SimulationTask;
 
+	TMap<FName, int32> BodyNameToIndexMap;
+	ImmediatePhysics::FActorHandle* WorldSpaceControlActorHandle;
+
 	struct FBodyAnimData
 	{
 		FBodyAnimData()
@@ -523,6 +570,8 @@ private:
 
 	FName CurrentConstraintProfile;
 
+	FName CurrentControlProfile;
+
 	TArray<USkeletalMeshComponent::FPendingRadialForces> PendingRadialForces;
 
 	FPerSolverFieldSystem PerSolverField;
@@ -560,8 +609,6 @@ private:
 
 	FCollisionQueryParams QueryParams;
 
-	FPhysScene* PhysScene;
-
 	// Used by CollectWorldObjects and UpdateWorldGeometry in the Task Thread. Typically, World
 	// should never be accessed off the Game Thread. However, since we're just doing overlaps this
 	// should be OK.
diff --git a/Engine/Plugins/Experimental/PhysicsControl/Source/PhysicsControl/Public/PhysicsControlComponent.h b/Engine/Plugins/Experimental/PhysicsControl/Source/PhysicsControl/Public/PhysicsControlComponent.h
--- a/Engine/Plugins/Experimental/PhysicsControl/Source/PhysicsControl/Public/PhysicsControlComponent.h
+++ b/Engine/Plugins/Experimental/PhysicsControl/Source/PhysicsControl/Public/PhysicsControlComponent.h
@@ -5,6 +5,7 @@
 #include "PhysicsControlData.h"
 #include "PhysicsControlLimbData.h"
 #include "PhysicsControlNameRecords.h"
+#include "PhysicsControlProfileAsset.h"
 #include "PhysicsControlRecord.h"
 
 #include "UObject/ObjectMacros.h"
@@ -17,7 +18,7 @@
 class FPrimitiveDrawInterface;
 struct FPhysicsControlComponentImpl;
 struct FPhysicsControlRecord;
-struct FPhysicsBodyModifier;
+struct FPhysicsBodyModifierRecord;
 struct FConstraintInstance;
 
 /**
@@ -1357,6 +1358,8 @@ public:
 	 * - "ControlType - i.e. "WorldSpace" or "ParentSpace", each of which will end up containing all controls of that type
 	 * - "ControlType_LimbName" - e.g. "WorldSpace_ArmLeft" or "ParentSpace_Head"
 	 * Each body modifier is added to "All" and a set named after the limb - e.g. "Spine" or "LegRight".
+	 * It is also possible to specify a mesh component to use for the "world" object - so that the world controls can 
+	 * be made to work in the space of another object (or a bone if that is a skeletal mesh component)
 	 */
 	UFUNCTION(BlueprintCallable, Category = PhysicsControl)
 	void CreateControlsAndBodyModifiersFromLimbBones(
@@ -1374,6 +1377,24 @@ public:
 		UMeshComponent*                             WorldComponent = nullptr,
 		FName                                       WorldBoneName = NAME_None);
 
+	/**
+	 * This uses the control profile asset (that should have already been set) to create
+	 * controls and body modifiers
+	 * It is also possible to specify a mesh component to use for the "world" object - so that the world controls can
+	 * be made to work in the space of another object (or a bone if that is a skeletal mesh component)
+	 */
+	UFUNCTION(BlueprintCallable, Category = PhysicsControl)
+	void CreateControlsAndBodyModifiersFromControlProfileAsset(
+		USkeletalMeshComponent* SkeletalMeshComponent,
+		UMeshComponent*         WorldComponent,
+		FName                   WorldBoneName);
+
+	/**
+	 * Looks up the profile which should exist in the registered control profile asset, and invokes it.
+	 */
+	UFUNCTION(BlueprintCallable, Category = PhysicsControl)
+	void InvokeControlProfile(FName ProfileName);
+
 	/**
 	 * Adds a Control to a Set. This will add a new set if necessary. For example, you might
 	 * make a set of Controls called "ParentSpace_Feet" by calling this twice, passing in the left and right
@@ -1627,6 +1648,9 @@ public:
 public:
 	// Public property data
 
+	UPROPERTY(EditAnywhere, BlueprintReadWrite, Category = PhysicsControlProfile)
+	TSoftObjectPtr<UPhysicsControlProfileAsset> PhysicsControlProfileAsset;
+
 	/**
 	 * If the component moves by more than this distance then it is treated as a teleport,
 	 * which prevents velocities being used for a frame. It is also used as the threshold for
@@ -1818,26 +1842,26 @@ protected:
 		bool                         bCalculateVelocity) const;
 
 	/** Updates the body based on the modifier */
-	void ApplyBodyModifier(FPhysicsBodyModifier& BodyModifier);
+	void ApplyBodyModifier(FPhysicsBodyModifierRecord& BodyModifier);
 
 	/**
 	 * This will set the kinematic target for the appropriate body based on the weighted target position and 
 	 * orientation (and whether any were found) for any controls that are related to the body modifier.
 	 */
-	void ApplyKinematicTarget(const FPhysicsBodyModifier& BodyModifier) const;
+	void ApplyKinematicTarget(const FPhysicsBodyModifierRecord& BodyModifier) const;
 
 	/**
 	 * Sets the body (simulated or kinematic) to have the position/velocity etc state that has been cached. 
 	 * Has no effect if there is no cached target data.
 	 */
-	void ResetToCachedTarget(const FPhysicsBodyModifier& BodyModifier) const;
+	void ResetToCachedTarget(const FPhysicsBodyModifierRecord& BodyModifier) const;
 
 	/**
 	 * Retrieves the body modifier for the name. Note that if Name is blank then the first modifier will be returned,
 	 * assuming there is one.
 	 */
-	FPhysicsBodyModifier* FindBodyModifier(const FName Name);
-	const FPhysicsBodyModifier* FindBodyModifier(const FName Name) const;
+	FPhysicsBodyModifierRecord* FindBodyModifierRecord(const FName Name);
+	const FPhysicsBodyModifierRecord* FindBodyModifierRecord(const FName Name) const;
 
 	/** 
 	 * When destroying a control or modifier, the record will normally be removed, but it can be retained if you
@@ -1855,11 +1879,9 @@ protected:
 	/** Destroys the modifier. It will optionally be removed from the array of records too */
 	bool DestroyBodyModifier(const FName Name, const EDestroyBehavior DestroyBehavior);
 
-	/** Calculates and returns a unique name based on the bone names and any existing records */
-	FName GetUniqueControlName(const FName ParentBoneName, const FName ChildBoneName, const FString& NamePrefix) const;
-
-	/** Calculates and returns a unique name based on the bone name and any existing records */
-	FName GetUniqueBodyModifierName(const FName BoneName) const;
+	/** Applies the updates to controls and body modifiers */
+	void ApplyControlAndModifierUpdates(
+		const FPhysicsControlControlAndModifierUpdates& ControlAndModifierUpdates);
 
 protected:
 
@@ -1871,8 +1893,8 @@ protected:
 	// and then need to be restored when the last body modifier is destroyed.
 	TMap<TWeakObjectPtr<USkeletalMeshComponent>, FModifiedSkeletalMeshData> ModifiedSkeletalMeshDatas;
 
-	TMap<FName, FPhysicsControlRecord> PhysicsControlRecords;
-	TMap<FName, FPhysicsBodyModifier> PhysicsBodyModifiers;
+	TMap<FName, FPhysicsControlRecord>      ControlRecords;
+	TMap<FName, FPhysicsBodyModifierRecord> BodyModifierRecords;
 
 	// Keep track of the names of everything we have created
 	FPhysicsControlNameRecords NameRecords;
diff --git a/Engine/Plugins/Experimental/PhysicsControl/Source/PhysicsControl/Public/PhysicsControlData.h b/Engine/Plugins/Experimental/PhysicsControl/Source/PhysicsControl/Public/PhysicsControlData.h
--- a/Engine/Plugins/Experimental/PhysicsControl/Source/PhysicsControl/Public/PhysicsControlData.h
+++ b/Engine/Plugins/Experimental/PhysicsControl/Source/PhysicsControl/Public/PhysicsControlData.h
@@ -93,6 +93,8 @@ struct PHYSICSCONTROL_API FPhysicsControlSetUpdates
 {
 	GENERATED_BODY();
 
+	FPhysicsControlSetUpdates& operator+=(const FPhysicsControlSetUpdates& Other);
+
 	UPROPERTY(EditAnywhere, BlueprintReadWrite, Category = PhysicsControl)
 	TArray<FPhysicsControlSetUpdate> ControlSetUpdates;
 
@@ -261,6 +263,41 @@ struct PHYSICSCONTROL_API FPhysicsControlModifierData
 };
 
 
+/**
+ * Represents a single/individual body modifier, with the modifier data
+ */
+USTRUCT(BlueprintType)
+struct PHYSICSCONTROL_API FPhysicsBodyModifier
+{
+	GENERATED_BODY();
+
+	FPhysicsBodyModifier() {}
+
+	FPhysicsBodyModifier(FName InBoneName, const FPhysicsControlModifierData& InModifierData)
+		: BoneName(InBoneName), ModifierData(InModifierData) {}
+
+	UPROPERTY(EditAnywhere, BlueprintReadWrite, Category = PhysicsControl)
+	FName BoneName;
+
+	UPROPERTY(EditAnywhere, BlueprintReadWrite, Category = PhysicsControl)
+	FPhysicsControlModifierData ModifierData;
+};
+
+/**
+ * Used on creation, to allow requesting the modifier to be in certain sets
+ */
+USTRUCT(BlueprintType)
+struct PHYSICSCONTROL_API FPhysicsBodyModifierCreationData
+{
+	GENERATED_BODY();
+
+	UPROPERTY(EditAnywhere, BlueprintReadWrite, Category = PhysicsControl)
+	FPhysicsBodyModifier Modifier;
+
+	UPROPERTY(EditAnywhere, BlueprintReadWrite, Category = PhysicsControl)
+	TArray<FName> Sets;
+};
+
 /**
  * Strength and damping etc parameters that will affect a control, with flags indicating
  * whether each element should get used.
@@ -819,33 +856,22 @@ struct PHYSICSCONTROL_API FPhysicsControl
 	FPhysicsControl() {}
 
 	FPhysicsControl(
-		UMeshComponent*              InParentMeshComponent,
 		const FName&                 InParentBoneName,
-		UMeshComponent*              InChildMeshComponent,
 		const FName&                 InChildBoneName,
-		const FPhysicsControlData&   InControlData,
-		const FPhysicsControlTarget& InControlTarget)
-		: ParentMeshComponent(InParentMeshComponent)
-		, ParentBoneName(InParentBoneName)
-		, ChildMeshComponent(InChildMeshComponent)
+		const FPhysicsControlData&   InControlData)
+		: ParentBoneName(InParentBoneName)
 		, ChildBoneName(InChildBoneName)
 		, ControlData(InControlData)
-		, ControlTarget(InControlTarget)
 	{
 	}
 
-	/**  The mesh that will be doing the driving. Blank/non-existent means it will happen in world space */
-	UPROPERTY(EditAnywhere, BlueprintReadWrite, Category = PhysicsControl)
-	TWeakObjectPtr<UMeshComponent> ParentMeshComponent;
+	// Indicates if the control is enabled
+	bool IsEnabled() const { return ControlData.bEnabled; }
 
 	/** The name of the skeletal mesh bone or the name of the static mesh body that will be doing the driving. */
 	UPROPERTY(EditAnywhere, BlueprintReadWrite, Category = PhysicsControl)
 	FName ParentBoneName;
 
-	/** The mesh that the control will be driving. */
-	UPROPERTY(EditAnywhere, BlueprintReadWrite, Category = PhysicsControl)
-	TWeakObjectPtr<UMeshComponent> ChildMeshComponent;
-
 	/** 
 	 * The name of the skeletal mesh bone or the name of the static mesh body that the control 
 	 * will be driving. */
@@ -865,14 +891,21 @@ struct PHYSICSCONTROL_API FPhysicsControl
 	 */
 	UPROPERTY(EditAnywhere, BlueprintReadWrite, Category = PhysicsControl)
 	FPhysicsControlMultiplier ControlMultiplier;
+};
+
+/**
+ * Used on creation, to allow requesting the control to be in certain sets
+ */
+USTRUCT(BlueprintType)
+struct PHYSICSCONTROL_API FPhysicsControlCreationData
+{
+	GENERATED_BODY();
 
-	/**
-	 * The position/orientation etc targets for the controls. These are procedural/explicit control targets -
-	 * skeletal meshes have the option to use skeletal animation as well, in which case these targets are 
-	 * expressed as relative to that animation.
-	 */
 	UPROPERTY(EditAnywhere, BlueprintReadWrite, Category = PhysicsControl)
-	FPhysicsControlTarget ControlTarget;
+	FPhysicsControl Control;
+
+	UPROPERTY(EditAnywhere, BlueprintReadWrite, Category = PhysicsControl)
+	TArray<FName> Sets;
 };
 
 /**
@@ -954,7 +987,15 @@ struct PHYSICSCONTROL_API FPhysicsControlControlAndModifierParameters
 	 */ 
 	UPROPERTY(EditAnywhere, BlueprintReadWrite, Category = Settings)
 	TArray<FPhysicsControlNamedControlParameters> ControlParameters;
-	
+
+	/**
+	 * Multipliers for existing controls. Each name can be the name of a control, or the name of a 
+	 * set of controls. They will only apply for one tick/update. They will be applied in order (so 
+	 * subsequent entries will override earlier ones if they apply to the same control).
+	 */ 
+	UPROPERTY(EditAnywhere, BlueprintReadWrite, Category = Settings)
+	TArray<FPhysicsControlNamedControlMultiplierParameters> ControlMultiplierParameters;
+
 	/**
 	 *  Parameters for existing modifiers. Each name can be the name of a modifier, or the name of a 
 	 * set of modifiers. They will only apply for one tick/update.
@@ -963,6 +1004,7 @@ struct PHYSICSCONTROL_API FPhysicsControlControlAndModifierParameters
 	TArray<FPhysicsControlNamedModifierParameters> ModifierParameters;
 
 	void Add(const FPhysicsControlNamedControlParameters& InParameters) { ControlParameters.Add(InParameters); }
+	void Add(const FPhysicsControlNamedControlMultiplierParameters& InParameters) { ControlMultiplierParameters.Add(InParameters); }
 	void Add(const FPhysicsControlNamedModifierParameters& InParameters) { ModifierParameters.Add(InParameters); }
 };
 
@@ -988,6 +1030,23 @@ struct PHYSICSCONTROL_API FPhysicsControlControlAndModifierUpdates
 	TArray<FPhysicsControlNamedModifierParameters> ModifierUpdates;
 };
 
+/**
+ * Collection of controls and body modifiers, used for creation
+ */
+USTRUCT(BlueprintType)
+struct PHYSICSCONTROL_API FPhysicsControlAndBodyModifierCreationDatas
+{
+	GENERATED_BODY();
+
+	FPhysicsControlAndBodyModifierCreationDatas& operator+=(const FPhysicsControlAndBodyModifierCreationDatas& other);
+
+	UPROPERTY(EditAnywhere, BlueprintReadWrite, Category = PhysicsControl)
+	TMap<FName, FPhysicsControlCreationData> Controls;
+
+	UPROPERTY(EditAnywhere, BlueprintReadWrite, Category = PhysicsControl)
+	TMap<FName, FPhysicsBodyModifierCreationData> Modifiers;
+};
+
 // Functions to interpolate between various control/modifier datas
 PHYSICSCONTROL_API FPhysicsControlData Interpolate(
 	const FPhysicsControlData& A, const FPhysicsControlData& B, const float Weight);
@@ -1000,3 +1059,5 @@ PHYSICSCONTROL_API FPhysicsControlModifierSparseData Interpolate(
 	const FPhysicsControlModifierSparseData& A, const FPhysicsControlModifierSparseData& B, const float Weight);
 
 
+
+
diff --git a/Engine/Plugins/Experimental/PhysicsControl/Source/PhysicsControl/Public/PhysicsControlLimbData.h b/Engine/Plugins/Experimental/PhysicsControl/Source/PhysicsControl/Public/PhysicsControlLimbData.h
--- a/Engine/Plugins/Experimental/PhysicsControl/Source/PhysicsControl/Public/PhysicsControlLimbData.h
+++ b/Engine/Plugins/Experimental/PhysicsControl/Source/PhysicsControl/Public/PhysicsControlLimbData.h
@@ -100,3 +100,23 @@ struct PHYSICSCONTROL_API FPhysicsControlNames
 	TArray<FName> Names;
 };
 
+USTRUCT(BlueprintType)
+struct PHYSICSCONTROL_API FPhysicsControlCharacterSetupData
+{
+	GENERATED_BODY();
+
+	FPhysicsControlCharacterSetupData& operator+=(const FPhysicsControlCharacterSetupData& other);
+
+	UPROPERTY(EditAnywhere, Category = ControlSetup)
+	TArray<FPhysicsControlLimbSetupData> LimbSetupData;
+
+	UPROPERTY(EditAnywhere, Category = ControlSetup)
+	FPhysicsControlData DefaultWorldSpaceControlData;
+
+	UPROPERTY(EditAnywhere, Category = ControlSetup)
+	FPhysicsControlData DefaultParentSpaceControlData;
+
+	UPROPERTY(EditAnywhere, Category = ControlSetup)
+	FPhysicsControlModifierData DefaultBodyModifierData;
+};
+
diff --git a/Engine/Plugins/Experimental/PhysicsControl/Source/PhysicsControl/Public/PhysicsControlLog.h b/Engine/Plugins/Experimental/PhysicsControl/Source/PhysicsControl/Public/PhysicsControlLog.h
--- a/Engine/Plugins/Experimental/PhysicsControl/Source/PhysicsControl/Public/PhysicsControlLog.h
+++ b/Engine/Plugins/Experimental/PhysicsControl/Source/PhysicsControl/Public/PhysicsControlLog.h
@@ -5,11 +5,9 @@
 #include "Logging/LogMacros.h"
 
 #if UE_BUILD_SHIPPING || UE_BUILD_TEST
-PHYSICSCONTROL_API DECLARE_LOG_CATEGORY_EXTERN(LogRigidBodyWithControl, Log, Warning);
-PHYSICSCONTROL_API DECLARE_LOG_CATEGORY_EXTERN(LogPhysicsControlComponent, Log, Warning);
+PHYSICSCONTROL_API DECLARE_LOG_CATEGORY_EXTERN(LogPhysicsControl, Log, Warning);
 #else
-PHYSICSCONTROL_API DECLARE_LOG_CATEGORY_EXTERN(LogRigidBodyWithControl, Log, All);
-PHYSICSCONTROL_API DECLARE_LOG_CATEGORY_EXTERN(LogPhysicsControlComponent, Log, All);
+PHYSICSCONTROL_API DECLARE_LOG_CATEGORY_EXTERN(LogPhysicsControl, Log, All);
 #endif
 
 
diff --git a/Engine/Plugins/Experimental/PhysicsControl/Source/PhysicsControl/Public/PhysicsControlOperatorNameGeneration.h b/Engine/Plugins/Experimental/PhysicsControl/Source/PhysicsControl/Public/PhysicsControlOperatorNameGeneration.h
--- a/Engine/Plugins/Experimental/PhysicsControl/Source/PhysicsControl/Public/PhysicsControlOperatorNameGeneration.h
+++ b/Engine/Plugins/Experimental/PhysicsControl/Source/PhysicsControl/Public/PhysicsControlOperatorNameGeneration.h
@@ -8,32 +8,137 @@
 
 class UPhysicsAsset;
 
+#include "RigidBodyControlData.h"
+#include "PhysicsControlRecord.h"
+
 struct FAnimNode_RigidBodyWithControl;
 struct FPhysicsControlLimbBones;
 struct FPhysicsControlLimbSetupData;
 struct FPhysicsControlSetUpdates;
 struct FReferenceSkeleton;
-struct FRigidBodyControlRecord;
-struct FRigidBodyModifierRecord;
 struct FPhysicsControlNameRecords;
 
+namespace UE
+{
+namespace PhysicsControl
+{
+
 constexpr int32 MaxNumControlsOrModifiersPerName = 16;
 
 // Parse the skeleton tree to figure out which bones are associated with which limbs. 
-PHYSICSCONTROL_API TMap<FName, FPhysicsControlLimbBones> GetLimbBones(const TArray<FPhysicsControlLimbSetupData>& LimbSetupData, const FReferenceSkeleton& RefSkeleton, UPhysicsAsset* const PhysicsAsset);
+PHYSICSCONTROL_API TMap<FName, FPhysicsControlLimbBones> GetLimbBones(
+	const TArray<FPhysicsControlLimbSetupData>& LimbSetupData, 
+	const FReferenceSkeleton&                   RefSkeleton, 
+	UPhysicsAsset*                              PhysicsAsset);
 
-// Populates the supplied body modifier names, control names and name records structures with the names and sets that could be created for the supplied node, limb bones, skeleton and physics asset.
-PHYSICSCONTROL_API void CollectOperatorNames(const FAnimNode_RigidBodyWithControl* const Node, TMap<FName, FPhysicsControlLimbBones> AllLimbBones, const FReferenceSkeleton& RefSkeleton, UPhysicsAsset* const PhysicsAsset, TSet<FName>& BodyModifierNames, TSet<FName>& ControlNames, FPhysicsControlNameRecords& NameRecords);
+// Populates the supplied body modifier names, control names and name records structures with the
+// names and sets that could be created for the supplied node, limb bones, skeleton and physics asset.
+PHYSICSCONTROL_API void CollectOperatorNames(
+	const FAnimNode_RigidBodyWithControl*            Node, 
+	const FPhysicsControlCharacterSetupData&         CharacterSetupData,
+	const FPhysicsControlAndBodyModifierCreationDatas& AdditionalControlsAndBodyModifiers,
+	TMap<FName, FPhysicsControlLimbBones>            AllLimbBones,
+	const FReferenceSkeleton&                        RefSkeleton, 
+	UPhysicsAsset*                                   PhysicsAsset, 
+	TSet<FName>&                                     BodyModifierNames, 
+	TSet<FName>&                                     ControlNames, 
+	FPhysicsControlNameRecords&                      NameRecords);
 
 // Creates the body modifiers, controls and sets for the supplied node, limb bones, skeleton and physics asset.
-PHYSICSCONTROL_API void CreateOperatorsForNode(FAnimNode_RigidBodyWithControl* const Node, TMap<FName, FPhysicsControlLimbBones> AllLimbBones, const FReferenceSkeleton& RefSkeleton, UPhysicsAsset* const PhysicsAsset, FPhysicsControlNameRecords& NameRecords);
+PHYSICSCONTROL_API void CreateOperatorsForNode(
+	FAnimNode_RigidBodyWithControl*                  Node, 
+	const FPhysicsControlCharacterSetupData&         CharacterSetupData,
+	const FPhysicsControlAndBodyModifierCreationDatas& AdditionalControlsAndBodyModifiers,
+	TMap<FName, FPhysicsControlLimbBones>            AllLimbBones,
+	const FReferenceSkeleton&                        RefSkeleton, 
+	UPhysicsAsset*                                   PhysicsAsset, 
+	FPhysicsControlNameRecords&                      NameRecords);
+
+// Adds the specified additional sets to the supplied Name Records structure.
+PHYSICSCONTROL_API void CreateAdditionalSets(
+	const FPhysicsControlSetUpdates& AdditionalSets, 
+	const TSet<FName>&               BodyModifierNames, 
+	const TSet<FName>&               ControlNames, 
+	FPhysicsControlNameRecords&      NameRecords);
 
 // Adds the specified additional sets to the supplied Name Records structure.
-PHYSICSCONTROL_API void CreateAdditionalSets(const FPhysicsControlSetUpdates& AdditionalSets, const TSet<FName>& BodyModifierNames, const TSet<FName>& ControlNames, FPhysicsControlNameRecords& NameRecords);
+PHYSICSCONTROL_API void CreateAdditionalSets(
+	const FPhysicsControlSetUpdates&             AdditionalSets, 
+	const TMap<FName, FRigidBodyModifierRecord>& BodyModifierRecords, 
+	const TMap<FName, FRigidBodyControlRecord>&  Controls, 
+	FPhysicsControlNameRecords&                  NameRecords);
 
 // Adds the specified additional sets to the supplied Name Records structure.
-PHYSICSCONTROL_API void CreateAdditionalSets(const FPhysicsControlSetUpdates& AdditionalSets, const TMap<FName, FRigidBodyModifierRecord>& BodyModifierNames, const TMap<FName, FRigidBodyControlRecord>& ControlNames, FPhysicsControlNameRecords& NameRecords);
+PHYSICSCONTROL_API void CreateAdditionalSets(
+	const FPhysicsControlSetUpdates&             AdditionalSets, 
+	const TMap<FName, FPhysicsBodyModifierRecord>&     BodyModifierRecords,
+	const TMap<FName, FPhysicsControlRecord>&    Controls,
+	FPhysicsControlNameRecords&                  NameRecords);
+
+//======================================================================================================================
+template<typename CollectionType>
+FName GetUniqueName(const FString& NameBase, const CollectionType& ExistingNames, const int32 MaxNameIndex)
+{
+	FString NameStr = NameBase;
+
+	// If the number gets too large, almost certainly we're in some nasty situation where this is
+	// getting called in a loop. Better to quit and fail, rather than allow the constraint set to
+	// increase without bound. 
+	for (int32 Index = 0; Index < MaxNameIndex; ++Index)
+	{
+		const FName Name(NameStr);
+		if (!ExistingNames.Find(Name))
+		{
+			return Name;
+		}
+		NameStr = FString::Format(TEXT("{0}_{1}"), { NameBase, Index });
+	}
+	UE_LOG(LogPhysicsControl, Warning,
+		TEXT("Unable to find a suitable Control name - the limit of (%d) has been exceeded"),
+		MaxNameIndex);
+
+	return NAME_None;
+}
+
+//======================================================================================================================
+template<typename CollectionType>
+FName GetUniqueBodyModifierName(const FName BodyName, const CollectionType& ExistingNames, const FString& NamePrefix)
+{
+	FString NameBase = NamePrefix;
+	if (!BodyName.IsNone())
+	{
+		NameBase += BodyName.ToString();
+	}
+	else
+	{
+		NameBase = TEXT("Body");
+	}
+
+	return GetUniqueName(NameBase, ExistingNames, MaxNumControlsOrModifiersPerName);
+}
+
+//======================================================================================================================
+template<typename CollectionType>
+FName GetUniqueControlName(
+	const FName ParentBodyName, const FName ChildBodyName, const CollectionType& ExistingNames, const FString& NamePrefix)
+{
+	FString NameBase = NamePrefix;
+	if (!ParentBodyName.IsNone())
+	{
+		NameBase += ParentBodyName.ToString() + TEXT("_");
+	}
+	if (!ChildBodyName.IsNone())
+	{
+		NameBase += ChildBodyName.ToString();
+	}
+
+	return GetUniqueName(NameBase, ExistingNames, MaxNumControlsOrModifiersPerName);
+}
+
+
+
+
+
 
-FName FindUniqueName(const FString& NameBase, const TSet<FName>& Keys, const int32 MaxNmeaIndex);
-FName FindUniqueBodyModifierName(const FName BodyName, const TSet<FName>& ExistingNames);
-FName FindUniqueControlName(const FName ParentBodyName, const FName ChildBodyName, const TSet<FName>& ExistingNames);
+} // namespace PhysicsControlComponent
+} // namespace UE
diff --git a/Engine/Plugins/Experimental/PhysicsControl/Source/PhysicsControl/Public/PhysicsControlProfileAsset.h b/Engine/Plugins/Experimental/PhysicsControl/Source/PhysicsControl/Public/PhysicsControlProfileAsset.h
--- a/Engine/Plugins/Experimental/PhysicsControl/Source/PhysicsControl/Public/PhysicsControlProfileAsset.h
+++ b/Engine/Plugins/Experimental/PhysicsControl/Source/PhysicsControl/Public/PhysicsControlProfileAsset.h
@@ -11,59 +11,6 @@
 
 class USkeletalMesh;
 
-
-USTRUCT(BlueprintType)
-struct PHYSICSCONTROL_API FPhysicsControlProfileInitialControl
-{
-	GENERATED_BODY()
-
-	UPROPERTY(EditAnywhere, BlueprintReadWrite, Category = PhysicsControl)
-	FName ParentBoneName;
-
-	UPROPERTY(EditAnywhere, BlueprintReadWrite, Category = PhysicsControl)
-	FName ChildBoneName;
-
-	UPROPERTY(EditAnywhere, BlueprintReadWrite, Category = PhysicsControl)
-	FPhysicsControlData ControlData;
-
-	UPROPERTY(EditAnywhere, BlueprintReadWrite, Category = PhysicsControl)
-	TArray<FName> Sets;
-};
-
-
-USTRUCT(BlueprintType)
-struct PHYSICSCONTROL_API FPhysicsControlProfileInitialBodyModifier
-{
-	GENERATED_BODY()
-
-	UPROPERTY(EditAnywhere, BlueprintReadWrite, Category = PhysicsControl)
-	FName BoneName;
-
-	UPROPERTY(EditAnywhere, BlueprintReadWrite, Category = PhysicsControl)
-	FPhysicsControlModifierData BodyModifierData;
-
-	UPROPERTY(EditAnywhere, BlueprintReadWrite, Category = PhysicsControl)
-	TArray<FName> Sets;
-};
-
-USTRUCT(BlueprintType)
-struct FPhysicsControlProfileCharacterSetup
-{
-	GENERATED_BODY();
-
-	UPROPERTY(EditAnywhere, Category = ControlSetup)
-	TArray<FPhysicsControlLimbSetupData> LimbSetupData;
-
-	UPROPERTY(EditAnywhere, Category = ControlSetup)
-	FPhysicsControlData DefaultWorldSpaceControlData;
-
-	UPROPERTY(EditAnywhere, Category = ControlSetup)
-	FPhysicsControlData DefaultParentSpaceControlData;
-
-	UPROPERTY(EditAnywhere, Category = ControlSetup)
-	FPhysicsControlModifierData DefaultBodyModifierData;
-};
-
 /**
  * Asset for storing Physics Control Profiles. These will contain data that define:
  * - Controls and body modifiers to be created on a mesh
@@ -81,63 +28,138 @@ class PHYSICSCONTROL_API UPhysicsControlProfileAsset : public UObject, public II
 public:
 	UPhysicsControlProfileAsset();
 
-	/** Shows all the controls etc that would be made */
-	UFUNCTION(CallInEditor, Category = Development)
-	void Log();
-
-	/** Sets the initial controls/modifiers based on character setup data */
-	UFUNCTION(CallInEditor, Category = Development)
-	void MakeControlsAndModifiersFromCharacterSetupData();
+	// Data that have been compiled from a combination of inherited and "My" data.
 
 	/**
 	 * We can define controls in the form of limbs etc here
 	 */
-	UPROPERTY(EditAnywhere, Category = Development)
-	FPhysicsControlProfileCharacterSetup CharacterSetupData;
+	UPROPERTY()
+	FPhysicsControlCharacterSetupData CharacterSetupData;
+
+	/**
+	 * Additional controls and modifiers. If these have the same name as one that's already 
+	 * created, they'll just override it.
+	 */
+	UPROPERTY()
+	FPhysicsControlAndBodyModifierCreationDatas AdditionalControlsAndModifiers;
+
+	/**
+	 * Additional control and body modifier sets
+	 */
+	UPROPERTY()
+	FPhysicsControlSetUpdates AdditionalSets;
+
+	/**
+	 * Initial updates to apply immediately after controls and modifiers are created
+	 */
+	UPROPERTY()
+	TArray<FPhysicsControlControlAndModifierUpdates> InitialControlAndModifierUpdates;
+
+	/**
+	 * The named profiles, which are essentially control and modifier updates
+	 */
+	UPROPERTY()
+	TMap<FName, FPhysicsControlControlAndModifierUpdates> Profiles;
 
 public:
+	// Data that will then be compiled down into the runtime data
+
+#if WITH_EDITORONLY_DATA
 	/** A profile asset to inherit from (can be null). If set, we will just add/modify data in that */
 	UPROPERTY(EditAnywhere, Category = Inheritance)
-	TObjectPtr<UPhysicsControlProfileAsset> ParentAsset;
+	TSoftObjectPtr<UPhysicsControlProfileAsset> ParentAsset;
 
-	/** Additional profile assets to use for the profiles */
+	/** 
+	 * Additional profile assets from which profiles (not the setup data, extra sets etc) will be added 
+	 * to this asset.
+	 */
 	UPROPERTY(EditAnywhere, Category = Inheritance)
-	TArray<TObjectPtr<UPhysicsControlProfileAsset>> AdditionalProfileAssets;
+	TArray<TSoftObjectPtr<UPhysicsControlProfileAsset>> AdditionalProfileAssets;
 
-public:
-	/** 
+	/**
 	* The skeletal mesh to use for generating controls and previewing. If it turns out this 
 	* doesn't need to be stored in the asset, it will get moved out of here.
 	*/
 	UPROPERTY(AssetRegistrySearchable, EditAnywhere, Category = PreviewMesh)
 	TSoftObjectPtr<USkeletalMesh> PreviewSkeletalMesh;
 
+public:
+	// "My" runtime data - i.e. the data that will be combined with what has been inherited
+	// We should have custom UI that displays this combined with the inherited data
+
+	/**
+	 * We can define controls in the form of limbs etc here
+	 */
+	UPROPERTY(EditAnywhere, Category = ProfileData, meta=(DisplayName="Character Setup Data"))
+	FPhysicsControlCharacterSetupData MyCharacterSetupData;
+
+	/**
+	 * Additional controls and modifiers. If these have the same name as one that's 
+	 * already created, they'll just override it.
+	 */
+	UPROPERTY(EditAnywhere, Category = ProfileData, meta = (DisplayName = "Additional Controls and Modifiers"))
+	FPhysicsControlAndBodyModifierCreationDatas MyAdditionalControlsAndModifiers;
+
 	/**
-	 * Additional controls. If these have the same name as one that's already created, they'll just override it.
+	 * Additional control and body modifier sets
 	 */
-	UPROPERTY(EditAnywhere, Category = ControlsAndBodyModifiers)
-	TMap<FName, FPhysicsControlProfileInitialControl> InitialControls;
+	UPROPERTY(EditAnywhere, Category = ProfileData, meta = (DisplayName = "Additional Sets"))
+	FPhysicsControlSetUpdates MyAdditionalSets;
 
 	/**
-	 * Additional body modifiers. If these have the same name as one that's already created, they'll just override it.
+	 * Initial updates to apply immediately after controls and modifiers are created
 	 */
-	UPROPERTY(EditAnywhere, Category = ControlsAndBodyModifiers)
-	TMap<FName, FPhysicsControlProfileInitialBodyModifier> InitialBodyModifiers;
+	UPROPERTY(EditAnywhere, Category = ProfileData, meta = (DisplayName = "Initial Control and Modifier Updates"))
+	TArray<FPhysicsControlControlAndModifierUpdates> MyInitialControlAndModifierUpdates;
 
 	/**
 	 * The named profiles, which are essentially control and modifier updates
 	 */
-	UPROPERTY(EditAnywhere, Category = Profiles)
-	TMap<FName, FPhysicsControlControlAndModifierUpdates> Profiles;
+	UPROPERTY(EditAnywhere, Category = ProfileData, meta = (DisplayName = "Profiles"))
+	TMap<FName, FPhysicsControlControlAndModifierUpdates> MyProfiles;
+#endif
 
 public:
+	// Buttons/actions in the editor
+
+#if WITH_EDITOR
+	/** Shows all the controls etc that would be made */
+	UFUNCTION(CallInEditor, Category = Development)
+	void ShowCompiledData() const;
+
+	/** 
+	 * Collapses inherited and authored profiles etc to make a profile asset that can be read without 
+	 * need for subsequent processing.
+	 */
+	UFUNCTION(CallInEditor)
+	void Compile();
+
+	/** Combines and returns data from our parent and ourself */
+	FPhysicsControlCharacterSetupData GetCharacterSetupData() const;
+
+	/** Combines and returns data from our parent and ourself */
+	FPhysicsControlAndBodyModifierCreationDatas GetAdditionalControlsAndModifiers() const;
+
+	/** Combines and returns data from our parent and ourself */
+	FPhysicsControlSetUpdates GetAdditionalSets() const;
+
+	/** Combines and returns data from our parent and ourself */
+	TArray<FPhysicsControlControlAndModifierUpdates> GetInitialControlAndModifierUpdates() const;
+
+	/** Combines and returns data from our parent and ourself */
+	TMap<FName, FPhysicsControlControlAndModifierUpdates> GetProfiles() const;
+
+#endif
+
+public:
+
 	/** IInterface_PreviewMeshProvider interface */
 	virtual void SetPreviewMesh(USkeletalMesh* PreviewMesh, bool bMarkAsDirty = true) override;
 	virtual USkeletalMesh* GetPreviewMesh() const override;
 	/** END IInterface_PreviewMeshProvider interface */
 
 #if WITH_EDITOR
-/* Get name of Preview Mesh property */
+	/* Get name of Preview Mesh property */
 	static const FName GetPreviewMeshPropertyName();
 #endif
 
diff --git a/Engine/Plugins/Experimental/PhysicsControl/Source/PhysicsControl/Public/PhysicsControlRecord.h b/Engine/Plugins/Experimental/PhysicsControl/Source/PhysicsControl/Public/PhysicsControlRecord.h
--- a/Engine/Plugins/Experimental/PhysicsControl/Source/PhysicsControl/Public/PhysicsControlRecord.h
+++ b/Engine/Plugins/Experimental/PhysicsControl/Source/PhysicsControl/Public/PhysicsControlRecord.h
@@ -15,10 +15,16 @@ class UMeshComponent;
  */
 struct FPhysicsControlRecord
 {
-	FPhysicsControlRecord(const FPhysicsControl& InControl) : PhysicsControl(InControl) {}
-
-	// Indicates if the control is enabled
-	bool Enabled() const { return PhysicsControl.ControlData.bEnabled; }
+	FPhysicsControlRecord(
+		const FPhysicsControl&       InControl,
+		const FPhysicsControlTarget& InControlTarget,
+		UMeshComponent*              InParentMeshComponent,
+		UMeshComponent*              InChildMeshComponent)
+		: PhysicsControl(InControl)
+		, ControlTarget(InControlTarget)
+		, ParentMeshComponent(InParentMeshComponent)
+		, ChildMeshComponent(InChildMeshComponent)
+	{}
 
 	/** Removes any constraint and resets the state */
 	void ResetConstraint();
@@ -40,23 +46,35 @@ struct FPhysicsControlRecord
 	/** The configuration data */
 	FPhysicsControl PhysicsControl;
 
+	/**
+	 * The position/orientation etc targets for the controls. These are procedural/explicit control targets -
+	 * skeletal meshes have the option to use skeletal animation as well, in which case these targets are 
+	 * expressed as relative to that animation.
+	 */
+	UPROPERTY(EditAnywhere, BlueprintReadWrite, Category = PhysicsControl)
+	FPhysicsControlTarget ControlTarget;
+
+	/**  The mesh that will be doing the driving. Blank/non-existent means it will happen in world space */
+	TWeakObjectPtr<UMeshComponent> ParentMeshComponent;
+
+	/** The mesh that the control will be driving. */
+	TWeakObjectPtr<UMeshComponent> ChildMeshComponent;
+
 	/** The underlying constraint used to implement the control. */
 	TSharedPtr<FConstraintInstance> ConstraintInstance;
-
 };
 
 /**
  * There will be a PhysicsBodyModifier created at runtime for every BodyInstance involved in the component
  */
-struct FPhysicsBodyModifier
+struct FPhysicsBodyModifierRecord
 {
-	FPhysicsBodyModifier(
+	FPhysicsBodyModifierRecord(
 		TWeakObjectPtr<UMeshComponent>  InMeshComponent, 
-		const FName&                InBoneName, 
-		FPhysicsControlModifierData InBodyModifierData)
+		const FName&                    InBoneName, 
+		FPhysicsControlModifierData     InBodyModifierData)
 		: MeshComponent(InMeshComponent)
-		, BoneName(InBoneName)
-		, BodyModifierData(InBodyModifierData)
+		, BodyModifier(InBoneName, InBodyModifierData)
 		, KinematicTargetPosition(FVector::ZeroVector)
 		, KinematicTargetOrientation(FQuat::Identity)
 		, bResetToCachedTarget(false)
@@ -65,11 +83,8 @@ struct FPhysicsBodyModifier
 	/**  The mesh that will be modified. */
 	TWeakObjectPtr<UMeshComponent> MeshComponent;
 
-	/** The name of the skeletal mesh bone or the name of the static mesh body that will be modified. */
-	FName BoneName;
-
-	/** Basic data about how the body should be modified */
-	FPhysicsControlModifierData BodyModifierData;
+	// The core data
+	FPhysicsBodyModifier BodyModifier;
 
 	/** 
 	 * The target position when kinematic. Note that this is applied on top of any animation 
diff --git a/Engine/Plugins/Experimental/PhysicsControl/Source/PhysicsControl/Public/RigidBodyControlData.h b/Engine/Plugins/Experimental/PhysicsControl/Source/PhysicsControl/Public/RigidBodyControlData.h
--- a/Engine/Plugins/Experimental/PhysicsControl/Source/PhysicsControl/Public/RigidBodyControlData.h
+++ b/Engine/Plugins/Experimental/PhysicsControl/Source/PhysicsControl/Public/RigidBodyControlData.h
@@ -2,10 +2,14 @@
 #pragma once
 
 #include "CoreMinimal.h"
+
 #include "PhysicsControlLimbData.h"
+#include "PhysicsControlData.h"
+#include "RigidBodyPoseData.h"
+
 #include "Physics/ImmediatePhysics/ImmediatePhysicsDeclares.h"
 #include "Animation/AnimTypes.h"
-#include "RigidBodyPoseData.h"
+
 #include "RigidBodyControlData.generated.h"
 
 /**
@@ -14,113 +18,6 @@
  * PhysicsControlComponent, so have a PhysicsControl prefix in anticipation of that.
  */
 
-/**
- * Setup data for all the bodies controlled by the node. Contains info to split the skeleton up into limbs, 
- * and default control and modifier settings for each of them.
- */
-USTRUCT(BlueprintType)
-struct FRigidBodySetupData
-{
-	GENERATED_BODY();
-
-	UPROPERTY(EditAnywhere, Category = ControlSetup)
-	TArray<FPhysicsControlLimbSetupData> LimbSetupData;
-
-	UPROPERTY(EditAnywhere, Category = ControlSetup)
-	FPhysicsControlData DefaultWorldSpaceControlData;
-
-	UPROPERTY(EditAnywhere, Category = ControlSetup)
-	FPhysicsControlData DefaultParentSpaceControlData;
-
-	UPROPERTY(EditAnywhere, Category = ControlSetup)
-	FPhysicsControlModifierData DefaultBodyModifierData;
-};
-
-
-/** 
- * Specifies a single/individual control between parent/child bodies, with the control data
- */
-USTRUCT(BlueprintType)
-struct PHYSICSCONTROL_API FRigidBodyControl
-{
-	GENERATED_BODY();
-
-	UPROPERTY(EditAnywhere, BlueprintReadWrite, Category = PhysicsControl)
-	FName ParentBoneName;
-
-	UPROPERTY(EditAnywhere, BlueprintReadWrite, Category = PhysicsControl)
-	FName ChildBoneName;
-
-	UPROPERTY(EditAnywhere, BlueprintReadWrite, Category = PhysicsControl)
-	FPhysicsControlData ControlData;
-
-	bool IsEnabled() const { return ControlData.bEnabled; }
-};
-
-/**
- * Used on creation, to allow requesting the control to be in certain sets
- */
-USTRUCT(BlueprintType)
-struct PHYSICSCONTROL_API FRigidBodyControlCreation
-{
-	GENERATED_BODY();
-
-	UPROPERTY(EditAnywhere, BlueprintReadWrite, Category = PhysicsControl)
-	FRigidBodyControl Control;
-
-	UPROPERTY(EditAnywhere, BlueprintReadWrite, Category = PhysicsControl)
-	TArray<FName> Sets;
-};
-
-/**
- * Used to create a single/individual body modifier, with the modifier data,
- * and what sets it should be added to.
- */
-USTRUCT(BlueprintType)
-struct PHYSICSCONTROL_API FRigidBodyModifier
-{
-	GENERATED_BODY();
-
-	UPROPERTY(EditAnywhere, BlueprintReadWrite, Category = PhysicsControl)
-	FName BoneName;
-
-	UPROPERTY(EditAnywhere, BlueprintReadWrite, Category = PhysicsControl)
-	FPhysicsControlModifierData ModifierData;
-
-	UPROPERTY(EditAnywhere, BlueprintReadWrite, Category = PhysicsControl)
-	TArray<FName> Sets;
-};
-
-/**
- * Used on creation, to allow requesting the modifier to be in certain sets
- */
-USTRUCT(BlueprintType)
-struct PHYSICSCONTROL_API FRigidBodyModifierCreation
-{
-	GENERATED_BODY();
-
-	UPROPERTY(EditAnywhere, BlueprintReadWrite, Category = PhysicsControl)
-	FRigidBodyModifier Modifier;
-
-	UPROPERTY(EditAnywhere, BlueprintReadWrite, Category = PhysicsControl)
-	TArray<FName> Sets;
-};
-
-/**
- * Collection of controls and body modifiers, used for creation
- */
-USTRUCT(BlueprintType)
-struct PHYSICSCONTROL_API FRigidBodyControlAndBodyModifierCreations
-{
-	GENERATED_BODY();
-
-	UPROPERTY(EditAnywhere, BlueprintReadWrite, Category = PhysicsControl)
-	TArray<FRigidBodyControlCreation> Controls;
-
-	UPROPERTY(EditAnywhere, BlueprintReadWrite, Category = PhysicsControl)
-	TArray<FRigidBodyModifierCreation> Modifiers;
-};
-
 /**
  * A single target for a control, which may be defined as an offset from the (implicit) animation target.
  */
@@ -205,7 +102,7 @@ struct PHYSICSCONTROL_API FRigidBodyKinematicTargets
  */
 struct FRigidBodyControlRecord
 {
-	FRigidBodyControlRecord(const FRigidBodyControl& InControl, ImmediatePhysics::FJointHandle* InJointHandle);
+	FRigidBodyControlRecord(const FPhysicsControl& InControl, ImmediatePhysics::FJointHandle* InJointHandle);
 
 	void ResetCurrent(bool bResetTarget);
 
@@ -215,13 +112,8 @@ struct FRigidBodyControlRecord
 	/** Returns the control point, which may be custom or automatic (centre of mass) */
 	FVector GetControlPoint(const ImmediatePhysics::FActorHandle* ChildActorHandle) const;
 
-	// TODO - might benefit from smaller - non-blueprint data members here ?
-	FRigidBodyControl               Control;
-	ImmediatePhysics::FJointHandle* JointHandle;
-
-	// This contains the currently active control data. It will be updated just prior to
-	// applying the controls, by setting it to the default, and then updating it with any parameters.
-	FPhysicsControlData ControlData;
+	// The configuration data
+	FPhysicsControl Control;
 
 	// Contains any control target that has been set
 	FRigidBodyControlTarget ControlTarget;
@@ -234,11 +126,24 @@ struct FRigidBodyControlRecord
 	// TODO just store the count we're interested in rather than the whole structure
 	FGraphTraversalCounter ExpectedUpdateCounter;
 
+	// This is the currently active control data. It will be updated just prior to applying
+	// the controls, by setting it to the default, and then updating it with any parameters. This is
+	// so that control updates can be made ephemeral.
+	FPhysicsControlData ControlData;
+
+	// This is the currently active control multiplier. It will be updated just prior to applying
+	// the controls, by setting it to the default, and then updating it with any parameters. This is
+	// so that control multiplier updates can be made ephemeral.
+	FPhysicsControlMultiplier ControlMultiplier;
+
 	// Cached child body index - needs to be updated whenever the bone name changes
 	int32 ChildBodyIndex;
 
 	// Cached parent body index - needs to be updated whenever the bone name changes
 	int32 ParentBodyIndex;
+
+	// The low-level physics representation
+	ImmediatePhysics::FJointHandle* JointHandle;
 };
 
 /**
@@ -248,16 +153,16 @@ struct FRigidBodyControlRecord
  */
 struct FRigidBodyModifierRecord
 {
-	FRigidBodyModifierRecord(const FRigidBodyModifier& InModifier, ImmediatePhysics::FActorHandle* InActorHandle);
+	FRigidBodyModifierRecord(const FPhysicsBodyModifier& InModifier, ImmediatePhysics::FActorHandle* InActorHandle);
 
 	void ResetCurrent();
 
-	// TODO - might benefit from smaller - non-blueprint data members here ?
-	FRigidBodyModifier              Modifier;
-	ImmediatePhysics::FActorHandle* ActorHandle;
+	FPhysicsBodyModifier            Modifier;
 
 	// This contains the currently active modifier data. It will be updated just prior to
 	// applying the controls, by setting it to the default, and then updating it with any parameters.
 	FPhysicsControlModifierData ModifierData;
+
+	ImmediatePhysics::FActorHandle* ActorHandle;
 };
 
diff --git a/Engine/Plugins/Experimental/PhysicsControl/Source/PhysicsControl/Public/RigidBodyPoseData.h b/Engine/Plugins/Experimental/PhysicsControl/Source/PhysicsControl/Public/RigidBodyPoseData.h
--- a/Engine/Plugins/Experimental/PhysicsControl/Source/PhysicsControl/Public/RigidBodyPoseData.h
+++ b/Engine/Plugins/Experimental/PhysicsControl/Source/PhysicsControl/Public/RigidBodyPoseData.h
@@ -88,11 +88,6 @@ struct FRigidBodyPoseData
 	 */
 	TArray<FPosQuat> BoneTMs;
 
-	/**
-	 * The origin (in world space). BoneData transforms are relative to this.
-	 */
-	FTransform WorldCoordinateFrame;
-
 	// Track when we were currently/last updated so the user can detect missing updates if calculating
 	// velocity etc
 	FGraphTraversalCounter UpdateCounter;
diff --git a/Engine/Plugins/Experimental/PhysicsControl/Source/PhysicsControlEditor/Private/OperatorEditor/SOperatorTreeWidget.cpp b/Engine/Plugins/Experimental/PhysicsControl/Source/PhysicsControlEditor/Private/OperatorEditor/SOperatorTreeWidget.cpp
--- a/Engine/Plugins/Experimental/PhysicsControl/Source/PhysicsControlEditor/Private/OperatorEditor/SOperatorTreeWidget.cpp
+++ b/Engine/Plugins/Experimental/PhysicsControl/Source/PhysicsControlEditor/Private/OperatorEditor/SOperatorTreeWidget.cpp
@@ -315,7 +315,7 @@ void SOperatorTreeWidget::Refresh()
 			// Create node item
 			NodeItemPtr Node = MakeShared<OperatorTreeNodeItem>(*Itr, FName(*NodeUniqueName), BlueprintName);
 
-			if (!Itr->Node.OverridePhysicsAsset)
+			if (!IsValid(Itr->Node.OverridePhysicsAsset))
 			{
 				Node->AddChild(MakeShared<OperatorTreeMessageItem>(*Itr, OperatorTreeMessageItem::Warning, LOCTEXT("OperatorNoPhysicsAssetDefinedText", "Can't create operators for limbs because no Physics Asset has been defined for this node.")));
 			}
diff --git a/Engine/Plugins/Experimental/PhysicsControl/Source/PhysicsControlEditor/Private/PhysicsControlEditorModule.cpp b/Engine/Plugins/Experimental/PhysicsControl/Source/PhysicsControlEditor/Private/PhysicsControlEditorModule.cpp
--- a/Engine/Plugins/Experimental/PhysicsControl/Source/PhysicsControlEditor/Private/PhysicsControlEditorModule.cpp
+++ b/Engine/Plugins/Experimental/PhysicsControl/Source/PhysicsControlEditor/Private/PhysicsControlEditorModule.cpp
@@ -17,18 +17,14 @@ static const FName PhysicsControlEditorModule_PhysicsControlEditorInterface("Phy
 //======================================================================================================================
 void FPhysicsControlEditorModule::StartupModule()
 {
-#ifdef ENABLE_PHYSICS_CONTROL_PROFILE_ASSET
 	IAssetTools& AssetTools = FModuleManager::LoadModuleChecked<FAssetToolsModule>("AssetTools").Get();
 	PhysicsControlProfileAssetActions = MakeShared<FPhysicsControlProfileAssetActions>();
 	AssetTools.RegisterAssetTypeActions(PhysicsControlProfileAssetActions.ToSharedRef());
-#endif
 
-#if ENABLE_PHYSICS_CONTROL_PROFILE_EDITOR
 	FEditorModeRegistry::Get().RegisterMode<FPhysicsControlProfileEditorMode>(
 		FPhysicsControlProfileEditorMode::ModeName, 
 		LOCTEXT("PhysicsControlProfileEditorMode", "PhysicsControlProfile"), 
 		FSlateIcon(), false);
-#endif
 
 	if (GUnrealEd)
 	{
@@ -64,17 +60,13 @@ void FPhysicsControlEditorModule::ShutdownModule()
 	}
 	
 	// Physics Control Profile editor/asset is disabled for now
-#if ENABLE_PHYSICS_CONTROL_PROFILE_EDITOR
 	FEditorModeRegistry::Get().UnregisterMode(FPhysicsControlProfileEditorMode::ModeName);
-#endif
 
-#ifdef ENABLE_PHYSICS_CONTROL_PROFILE_ASSET
 	if (FModuleManager::Get().IsModuleLoaded("AssetTools"))
 	{
 		FAssetToolsModule::GetModule().Get().UnregisterAssetTypeActions(
 			PhysicsControlProfileAssetActions.ToSharedRef());
 	}
-#endif
 
 	if (GUnrealEd)
 	{
diff --git a/Engine/Plugins/Experimental/PhysicsControl/Source/PhysicsControlEditor/Private/PhysicsControlProfileAssetActions.cpp b/Engine/Plugins/Experimental/PhysicsControl/Source/PhysicsControlEditor/Private/PhysicsControlProfileAssetActions.cpp
--- a/Engine/Plugins/Experimental/PhysicsControl/Source/PhysicsControlEditor/Private/PhysicsControlProfileAssetActions.cpp
+++ b/Engine/Plugins/Experimental/PhysicsControl/Source/PhysicsControlEditor/Private/PhysicsControlProfileAssetActions.cpp
@@ -29,7 +29,6 @@ uint32 FPhysicsControlProfileAssetActions::GetCategories()
 	return EAssetTypeCategories::Physics;
 }
 
-#if ENABLE_PHYSICS_CONTROL_PROFILE_EDITOR
 //======================================================================================================================
 void FPhysicsControlProfileAssetActions::OpenAssetEditor(
 	const TArray<UObject*>& InObjects, TSharedPtr<IToolkitHost> EditWithinLevelEditor)
@@ -46,4 +45,3 @@ void FPhysicsControlProfileAssetActions::OpenAssetEditor(
 		}
 	}
 }
-#endif
diff --git a/Engine/Plugins/Experimental/PhysicsControl/Source/PhysicsControlEditor/Private/PhysicsControlProfileAssetActions.h b/Engine/Plugins/Experimental/PhysicsControl/Source/PhysicsControlEditor/Private/PhysicsControlProfileAssetActions.h
--- a/Engine/Plugins/Experimental/PhysicsControl/Source/PhysicsControlEditor/Private/PhysicsControlProfileAssetActions.h
+++ b/Engine/Plugins/Experimental/PhysicsControl/Source/PhysicsControlEditor/Private/PhysicsControlProfileAssetActions.h
@@ -5,11 +5,6 @@
 #include "CoreMinimal.h"
 #include "AssetTypeActions_Base.h"
 
-#define ENABLE_PHYSICS_CONTROL_PROFILE_ASSET
-
-// Enable/disable the ability to actually create control profile assets in editor
-#define ENABLE_PHYSICS_CONTROL_PROFILE_EDITOR 0
-
 class FPhysicsControlProfileAssetActions : public FAssetTypeActions_Base
 {
 public:
@@ -18,9 +13,7 @@ public:
 	FColor GetTypeColor() const override;
 	uint32 GetCategories() override;
 
-#if ENABLE_PHYSICS_CONTROL_PROFILE_EDITOR
 	void OpenAssetEditor(
 		const TArray<UObject*>& InObjects, 
 		TSharedPtr<class IToolkitHost> EditWithinLevelEditor = TSharedPtr<IToolkitHost>()) override;
-#endif
 };
diff --git a/Engine/Plugins/Experimental/PhysicsControl/Source/PhysicsControlUncookedOnly/Private/AnimGraphNode_RigidBodyWithControl.cpp b/Engine/Plugins/Experimental/PhysicsControl/Source/PhysicsControlUncookedOnly/Private/AnimGraphNode_RigidBodyWithControl.cpp
--- a/Engine/Plugins/Experimental/PhysicsControl/Source/PhysicsControlUncookedOnly/Private/AnimGraphNode_RigidBodyWithControl.cpp
+++ b/Engine/Plugins/Experimental/PhysicsControl/Source/PhysicsControlUncookedOnly/Private/AnimGraphNode_RigidBodyWithControl.cpp
@@ -34,12 +34,16 @@ UAnimGraphNode_RigidBodyWithControl::UAnimGraphNode_RigidBodyWithControl(const F
 
 FText UAnimGraphNode_RigidBodyWithControl::GetControllerDescription() const
 {
-	return LOCTEXT("AnimGraphNode_RigidBodyWithControl_ControllerDescription", "Rigid body simulation with Control for physics asset");
+	return LOCTEXT(
+		"AnimGraphNode_RigidBodyWithControl_ControllerDescription", 
+		"Rigid body simulation with Control for physics asset");
 }
 
 FText UAnimGraphNode_RigidBodyWithControl::GetTooltipText() const
 {
-	return LOCTEXT("AnimGraphNode_RigidBodyWithControl_Tooltip", "This simulates based on the skeletal mesh component's physics asset with control options");
+	return LOCTEXT(
+		"AnimGraphNode_RigidBodyWithControl_Tooltip", 
+		"This simulates based on the skeletal mesh component's physics asset with control options");
 }
 
 FText UAnimGraphNode_RigidBodyWithControl::GetNodeTitle(ENodeTitleType::Type TitleType) const
@@ -57,21 +61,29 @@ FString UAnimGraphNode_RigidBodyWithControl::GetNodeCategory() const
 	return TEXT("Animation|Dynamics");
 }
 
-void UAnimGraphNode_RigidBodyWithControl::ValidateAnimNodeDuringCompilation(USkeleton* ForSkeleton, FCompilerResultsLog& MessageLog)
+void UAnimGraphNode_RigidBodyWithControl::ValidateAnimNodeDuringCompilation(
+	USkeleton* ForSkeleton, FCompilerResultsLog& MessageLog)
 {
 	Super::ValidateAnimNodeDuringCompilation(ForSkeleton, MessageLog);
 }
 
-void UAnimGraphNode_RigidBodyWithControl::Draw(FPrimitiveDrawInterface* PDI, USkeletalMeshComponent* PreviewSkelMeshComp, const bool bIsSelected, const bool bIsPoseWatchEnabled) const
+void UAnimGraphNode_RigidBodyWithControl::Draw(
+	FPrimitiveDrawInterface* PDI, 
+	USkeletalMeshComponent*  PreviewSkelMeshComp, 
+	const bool               bIsSelected, 
+	const bool               bIsPoseWatchEnabled) const
 {
-	if (const FAnimNode_RigidBodyWithControl* const RuntimeRigidBodyNode = GetDebuggedAnimNode<FAnimNode_RigidBodyWithControl>())
+	if (const FAnimNode_RigidBodyWithControl* const RuntimeRigidBodyNode = 
+		GetDebuggedAnimNode<FAnimNode_RigidBodyWithControl>())
 	{
 		if (UPhysicsAsset* const PhysicsAsset = RuntimeRigidBodyNode->GetPhysicsAsset())
 		{
-			IPhysicsAssetRenderInterface& PhysicsAssetRenderInterface = IModularFeatures::Get().GetModularFeature<IPhysicsAssetRenderInterface>("PhysicsAssetRenderInterface");
+			IPhysicsAssetRenderInterface& PhysicsAssetRenderInterface = 
+				IModularFeatures::Get().GetModularFeature<IPhysicsAssetRenderInterface>("PhysicsAssetRenderInterface");
 
 			// Draw Bodies.
-			if (bIsSelected || (bIsPoseWatchEnabled && PoseWatchElementBodies.IsValid() && PoseWatchElementBodies->GetIsVisible()))
+			if (bIsSelected || 
+				(bIsPoseWatchEnabled && PoseWatchElementBodies.IsValid() && PoseWatchElementBodies->GetIsVisible()))
 			{
 				FColor PrimitiveColorOverride = FColor::Transparent;
 
@@ -82,11 +94,13 @@ void UAnimGraphNode_RigidBodyWithControl::Draw(FPrimitiveDrawInterface* PDI, USk
 					PrimitiveColorOverride.A = 255;
 				}
 
-				PhysicsAssetRenderInterface.DebugDrawBodies(PreviewSkelMeshComp, PhysicsAsset, PDI, PrimitiveColorOverride);
+				PhysicsAssetRenderInterface.DebugDrawBodies(
+					PreviewSkelMeshComp, PhysicsAsset, PDI, PrimitiveColorOverride);
 			}
 
 			// Draw Constraints.
-			if (bIsSelected || (bIsPoseWatchEnabled && PoseWatchElementConstraints.IsValid() && PoseWatchElementConstraints->GetIsVisible()))
+			if (bIsSelected || 
+				(bIsPoseWatchEnabled && PoseWatchElementConstraints.IsValid() && PoseWatchElementConstraints->GetIsVisible()))
 			{
 				PhysicsAssetRenderInterface.DebugDrawConstraints(PreviewSkelMeshComp, PhysicsAsset, PDI);
 			}
@@ -94,7 +108,9 @@ void UAnimGraphNode_RigidBodyWithControl::Draw(FPrimitiveDrawInterface* PDI, USk
 		
 		
 		// Draw Space Controls
-		if (bIsSelected || (bIsPoseWatchEnabled && PoseWatchElementWorldSpaceControls.IsValid() && PoseWatchElementWorldSpaceControls->GetIsVisible()))
+		if (bIsSelected || 
+			(bIsPoseWatchEnabled && 
+				PoseWatchElementWorldSpaceControls.IsValid() && PoseWatchElementWorldSpaceControls->GetIsVisible()))
 		{
 			for (const TMap<FName, FRigidBodyControlRecord>::ElementType& NameRecordPair : RuntimeRigidBodyNode->ControlRecords)
 			{
@@ -113,7 +129,11 @@ void UAnimGraphNode_RigidBodyWithControl::Draw(FPrimitiveDrawInterface* PDI, USk
 	}
 }
 
-void UAnimGraphNode_RigidBodyWithControl::OnPoseWatchChanged(const bool IsPoseWatchEnabled, TObjectPtr<UPoseWatch> InPoseWatch, FEditorModeTools& InModeTools, FAnimNode_Base* InRuntimeNode)
+void UAnimGraphNode_RigidBodyWithControl::OnPoseWatchChanged(
+	const bool             IsPoseWatchEnabled, 
+	TObjectPtr<UPoseWatch> InPoseWatch, 
+	FEditorModeTools&      InModeTools, 
+	FAnimNode_Base*        InRuntimeNode)
 {
 	Super::OnPoseWatchChanged(IsPoseWatchEnabled, InPoseWatch, InModeTools, InRuntimeNode);
 
@@ -122,10 +142,18 @@ void UAnimGraphNode_RigidBodyWithControl::OnPoseWatchChanged(const bool IsPoseWa
 	if (PoseWatch)
 	{
 		// A new pose watch has been created for this node - add node specific pose watch components.
-		PoseWatchElementBodies = InPoseWatch.Get()->FindOrAddElement(FText(LOCTEXT("PoseWatchElementLabel_RigidBodyWithControl_PhysicsBodies", "Physics Bodies")), TEXT("PhysicsAssetEditor.Tree.Body"));
-		PoseWatchElementConstraints = InPoseWatch.Get()->FindOrAddElement(FText(LOCTEXT("PoseWatchElementLabel__RigidBodyWithControl_PhysicsConstraints", "Physics Constraints")), TEXT("PhysicsAssetEditor.Tree.Constraint"));
-		PoseWatchElementParentSpaceControls = InPoseWatch.Get()->FindOrAddElement(FText(LOCTEXT("PoseWatchElementLabel_RigidBodyWithControl_ParentSpaceControls", "Parent Space Controls")), TEXT("PhysicsAssetEditor.Tree.Body"));
-		PoseWatchElementWorldSpaceControls = InPoseWatch.Get()->FindOrAddElement(FText(LOCTEXT("PoseWatchElementLabel_RigidBodyWithControl_WorldSpaceControls", "World Space Controls")), TEXT("PhysicsAssetEditor.Tree.Body"));
+		PoseWatchElementBodies = InPoseWatch.Get()->FindOrAddElement(FText(
+			LOCTEXT("PoseWatchElementLabel_RigidBodyWithControl_PhysicsBodies", "Physics Bodies")), 
+			TEXT("PhysicsAssetEditor.Tree.Body"));
+		PoseWatchElementConstraints = InPoseWatch.Get()->FindOrAddElement(FText(
+			LOCTEXT("PoseWatchElementLabel__RigidBodyWithControl_PhysicsConstraints", "Physics Constraints")), 
+			TEXT("PhysicsAssetEditor.Tree.Constraint"));
+		PoseWatchElementParentSpaceControls = InPoseWatch.Get()->FindOrAddElement(FText(
+			LOCTEXT("PoseWatchElementLabel_RigidBodyWithControl_ParentSpaceControls", "Parent Space Controls")), 
+			TEXT("PhysicsAssetEditor.Tree.Body"));
+		PoseWatchElementWorldSpaceControls = InPoseWatch.Get()->FindOrAddElement(FText(
+			LOCTEXT("PoseWatchElementLabel_RigidBodyWithControl_WorldSpaceControls", "World Space Controls")), 
+			TEXT("PhysicsAssetEditor.Tree.Body"));
 
 		check(PoseWatchElementConstraints.IsValid()); // Expect to find a valid component;
 		if (PoseWatchElementConstraints.IsValid())
@@ -138,13 +166,12 @@ void UAnimGraphNode_RigidBodyWithControl::OnPoseWatchChanged(const bool IsPoseWa
 void UAnimGraphNode_RigidBodyWithControl::CustomizeDetails(IDetailLayoutBuilder& DetailBuilder)
 {
 	Super::CustomizeDetails(DetailBuilder);
-
 	IDetailCategoryBuilder& ViewportCategory = DetailBuilder.EditCategory(TEXT("Debug Visualization"));
-	
 	FAnimNode_RigidBodyWithControl* const RigidBodyNode = static_cast<FAnimNode_RigidBodyWithControl*>(GetDebuggedAnimNode());
 	
 	{
-		FDetailWidgetRow& ControlsEditorWidgetRow = ViewportCategory.AddCustomRow(LOCTEXT("ToggleControlsEditorWidgetRowButtonRow", "ControlsEditor"));
+		FDetailWidgetRow& ControlsEditorWidgetRow = ViewportCategory.AddCustomRow(
+			LOCTEXT("ToggleControlsEditorWidgetRowButtonRow", "ControlsEditor"));
 
 		ControlsEditorWidgetRow
 			[
@@ -167,7 +194,8 @@ void UAnimGraphNode_RigidBodyWithControl::CustomizeDetails(IDetailLayoutBuilder&
 	}
 
 	{
-		FDetailWidgetRow& DebugVisualizationWidgetRow = ViewportCategory.AddCustomRow(LOCTEXT("ToggleDebugVisualizationButtonRow", "DebugVisualization"));
+		FDetailWidgetRow& DebugVisualizationWidgetRow = ViewportCategory.AddCustomRow(
+			LOCTEXT("ToggleDebugVisualizationButtonRow", "DebugVisualization"));
 
 		DebugVisualizationWidgetRow
 			[
@@ -208,7 +236,8 @@ void UAnimGraphNode_RigidBodyWithControl::CustomizeDetails(IDetailLayoutBuilder&
 
 void UAnimGraphNode_RigidBodyWithControl::PostChange()
 {
-	IPhysicsControlOperatorEditorInterface& PhysicsControlEditorInterface = IModularFeatures::Get().GetModularFeature<IPhysicsControlOperatorEditorInterface>("PhysicsControlEditorInterface");
+	IPhysicsControlOperatorEditorInterface& PhysicsControlEditorInterface = 
+		IModularFeatures::Get().GetModularFeature<IPhysicsControlOperatorEditorInterface>("PhysicsControlEditorInterface");
 	PhysicsControlEditorInterface.RequestRefresh();
 }
 
@@ -216,7 +245,8 @@ void UAnimGraphNode_RigidBodyWithControl::PostEditChangeProperty(struct FPropert
 {
 	Super::PostEditChangeProperty(PropertyChangedEvent);
 
-	IPhysicsAssetRenderInterface& PhysicsAssetRenderInterface = IModularFeatures::Get().GetModularFeature<IPhysicsAssetRenderInterface>("PhysicsAssetRenderInterface");
+	IPhysicsAssetRenderInterface& PhysicsAssetRenderInterface = 
+		IModularFeatures::Get().GetModularFeature<IPhysicsAssetRenderInterface>("PhysicsAssetRenderInterface");
 	PhysicsAssetRenderInterface.SaveConfig();
 
 	PostChange();
@@ -242,9 +272,11 @@ void UAnimGraphNode_RigidBodyWithControl::DestroyNode()
 
 void UAnimGraphNode_RigidBodyWithControl::ToggleBodyVisibility()
 {
-	FAnimNode_RigidBodyWithControl* const RigidBodyNode = static_cast<FAnimNode_RigidBodyWithControl*>(GetDebuggedAnimNode());
+	FAnimNode_RigidBodyWithControl* const RigidBodyNode = 
+		static_cast<FAnimNode_RigidBodyWithControl*>(GetDebuggedAnimNode());
 	
-	IPhysicsAssetRenderInterface& PhysicsAssetRenderInterface = IModularFeatures::Get().GetModularFeature<IPhysicsAssetRenderInterface>("PhysicsAssetRenderInterface");
+	IPhysicsAssetRenderInterface& PhysicsAssetRenderInterface = 
+		IModularFeatures::Get().GetModularFeature<IPhysicsAssetRenderInterface>("PhysicsAssetRenderInterface");
 
 	if (RigidBodyNode)
 	{
@@ -254,8 +286,10 @@ void UAnimGraphNode_RigidBodyWithControl::ToggleBodyVisibility()
 
 void UAnimGraphNode_RigidBodyWithControl::ToggleConstraintVisibility()
 {
-	FAnimNode_RigidBodyWithControl* const RigidBodyNode = static_cast<FAnimNode_RigidBodyWithControl*>(GetDebuggedAnimNode());
-	IPhysicsAssetRenderInterface& PhysicsAssetRenderInterface = IModularFeatures::Get().GetModularFeature<IPhysicsAssetRenderInterface>("PhysicsAssetRenderInterface");
+	FAnimNode_RigidBodyWithControl* const RigidBodyNode = 
+		static_cast<FAnimNode_RigidBodyWithControl*>(GetDebuggedAnimNode());
+	IPhysicsAssetRenderInterface& PhysicsAssetRenderInterface = 
+		IModularFeatures::Get().GetModularFeature<IPhysicsAssetRenderInterface>("PhysicsAssetRenderInterface");
 
 	if (RigidBodyNode)
 	{
@@ -265,8 +299,10 @@ void UAnimGraphNode_RigidBodyWithControl::ToggleConstraintVisibility()
 
 bool UAnimGraphNode_RigidBodyWithControl::AreAnyBodiesHidden() const
 {
-	FAnimNode_RigidBodyWithControl* const RigidBodyNode = static_cast<FAnimNode_RigidBodyWithControl*>(GetDebuggedAnimNode());
-	IPhysicsAssetRenderInterface& PhysicsAssetRenderInterface = IModularFeatures::Get().GetModularFeature<IPhysicsAssetRenderInterface>("PhysicsAssetRenderInterface");
+	FAnimNode_RigidBodyWithControl* const RigidBodyNode = 
+		static_cast<FAnimNode_RigidBodyWithControl*>(GetDebuggedAnimNode());
+	IPhysicsAssetRenderInterface& PhysicsAssetRenderInterface = 
+		IModularFeatures::Get().GetModularFeature<IPhysicsAssetRenderInterface>("PhysicsAssetRenderInterface");
 
 	if (RigidBodyNode)
 	{
@@ -278,8 +314,10 @@ bool UAnimGraphNode_RigidBodyWithControl::AreAnyBodiesHidden() const
 
 bool UAnimGraphNode_RigidBodyWithControl::AreAnyConstraintsHidden() const
 {
-	FAnimNode_RigidBodyWithControl* const RigidBodyNode = static_cast<FAnimNode_RigidBodyWithControl*>(GetDebuggedAnimNode());
-	IPhysicsAssetRenderInterface& PhysicsAssetRenderInterface = IModularFeatures::Get().GetModularFeature<IPhysicsAssetRenderInterface>("PhysicsAssetRenderInterface");
+	FAnimNode_RigidBodyWithControl* const RigidBodyNode = 
+		static_cast<FAnimNode_RigidBodyWithControl*>(GetDebuggedAnimNode());
+	IPhysicsAssetRenderInterface& PhysicsAssetRenderInterface = 
+		IModularFeatures::Get().GetModularFeature<IPhysicsAssetRenderInterface>("PhysicsAssetRenderInterface");
 
 	if (RigidBodyNode)
 	{
@@ -291,13 +329,15 @@ bool UAnimGraphNode_RigidBodyWithControl::AreAnyConstraintsHidden() const
 
 void UAnimGraphNode_RigidBodyWithControl::ToggleControlEditorTab()
 {
-	IPhysicsControlOperatorEditorInterface& PhysicsControlEditorInterface = IModularFeatures::Get().GetModularFeature<IPhysicsControlOperatorEditorInterface>("PhysicsControlEditorInterface");
+	IPhysicsControlOperatorEditorInterface& PhysicsControlEditorInterface = 
+		IModularFeatures::Get().GetModularFeature<IPhysicsControlOperatorEditorInterface>("PhysicsControlEditorInterface");
 	PhysicsControlEditorInterface.ToggleOperatorNamesTab();
 }
 
 bool UAnimGraphNode_RigidBodyWithControl::IsControlEditorTabOpen() const
 {
-	IPhysicsControlOperatorEditorInterface& PhysicsControlEditorInterface = IModularFeatures::Get().GetModularFeature<IPhysicsControlOperatorEditorInterface>("PhysicsControlEditorInterface");
+	IPhysicsControlOperatorEditorInterface& PhysicsControlEditorInterface = 
+		IModularFeatures::Get().GetModularFeature<IPhysicsControlOperatorEditorInterface>("PhysicsControlEditorInterface");
 	return PhysicsControlEditorInterface.IsOperatorNamesTabOpen();
 }
 
@@ -313,18 +353,23 @@ TArray<TPair<FName, TArray<FName>>> UAnimGraphNode_RigidBodyWithControl::Generat
 
 		// These functions will create the base set of controls and modifiers from SetupData
 		TMap<FName, FPhysicsControlLimbBones> AllLimbBones =
-			GetLimbBones(Node.SetupData.LimbSetupData, RefSkeleton, Node.OverridePhysicsAsset.Get());
+			UE::PhysicsControl::GetLimbBones(
+				Node.CharacterSetupData.LimbSetupData, RefSkeleton, Node.OverridePhysicsAsset.Get());
 
 		TSet<FName> BodyModifierNames;
 		TSet<FName> ControlNames;
 		FPhysicsControlNameRecords NameRecords;
 
-		CollectOperatorNames(&Node, AllLimbBones, RefSkeleton, Node.OverridePhysicsAsset.Get(), BodyModifierNames, ControlNames, NameRecords);
+		// TODO This needs to account for the control profile asset
+		UE::PhysicsControl::CollectOperatorNames(
+			&Node, Node.CharacterSetupData, Node.AdditionalControlsAndBodyModifiers,
+			AllLimbBones, RefSkeleton, Node.OverridePhysicsAsset.Get(), BodyModifierNames, ControlNames, NameRecords);
 
 		// Create any additional sets that have been requested
-		CreateAdditionalSets(Node.AdditionalSets, BodyModifierNames, ControlNames, NameRecords);
+		UE::PhysicsControl::CreateAdditionalSets(Node.AdditionalSets, BodyModifierNames, ControlNames, NameRecords);
 
-		auto TransformOperatorNamesAndTags = [&GeneratedOperatorNames](const FName TypeTag, const TSet<FName>& Names, const TMap<FName, TArray<FName>>& SetToOperatorNameMap)
+		auto TransformOperatorNamesAndTags = [&GeneratedOperatorNames](
+			const FName TypeTag, const TSet<FName>& Names, const TMap<FName, TArray<FName>>& SetToOperatorNameMap)
 		{
 			for (const FName OperatorName : Names)
 			{
diff --git a/Engine/Source/Runtime/AnimGraphRuntime/Private/BoneControllers/AnimNode_RigidBody.cpp b/Engine/Source/Runtime/AnimGraphRuntime/Private/BoneControllers/AnimNode_RigidBody.cpp
--- a/Engine/Source/Runtime/AnimGraphRuntime/Private/BoneControllers/AnimNode_RigidBody.cpp
+++ b/Engine/Source/Runtime/AnimGraphRuntime/Private/BoneControllers/AnimNode_RigidBody.cpp
@@ -992,7 +992,7 @@ void ComputeBodyInsertionOrder(TArray<FBoneIndexType>& InsertionOrder, const USk
 
 UPhysicsAsset* FAnimNode_RigidBody::GetPhysicsAssetToBeUsed(const UAnimInstance* InAnimInstance) const
 {
-	if (OverridePhysicsAsset)
+	if (IsValid(OverridePhysicsAsset))
 	{
 		return ToRawPtr(OverridePhysicsAsset);
 	}
