PhysicsControl - support body modifier "default" movement type, so it doesn't force a change to the movement. Fixes problems initialising physics assets that have kinematic bodies. Make sure that disabling controls removes constraints and reverts the RigidBodyWithControl node back to its uncontrolled state.

diff --git a/Engine/Plugins/Experimental/PhysicsControl/Source/PhysicsControl/Private/AnimNode_RigidBodyWithControl.cpp b/Engine/Plugins/Experimental/PhysicsControl/Source/PhysicsControl/Private/AnimNode_RigidBodyWithControl.cpp
--- a/Engine/Plugins/Experimental/PhysicsControl/Source/PhysicsControl/Private/AnimNode_RigidBodyWithControl.cpp
+++ b/Engine/Plugins/Experimental/PhysicsControl/Source/PhysicsControl/Private/AnimNode_RigidBodyWithControl.cpp
@@ -132,10 +132,11 @@ FAnimNode_RigidBodyWithControl::FAnimNode_RigidBodyWithControl()
 	, bUpdateCacheEveryFrame(true)
 	, BaseBoneRef()
 	, OverlapChannel(ECC_WorldStatic)
-	, bCalculateVelocitiesForWorldGeometry(true)
 	, SimulationSpace(ESimulationSpace::ComponentSpace)
+	, bCalculateVelocitiesForWorldGeometry(true)
 	, bForceDisableCollisionBetweenConstraintBodies(false)
 	, bUseExternalClothCollision(false)
+	, bMakeKinematicConstraints(true)
 	, ResetSimulatedTeleportType(ETeleportType::None)
 	, bEnableWorldGeometry(false)
 	, bOverrideWorldGravity(false)
@@ -150,7 +151,7 @@ FAnimNode_RigidBodyWithControl::FAnimNode_RigidBodyWithControl()
 	, bSimulationStarted(false)
 	, bCheckForBodyTransformInit(false)
 	, bHaveSetupControls(false)
-	, SimulationTiming(ESimulationTiming::Default)
+	, SimulationTiming(ESimulationTiming::Synchronous)
 	, WorldTimeSeconds(0.0f)
 	, LastEvalTimeSeconds(0.0f)
 	, AccumulatedDeltaTime(0.0f)
@@ -460,6 +461,7 @@ void FAnimNode_RigidBodyWithControl::DestroyPhysicsSimulation()
 	FlushDeferredSimulationTask();
 	delete PhysicsSimulation;
 	PhysicsSimulation = nullptr;
+	CurrentConstraintProfile = FName();
 }
 
 void FAnimNode_RigidBodyWithControl::SetupControls(USkeletalMeshComponent* const SkeletalMeshComponent)
@@ -513,6 +515,10 @@ void FAnimNode_RigidBodyWithControl::EvaluateSkeletalControl_AnyThread(FComponen
 		{
 			SetupControls(Output.AnimInstanceProxy->GetSkelMeshComponent());
 		}
+		else if (bHaveSetupControls && !bEnableControls)
+		{
+			DestroyControlsAndBodyModifiers();
+		}
 
 		const FBoneContainer& BoneContainer = Output.Pose.GetPose().GetBoneContainer();
 		const FTransform CompWorldSpaceTM = Output.AnimInstanceProxy->GetComponentTransform();
@@ -762,16 +768,19 @@ void FAnimNode_RigidBodyWithControl::EvaluateSkeletalControl_AnyThread(FComponen
 				ApplyCurrentConstraintProfile();
 			}
 
-			if (ControlProfile != CurrentControlProfile)
+			if (bHaveSetupControls)
 			{
-				CurrentControlProfile = ControlProfile;
-				ApplyCurrentControlProfile();
-			}
+				if (ControlProfile != CurrentControlProfile)
+				{
+					CurrentControlProfile = ControlProfile;
+					ApplyCurrentControlProfile();
+				}
 
-			// Apply the controls. Note that these won't set kinematic targets - we'll do that afterwards as they
-			// can apply to bodies that aren't under the influence of a modifier
-			ApplyControlAndModifierUpdatesAndParametersToRecords(ControlAndModifierUpdates, ControlAndModifierParameters);
-			ApplyControlsAndModifiers(SimSpaceGravity, DeltaSeconds);
+				// Apply the controls. Note that these won't set kinematic targets - we'll do that afterwards as they
+				// can apply to bodies that aren't under the influence of a modifier
+				ApplyControlAndModifierUpdatesAndParametersToRecords(ControlAndModifierUpdates, ControlAndModifierParameters);
+				ApplyControlsAndModifiers(SimSpaceGravity, DeltaSeconds);
+			}
 
 			// Note that the simulation interpolates kinematic targets to handle substepping
 			for (const RigidBodyWithControl::FOutputBoneData& OutputData : OutputBoneData)
@@ -1078,7 +1087,12 @@ void FAnimNode_RigidBodyWithControl::InitPhysics(const UAnimInstance* InAnimInst
 					TotalMass += InvMass > 0.f ? 1.f / InvMass : 0.f;
 					if (!bSimulated)
 					{
-						ActorHandle->SetIsKinematic(true);
+						// Note that particles are always created disabled (why?). However,
+						// SetEnabled only operates on dyanmics, so we need to enable the particle
+						// before making it kinematic, otherwise we end up with particles that are
+						// disabled, but still simulate.
+						ActorHandle->SetEnabled(true);
+						PhysicsSimulation->SetIsKinematic(ActorHandle, true);
 					}
 
 					ActorHandle->SetName(BodySetup->BoneName);
@@ -1143,13 +1157,14 @@ void FAnimNode_RigidBodyWithControl::InitPhysics(const UAnimInstance* InAnimInst
 
 				if(Body1Handle && Body2Handle)
 				{
-					if (Body1Handle->IsSimulated() || Body2Handle->IsSimulated())
+					if (bMakeKinematicConstraints || (Body1Handle->IsSimulated() || Body2Handle->IsSimulated()))
 					{
-						ImmediatePhysics::FJointHandle* NewJointHandle =
+						ImmediatePhysics::FJointHandle* JointHandle =
 							PhysicsSimulation->CreateJoint(CI, Body1Handle, Body2Handle);
+
 						// Record the joint handle under the child bone (i.e. more leaf-ward bone),
 						// since each bone may have multiple children, but only one parent.
-						NamesToJointHandles.Add(CI->ConstraintBone1, NewJointHandle);
+						NamesToJointHandles.Add(CI->ConstraintBone1, JointHandle);
 
 						if (bForceDisableCollisionBetweenConstraintBodies)
 						{
@@ -1164,7 +1179,7 @@ void FAnimNode_RigidBodyWithControl::InitPhysics(const UAnimInstance* InAnimInst
 						int32 BodyIndex;
 						if (Bodies.Find(Body1Handle, BodyIndex))
 						{
-							Joints[BodyIndex] = NewJointHandle;
+							Joints[BodyIndex] = JointHandle;
 						}
 
 						if (CI->IsCollisionDisabled())
@@ -1614,7 +1629,7 @@ inline bool IsComponentDesiredInSim(const UPrimitiveComponent* Component)
 	{
 		return false;
 	}
-	if (!Component->GetBodyInstance()->GetCollisionEnabled())
+	if (!CollisionEnabledHasPhysics(Component->GetBodyInstance()->GetCollisionEnabled()))
 	{
 		return false;
 	}
diff --git a/Engine/Plugins/Experimental/PhysicsControl/Source/PhysicsControl/Private/PhysicsControlComponentImpl.cpp b/Engine/Plugins/Experimental/PhysicsControl/Source/PhysicsControl/Private/PhysicsControlComponentImpl.cpp
--- a/Engine/Plugins/Experimental/PhysicsControl/Source/PhysicsControl/Private/PhysicsControlComponentImpl.cpp
+++ b/Engine/Plugins/Experimental/PhysicsControl/Source/PhysicsControl/Private/PhysicsControlComponentImpl.cpp
@@ -739,6 +739,9 @@ void UPhysicsControlComponent::ApplyBodyModifier(FPhysicsBodyModifierRecord& Rec
 		case EPhysicsMovementType::Simulated:
 			BodyInstance->SetInstanceSimulatePhysics(true, false, true);
 			break;
+		case EPhysicsMovementType::Default:
+			// Default means do nothing, so let's do exactly that
+			break;
 		default:
 			UE_LOG(LogPhysicsControl, Warning, TEXT("Invalid movement type %d"),
 				int(Record.BodyModifier.ModifierData.MovementType));
diff --git a/Engine/Plugins/Experimental/PhysicsControl/Source/PhysicsControl/Private/RigidBodyWithControl.cpp b/Engine/Plugins/Experimental/PhysicsControl/Source/PhysicsControl/Private/RigidBodyWithControl.cpp
--- a/Engine/Plugins/Experimental/PhysicsControl/Source/PhysicsControl/Private/RigidBodyWithControl.cpp
+++ b/Engine/Plugins/Experimental/PhysicsControl/Source/PhysicsControl/Private/RigidBodyWithControl.cpp
@@ -357,16 +357,17 @@ void FAnimNode_RigidBodyWithControl::InitControlsAndBodyModifiers(const FReferen
 	TMap<FName, FPhysicsControlLimbBones> AllLimbBones =
 		UE::PhysicsControl::GetLimbBones(SetupData.LimbSetupData, RefSkeleton, GetPhysicsAsset());
 
-	FPhysicsControlAndBodyModifierCreationDatas AdditionalOperators;
+	FPhysicsControlAndBodyModifierCreationDatas ControlAndBodyModifierCreationDatas;
 	if (IsValid(PhysicsControlProfileAsset))
 	{
-		AdditionalOperators = PhysicsControlProfileAsset->AdditionalControlsAndModifiers;
+		ControlAndBodyModifierCreationDatas = PhysicsControlProfileAsset->AdditionalControlsAndModifiers;
 	}
-	AdditionalOperators += AdditionalControlsAndBodyModifiers;
+	ControlAndBodyModifierCreationDatas += AdditionalControlsAndBodyModifiers;
 
 	// An "operator" is a control or a body modifier. This will also add them to sets etc.
 	UE::PhysicsControl::CreateOperatorsForNode(
-		this, SetupData, AdditionalControlsAndBodyModifiers, AllLimbBones, RefSkeleton, GetPhysicsAsset(), NameRecords);
+		this, SetupData, ControlAndBodyModifierCreationDatas, 
+		AllLimbBones, RefSkeleton, GetPhysicsAsset(), NameRecords);
 
 	for (TMap<FName, FRigidBodyControlRecord>::ElementType& NameRecordPair : ControlRecords)
 	{
@@ -395,11 +396,22 @@ void FAnimNode_RigidBodyWithControl::InitControlsAndBodyModifiers(const FReferen
 //======================================================================================================================
 void FAnimNode_RigidBodyWithControl::DestroyControlsAndBodyModifiers()
 {
-	// TODO - make sure this fn is complete
+	// This is needed because deleting the joint handle doesn't actually remove the constraint from
+	// the simulation.
+	for (TMap<FName, FRigidBodyControlRecord>::ElementType& NameRecordPair : ControlRecords)
+	{
+		ImmediatePhysics::FJointHandle* JointHandle = NameRecordPair.Value.JointHandle;
+		if (JointHandle)
+		{
+			PhysicsSimulation->DestroyJoint(JointHandle);
+		}
+	}
+
 	ControlRecords.Reset();
 	ModifierRecords.Reset();
 	NameRecords.Reset();
 	bHaveSetupControls = false;
+	CurrentControlProfile = FName();
 }
 
 //======================================================================================================================
@@ -664,10 +676,13 @@ void FAnimNode_RigidBodyWithControl::ApplyModifier(
 {
 	if (BodyModifierRecord.ActorHandle)
 	{
-		// Note that there's an early out if there's no change needed, so this should be OK.
-		PhysicsSimulation->SetIsKinematic(
-			BodyModifierRecord.ActorHandle,
-			BodyModifierRecord.ModifierData.MovementType != EPhysicsMovementType::Simulated);
+		if (BodyModifierRecord.ModifierData.MovementType != EPhysicsMovementType::Default)
+		{
+			// Note that there's an early out if there's no change needed, so this should be OK.
+			PhysicsSimulation->SetIsKinematic(
+				BodyModifierRecord.ActorHandle,
+				BodyModifierRecord.ModifierData.MovementType != EPhysicsMovementType::Simulated);
+		}
 
 		// Note that the actual kinematic targets will be set separately, since they need to be set
 		// for all kinematics whether or not they were under a modifier.
@@ -682,8 +697,8 @@ void FAnimNode_RigidBodyWithControl::ApplyModifier(
 		}
 
 		// Set collision
-		Chaos::FGeometryParticleHandle* ParticleHandle = BodyModifierRecord.ActorHandle->GetParticle();
-		//ParticleHandle->SetHasCollision(CollisionEnabledHasPhysics(BodyModifierRecord.ModifierData.CollisionType));
+		PhysicsSimulation->SetHasCollision(
+			BodyModifierRecord.ActorHandle, CollisionEnabledHasPhysics(BodyModifierRecord.ModifierData.CollisionType));
 	}
 }
 
diff --git a/Engine/Plugins/Experimental/PhysicsControl/Source/PhysicsControl/Public/AnimNode_RigidBodyWithControl.h b/Engine/Plugins/Experimental/PhysicsControl/Source/PhysicsControl/Public/AnimNode_RigidBodyWithControl.h
--- a/Engine/Plugins/Experimental/PhysicsControl/Source/PhysicsControl/Public/AnimNode_RigidBodyWithControl.h
+++ b/Engine/Plugins/Experimental/PhysicsControl/Source/PhysicsControl/Public/AnimNode_RigidBodyWithControl.h
@@ -179,6 +179,10 @@ public:
 	UPROPERTY(EditAnywhere, Category = Settings, meta = (editcondition = "bEnableWorldGeometry"))
 	TEnumAsByte<ECollisionChannel> OverlapChannel;
 
+	/** What space to simulate the bodies in. This affects how velocities are generated */
+	UPROPERTY(EditAnywhere, Category = Settings)
+	ESimulationSpace SimulationSpace;
+
 	/** 
 	 * Whether or not to calculate velocities for world geometry. Note that if the simulation space is
 	 * not set to world, then even static objects "should" have velocities calculated. There is a cost 
@@ -188,17 +192,20 @@ public:
 	UPROPERTY(EditAnywhere, Category = Settings, meta = (editcondition = "bEnableWorldGeometry"))
 	uint8 bCalculateVelocitiesForWorldGeometry : 1;
 
-	/** What space to simulate the bodies in. This affects how velocities are generated */
-	UPROPERTY(EditAnywhere, Category = Settings)
-	ESimulationSpace SimulationSpace;
-
 	/** Whether to allow collisions between two bodies joined by a constraint  */
 	UPROPERTY(EditAnywhere, Category = Settings)
-	bool bForceDisableCollisionBetweenConstraintBodies;
+	uint8 bForceDisableCollisionBetweenConstraintBodies : 1;
 
 	/** If true, kinematic objects will be added to the simulation at runtime to represent any cloth colliders defined for the parent object. */
 	UPROPERTY(EditAnywhere, Category = Settings)
-	bool bUseExternalClothCollision;
+	uint8 bUseExternalClothCollision : 1;
+
+	/** 
+	 * If true, constraints will be made even when both objects are kinematic, in case they are subsequently 
+	 * made dynamic. There may be a small performance benefit to disabling this if you know this won't happen. 
+	 */
+	UPROPERTY(EditAnywhere, Category = Settings)
+	uint8 bMakeKinematicConstraints : 1;
 
 private:
 	ETeleportType ResetSimulatedTeleportType;
@@ -492,7 +499,9 @@ public:
 
 	/* 
 	 * Whether the physics simulation runs synchronously with the node's evaluation or is run in the 
-	 * background until the next frame. 
+	 * background until the next frame. Note that running with deferred timing can improve performance, 
+	 * but there will be a frame's lag in the output (relative to the input animation), and there can
+	 * be problems with handling collisions.
 	 */
 	UPROPERTY(EditAnywhere, Category=Settings, AdvancedDisplay)
 	ESimulationTiming SimulationTiming;
diff --git a/Engine/Plugins/Experimental/PhysicsControl/Source/PhysicsControl/Public/PhysicsControlData.h b/Engine/Plugins/Experimental/PhysicsControl/Source/PhysicsControl/Public/PhysicsControlData.h
--- a/Engine/Plugins/Experimental/PhysicsControl/Source/PhysicsControl/Public/PhysicsControlData.h
+++ b/Engine/Plugins/Experimental/PhysicsControl/Source/PhysicsControl/Public/PhysicsControlData.h
@@ -26,7 +26,10 @@ enum class EPhysicsMovementType : uint8
 	// kinematic target set in the Body Modifier.
 	Kinematic,
 	// Simulated means that the object will be controlled by the physics solver
-	Simulated
+	Simulated,
+	// Default means that the movement type shouldn't be changed - for example, it will use the
+	// value in the physics asset
+	Default
 };
 
 inline FName GetPhysicsMovementTypeName(const EPhysicsMovementType MovementType)
@@ -39,6 +42,8 @@ inline FName GetPhysicsMovementTypeName(const EPhysicsMovementType MovementType)
 		return "Kinematic";
 	case EPhysicsMovementType::Simulated:
 		return "Simulated";
+	case EPhysicsMovementType::Default:
+		return "Default";
 	}
 	return "None";
 }
