[UE] Changes to RBwC constraint mapping
- Make UI harder to use incorrectly
- Support chooseing to apply transforms from authored SKM or default transforms or nothing to constraint positions and orientations separately.
- Help text has been re-written and now includes an explicit warning about using the default transforms option.

#rnx

diff --git a/Engine/Plugins/Experimental/PhysicsControl/Source/PhysicsControl/Private/AnimNode_RigidBodyWithControl.cpp b/Engine/Plugins/Experimental/PhysicsControl/Source/PhysicsControl/Private/AnimNode_RigidBodyWithControl.cpp
--- a/Engine/Plugins/Experimental/PhysicsControl/Source/PhysicsControl/Private/AnimNode_RigidBodyWithControl.cpp
+++ b/Engine/Plugins/Experimental/PhysicsControl/Source/PhysicsControl/Private/AnimNode_RigidBodyWithControl.cpp
@@ -143,7 +143,8 @@ FAnimNode_RigidBodyWithControl::FAnimNode_RigidBodyWithControl()
 	, bOverrideWorldGravity(false)
 	, bTransferBoneVelocities(false)
 	, bFreezeIncomingPoseOnStart(false)
-	, bModifyConstraintTransformsToMatchSkeleton(false)
+	, PhysicsAssetConditioningConstraintPosition(MapConstraintsBehaviorType::None)
+	, PhysicsAssetConditioningConstraintOrientation(MapConstraintsBehaviorType::None)
 	, WorldSpaceMinimumScale(0.01f)
 	, EvaluationResetTime(0.01f)
 	, bEnableControls(false)
@@ -1028,10 +1029,9 @@ void FAnimNode_RigidBodyWithControl::InitPhysics(const UAnimInstance* InAnimInst
 			FPhysicsAggregateHandle(),
 			bCreateBodiesInRefPose);
 
-		if (bModifyConstraintTransformsToMatchSkeleton && 
-			CVarEnableRigidBodyNodeWithControlMatchingConstraintsToSkeleton.GetValueOnAnyThread() > 0)
+		if (CVarEnableRigidBodyNodeWithControlMatchingConstraintsToSkeleton.GetValueOnAnyThread() > 0)
 		{
-			TransformConstraintsToMatchSkeletalMesh(SkeletalMeshAsset, HighLevelConstraintInstances);
+			TransformConstraintsToMatchSkeletalMesh(SkeletalMeshAsset, PhysicsAssetConditioningConstraintPosition, PhysicsAssetConditioningConstraintOrientation, HighLevelConstraintInstances);
 		}
 
 		TMap<FName, ImmediatePhysics::FActorHandle*> NamesToHandles;
diff --git a/Engine/Plugins/Experimental/PhysicsControl/Source/PhysicsControl/Private/RigidBodyWithControl.cpp b/Engine/Plugins/Experimental/PhysicsControl/Source/PhysicsControl/Private/RigidBodyWithControl.cpp
--- a/Engine/Plugins/Experimental/PhysicsControl/Source/PhysicsControl/Private/RigidBodyWithControl.cpp
+++ b/Engine/Plugins/Experimental/PhysicsControl/Source/PhysicsControl/Private/RigidBodyWithControl.cpp
@@ -130,6 +130,14 @@ template<typename TRecord, typename TParameters> void ApplyControlAndModifierPar
 		}
 	}
 }
+//======================================================================================================================
+FName MapConstraintsBehaviorTypeToString(const MapConstraintsBehaviorType InType)
+{
+	if (InType == MapConstraintsBehaviorType::AuthoredSkeleton) { return FName("AuthoredSkeleton"); }
+	if (InType == MapConstraintsBehaviorType::DefaultTransform) { return FName("DefaultTransform"); }
+	
+	return FName("None");
+}
 
 //======================================================================================================================
 const FTransform FAnimNode_RigidBodyWithControl::GetBodyTransform(const int32 BodyIndex) const
@@ -1132,91 +1140,90 @@ void FAnimNode_RigidBodyWithControl::ApplyCurrentControlProfile()
 
 //======================================================================================================================
 void FAnimNode_RigidBodyWithControl::TransformConstraintsToMatchSkeletalMesh(
-	const USkeletalMesh*          const SkeletalMeshAsset, 
+	const USkeletalMesh* const SkeletalMeshAsset,
+	const MapConstraintsBehaviorType PositionBehavior,
+	const MapConstraintsBehaviorType OrientationBehavior,
 	TArray<FConstraintInstance*>& ConstraintInstances)
 {
 	// Bone1 = Child
 	// Bone2 = Parent 
 
-	if (SkeletalMeshAsset && PhysicsAssetAuthoredSkeletalMesh)
+	if (SkeletalMeshAsset != nullptr)
 	{
-		if (SkeletalMeshAsset != PhysicsAssetAuthoredSkeletalMesh)
-		{
-			UE_LOGFMT(LogPhysicsControl, Log, 
-				"Modify Constraint parent transforms to correct for the difference between the Skeleton used to create the Physics asset \"{0}\" and the current skeleton \"{1}\".", 
-				PhysicsAssetAuthoredSkeletalMesh->GetName(), SkeletalMeshAsset->GetName());
+		const bool bAuthoredTransformRequired = ((PositionBehavior == MapConstraintsBehaviorType::AuthoredSkeleton) || (OrientationBehavior == MapConstraintsBehaviorType::AuthoredSkeleton)) && (PhysicsAssetAuthoredSkeletalMesh != nullptr) && (SkeletalMeshAsset != PhysicsAssetAuthoredSkeletalMesh);
+		const bool bDefaultTransformRequired = (PositionBehavior == MapConstraintsBehaviorType::DefaultTransform) || (OrientationBehavior == MapConstraintsBehaviorType::DefaultTransform);
 
-			const FReferenceSkeleton& OriginalReferenceSkeleton = PhysicsAssetAuthoredSkeletalMesh->GetRefSkeleton();
-			const FReferenceSkeleton& CurrentReferenceSkeleton = SkeletalMeshAsset->GetRefSkeleton();
+		if (bAuthoredTransformRequired || bDefaultTransformRequired)
+		{
+#if !NO_LOGGING
+			const FString AuthoredSkeletalMeshName = (PhysicsAssetAuthoredSkeletalMesh) ? PhysicsAssetAuthoredSkeletalMesh->GetName() : FString("UNDEFINED");
+			UE_LOGFMT(LogPhysicsControl, Log,
+				"Modify Constraint parent transforms to match the current skeleton \"{0}\". Settings: Authored Skeleton {1}, Position set from {2}, Orientation set from {3}",
+				SkeletalMeshAsset->GetName(), AuthoredSkeletalMeshName, MapConstraintsBehaviorTypeToString(PositionBehavior), MapConstraintsBehaviorTypeToString(OrientationBehavior));
+#endif
+			
 
 			for (FConstraintInstance* const ConstraintInstance : ConstraintInstances)
 			{
-#if !NO_LOGGING
-				const FVector LogPreviousConstraintPositionRelParent = ConstraintInstance->Pos2;
-#endif
-				const FTransform CurrentParentRelChildTM = CalculateRelativeBoneTransform(
-					ConstraintInstance->ConstraintBone1, ConstraintInstance->ConstraintBone2, CurrentReferenceSkeleton);
-				const FTransform OriginalParentRelChildTM = CalculateRelativeBoneTransform(
-					ConstraintInstance->ConstraintBone1, ConstraintInstance->ConstraintBone2, OriginalReferenceSkeleton);
-
-				// Find the transform that maps the parent-bone-relative-to-the-child-bone transform
-				// in the original skeleton to the parent-bone-relative-to-the-child-bone transform
-				// in the current skeleton.
-				// Should be equivalent to CurrentParentRelChildTM * OriginalParentRelChildTM.Inverse()
-				const FTransform OriginalToCurrentParentRelChildTM = 
-					CurrentParentRelChildTM.GetRelativeTransform(OriginalParentRelChildTM); 
-
-				// Update the constraints transform relative to the parent bone.
-				const FTransform OriginalRefFrame = ConstraintInstance->GetRefFrame(EConstraintFrame::Frame2);
-				const FTransform CurrentRefFrame = OriginalToCurrentParentRelChildTM * OriginalRefFrame;
-				ConstraintInstance->SetRefFrame(EConstraintFrame::Frame2, CurrentRefFrame);
+				FTransform AuthoredCurrentRefFrame;
+				FTransform DefaultCurrentRefFrame;
 
-#if !NO_LOGGING
-				UE_LOGFMT(LogPhysicsControl, Log, 
-					"Matched Constraint {0} - {1} Parent Transform - position was {2} now {3}.", 
-					ConstraintInstance->ConstraintBone1.ToString(), 
-					ConstraintInstance->ConstraintBone2.ToString(), 
-					LogPreviousConstraintPositionRelParent.ToCompactString(), 
-					ConstraintInstance->Pos2.ToCompactString());
-#endif
-			}
-		}
-#if !NO_LOGGING
-		else
-		{
-			UE_LOGFMT(LogPhysicsControl, Log, 
-				"Do not modify constraint parent transforms as the Skeleton used to create the Physics asset \"{0}\" matches the current skeleton \"{1}\".", 
-				PhysicsAssetAuthoredSkeletalMesh->GetName(), SkeletalMeshAsset->GetName());
-		}
-#endif
-	}
-	else if (SkeletalMeshAsset)
-	{
-		UE_LOGFMT(LogPhysicsControl, Log, 
-			"Snap Constraint parent transforms to the current skeleton \"{0}\" (Authored Skeleton is undefined in node details).", 
-			SkeletalMeshAsset->GetName());
+				if (bAuthoredTransformRequired)
+				{
+					const FTransform CurrentParentRelChildTM = CalculateRelativeBoneTransform(
+						ConstraintInstance->ConstraintBone1, ConstraintInstance->ConstraintBone2, SkeletalMeshAsset->GetRefSkeleton());
+					const FTransform OriginalParentRelChildTM = CalculateRelativeBoneTransform(
+						ConstraintInstance->ConstraintBone1, ConstraintInstance->ConstraintBone2, PhysicsAssetAuthoredSkeletalMesh->GetRefSkeleton());
+
+					// Find the transform that maps the parent-bone-relative-to-the-child-bone transform
+					// in the original skeleton to the parent-bone-relative-to-the-child-bone transform
+					// in the current skeleton.
+					// Should be equivalent to CurrentParentRelChildTM * OriginalParentRelChildTM.Inverse()
+					const FTransform OriginalToCurrentParentRelChildTM =
+						CurrentParentRelChildTM.GetRelativeTransform(OriginalParentRelChildTM);
+
+					// Update the constraints transform relative to the parent bone.
+					const FTransform OriginalRefFrame = ConstraintInstance->GetRefFrame(EConstraintFrame::Frame2);
+					AuthoredCurrentRefFrame = OriginalToCurrentParentRelChildTM * OriginalRefFrame;
+				}
 
-		const FReferenceSkeleton& CurrentReferenceSkeleton = SkeletalMeshAsset->GetRefSkeleton();
+				if (bDefaultTransformRequired)
+				{
+					DefaultCurrentRefFrame = CalculateRelativeBoneTransform(
+						ConstraintInstance->ConstraintBone1, ConstraintInstance->ConstraintBone2, SkeletalMeshAsset->GetRefSkeleton());
+				}
 
-		// Move all the constraints to the default (snapped) location relative to the parent bone.
-		for (FConstraintInstance* const ConstraintInstance : ConstraintInstances)
-		{
 #if !NO_LOGGING
-			const FVector LogPreviousConstraintPositionRelParent = ConstraintInstance->Pos2;
+				const FTransform LogPreviousConstraintTransformRelParent = ConstraintInstance->GetRefFrame(EConstraintFrame::Frame2);
 #endif
 
-			const FTransform ParentRelChildTM = CalculateRelativeBoneTransform(
-				ConstraintInstance->ConstraintBone1, ConstraintInstance->ConstraintBone2, CurrentReferenceSkeleton);			
-			ConstraintInstance->SetRefFrame(EConstraintFrame::Frame2, ParentRelChildTM);
+				if (PositionBehavior == MapConstraintsBehaviorType::AuthoredSkeleton)
+				{
+					ConstraintInstance->SetRefPosition(EConstraintFrame::Frame2, AuthoredCurrentRefFrame.GetTranslation());
+				}
+				else if (PositionBehavior == MapConstraintsBehaviorType::DefaultTransform)
+				{
+					ConstraintInstance->SetRefPosition(EConstraintFrame::Frame2, DefaultCurrentRefFrame.GetTranslation());
+				}
+
+				if (OrientationBehavior == MapConstraintsBehaviorType::AuthoredSkeleton)
+				{
+					ConstraintInstance->SetRefOrientation(EConstraintFrame::Frame2, AuthoredCurrentRefFrame.GetUnitAxis(EAxis::X), AuthoredCurrentRefFrame.GetUnitAxis(EAxis::Y));
+				}
+				else if (OrientationBehavior == MapConstraintsBehaviorType::DefaultTransform)
+				{
+					ConstraintInstance->SetRefOrientation(EConstraintFrame::Frame2, DefaultCurrentRefFrame.GetUnitAxis(EAxis::X), DefaultCurrentRefFrame.GetUnitAxis(EAxis::Y));
+				}
 
 #if !NO_LOGGING
-			UE_LOGFMT(LogPhysicsControl, Log, 
-				"Snapped Constraint {0} - {1} Parent Transform - position was {2} now {3}.", 
-				ConstraintInstance->ConstraintBone1.ToString(), 
-				ConstraintInstance->ConstraintBone2.ToString(), 
-				LogPreviousConstraintPositionRelParent.ToCompactString(), 
-				ConstraintInstance->Pos2.ToCompactString());
+				UE_LOGFMT(LogPhysicsControl, Log,
+					"Constraint {0} - {1}  - transform relative to parent was {2} now {3}.",
+					ConstraintInstance->ConstraintBone1.ToString(),
+					ConstraintInstance->ConstraintBone2.ToString(),
+					LogPreviousConstraintTransformRelParent.ToString(),
+					ConstraintInstance->GetRefFrame(EConstraintFrame::Frame2).ToString());
 #endif
+			}
 		}
 	}
 }
diff --git a/Engine/Plugins/Experimental/PhysicsControl/Source/PhysicsControl/Public/AnimNode_RigidBodyWithControl.h b/Engine/Plugins/Experimental/PhysicsControl/Source/PhysicsControl/Public/AnimNode_RigidBodyWithControl.h
--- a/Engine/Plugins/Experimental/PhysicsControl/Source/PhysicsControl/Public/AnimNode_RigidBodyWithControl.h
+++ b/Engine/Plugins/Experimental/PhysicsControl/Source/PhysicsControl/Public/AnimNode_RigidBodyWithControl.h
@@ -42,6 +42,14 @@ struct FSimulationSpaceData
 	FVector    AngularAcc;
 };
 
+UENUM(BlueprintType)
+enum class MapConstraintsBehaviorType : uint8
+{
+	AuthoredSkeleton UMETA(DisplayName = "From Authored Skeletal Mesh", ToolTip = "Apply any difference in the relative transform of the constraint's parent and child bones between the authored and current skeleton to the constraint's transform relative to the parent bone. This is normally the best option and should effectively keep the joint in the authored location relative to the parent bone."),
+	DefaultTransform UMETA(DisplayName = "From Child Bone Transform (default)", ToolTip = "Set constraints transform relative to the parent bone to match the child bone in the current skeleton."),
+	None UMETA(DisplayName = "Nothing", ToolTip = "Nothing"),
+};
+
 /**
  * Controller that simulates physics based on the physics asset of the skeletal mesh component
  */
@@ -233,11 +241,22 @@ public:
 	uint8 bFreezeIncomingPoseOnStart : 1;
 
 	/**
-		Change the parent space transforms of constraints read from the physics asset to match the relative 
-		bone transforms in the in-coming skeleton.
+	* Configure the way constraint positions can be modified at runtime to match the current Skeletal Mesh.
+	*/
+	UPROPERTY(EditAnywhere, Category = PhysicsAssetConditioning, meta = (DisplayName = "Set Position Relative to Parent", DisplayAfter = "PhysicsAssetAuthoredSkeletalMesh"))
+	MapConstraintsBehaviorType PhysicsAssetConditioningConstraintPosition;
+	
+	/**
+	* Configure the way constraint orientations can be modified at runtime to match the current Skeletal Mesh.
 	*/
-	UPROPERTY(EditAnywhere, Category = PhysicsAssetConditioning, meta = (InlineEditConditionToggle))
-	uint8 bModifyConstraintTransformsToMatchSkeleton : 1;
+	UPROPERTY(EditAnywhere, Category = PhysicsAssetConditioning, meta = (DisplayName = "Set Orientation Relative to Parent", DisplayAfter = "PhysicsAssetAuthoredSkeletalMesh"))
+	MapConstraintsBehaviorType PhysicsAssetConditioningConstraintOrientation;
+
+#if WITH_EDITORONLY_DATA
+	UPROPERTY()
+	uint8 bModifyConstraintTransformsToMatchSkeleton_DEPRECATED : 1;
+#endif
+
 
 	/**
 		For world-space simulations, if the magnitude of the component's 3D scale is less than 
@@ -334,17 +353,30 @@ public:
 	UPROPERTY(EditAnywhere, Category = Controls, meta = (PinShownByDefault))
 	FRigidBodyKinematicTargets KinematicTargets;
 
-	/**
-	 * If this option is enabled, each Constraint's parent transform will be updated to adjust its position and 
-	 * orientation to account for the difference in the constraints child transform between the skeleton used 
-	 * to author the physics asset and the current skeleton (if the authored skeleton is defined). If the authored
-	 * skeleton is not defined/unavailable then the parent bone's transform is set to the default transform that
-	 * would have been used in the physics asset.
-	 * This can be used to created a simulated character that has somewhat different bone length/orientations compared
-	 * to the one used to create the physics asset, in order to avoid the need to customize the physics asset
-	 * for every skeleton you use.
+	/** 
+	 * These options configure the way constraint transforms can be modified at runtime to match the current Skeletal
+	 * Mesh. This can be useful when applying a single RBwC node to assets with different skeletons.
+	 *
+	 * When the "From authored skeletal mesh" option is selected, the Position and/or Orientation of each constraint
+	 * relative to its parent bone will be changed to account for any difference in relative bone transforms 
+	 * between that Skeletal Mesh and the current one. This means that if bone positions or orientations have changed
+	 * between the authored Skeletal Mesh and the one currently being animated then the constraint positions and
+	 * orientations will be updated to match. The authored Skeletal Mesh should normally be set to the same skeletal 
+	 * mesh that was used when creating the physics asset. WARNING This option will do nothing if the authored 
+	 * skeleton is not defined or cannot be found when the node is initialized at runtime.
+	 *
+	 * When the "To default" option is selected the Position and/or Orientation of each constraint relative to its 
+	 * parent bone will be set to the Position and/or Orientation of the constraint's child bone relative to its parent
+	 * bone. This is the same as the default transform applied in the Physics Asset editing tool when a constraint is
+	 * created or 'snapped' to the default transform.
+	 *
+	 * [WARNING]
+	 * The "To default" option can produce undesirable results and should be used with caution but may be useful
+	 * for simple assets or to minimize sudden changes in bone positions that may arise if the
+	 * authored skeleton is unavailable during the early stages of initialization. In most cases its probably more
+	 * useful to only set constraint positions to the default.
 	 */
-	UPROPERTY(EditAnywhere, Category = PhysicsAssetConditioning, meta = (DisplayName = "Map constraints to Skeleton", editcondition = "bModifyConstraintTransformsToMatchSkeleton"))
+	UPROPERTY(EditAnywhere, Category = PhysicsAssetConditioning, meta = (DisplayName = "Authored Skeletal Mesh"))
 	USkeletalMesh* PhysicsAssetAuthoredSkeletalMesh;
 
 	/**
@@ -462,7 +494,10 @@ private:
 	// Modify Constraint transforms relative to the parent bone to correct for the difference
 	// between the Skeleton used to create the Physics asset and the current skeleton.
 	void TransformConstraintsToMatchSkeletalMesh(
-		const USkeletalMesh* const SkeletalMeshAsset, TArray<FConstraintInstance*>& ConstraintInstances);
+		const USkeletalMesh* const SkeletalMeshAsset,
+		const MapConstraintsBehaviorType PositionBehavior,
+		const MapConstraintsBehaviorType OrientationBehavior,
+		TArray<FConstraintInstance*>& ConstraintInstances);
 
 	// Gather cloth collision sources from the supplied Skeltal Mesh and add a kinematic actor
 	// representing each one of them to the sim.
