Trivial typo fix - initialising RigidBodyWithControl controls

diff --git a/Engine/Plugins/Experimental/PhysicsControl/Source/PhysicsControl/Private/PhysicsControlLimbData.cpp b/Engine/Plugins/Experimental/PhysicsControl/Source/PhysicsControl/Private/PhysicsControlLimbData.cpp
--- a/Engine/Plugins/Experimental/PhysicsControl/Source/PhysicsControl/Private/PhysicsControlLimbData.cpp
+++ b/Engine/Plugins/Experimental/PhysicsControl/Source/PhysicsControl/Private/PhysicsControlLimbData.cpp
@@ -7,7 +7,7 @@ FPhysicsControlCharacterSetupData& FPhysicsControlCharacterSetupData::operator+=
 	const FPhysicsControlCharacterSetupData& Other)
 {
 	LimbSetupData += Other.LimbSetupData;
-	DefaultWorldSpaceControlData = Other.DefaultParentSpaceControlData;
+	DefaultWorldSpaceControlData = Other.DefaultWorldSpaceControlData;
 	DefaultParentSpaceControlData = Other.DefaultParentSpaceControlData;
 	DefaultBodyModifierData = Other.DefaultBodyModifierData;
 	return *this;
