Change Physics Control Component to calculate target velocities as they're used, rather than caching them with the pose. Plus a bit of refactoring work.

diff --git a/Engine/Plugins/Experimental/PhysicsControl/Source/PhysicsControl/Private/AnimNode_RigidBodyWithControl.cpp b/Engine/Plugins/Experimental/PhysicsControl/Source/PhysicsControl/Private/AnimNode_RigidBodyWithControl.cpp
--- a/Engine/Plugins/Experimental/PhysicsControl/Source/PhysicsControl/Private/AnimNode_RigidBodyWithControl.cpp
+++ b/Engine/Plugins/Experimental/PhysicsControl/Source/PhysicsControl/Private/AnimNode_RigidBodyWithControl.cpp
@@ -288,9 +288,10 @@ void FAnimNode_RigidBodyWithControl::EvaluateComponentPose_AnyThread(FComponentS
 	}
 }
 
-void FAnimNode_RigidBodyWithControl::InitializeNewBodyTransformsDuringSimulation(FComponentSpacePoseContext& Output, const FTransform& ComponentTransform, const FTransform& BaseBoneTM)
+void FAnimNode_RigidBodyWithControl::InitializeNewBodyTransformsDuringSimulation(
+	FComponentSpacePoseContext& Output, const FTransform& ComponentTransform, const FTransform& BaseBoneTM)
 {
-	for (const RigidBodyWithControl::FOutputBoneData& OutputData : OutputBoneData)
+	for (const UE::PhysicsControl::FOutputBoneData& OutputData : OutputBoneData)
 	{
 		const int32 BodyIndex = OutputData.BodyIndex;
 		FBodyAnimData& BodyData = BodyAnimData[BodyIndex];
@@ -577,7 +578,7 @@ void FAnimNode_RigidBodyWithControl::EvaluateSkeletalControl_AnyThread(FComponen
 			// Capture bone velocities if we have captured a bone velocity pose.
 			if (bTransferBoneVelocities && (CapturedBoneVelocityPose.GetPose().GetNumBones() > 0))
 			{
-				for (const RigidBodyWithControl::FOutputBoneData& OutputData : OutputBoneData)
+				for (const UE::PhysicsControl::FOutputBoneData& OutputData : OutputBoneData)
 				{
 					const int32 BodyIndex = OutputData.BodyIndex;
 					FBodyAnimData& BodyData = BodyAnimData[BodyIndex];
@@ -628,7 +629,7 @@ void FAnimNode_RigidBodyWithControl::EvaluateSkeletalControl_AnyThread(FComponen
 					UE_LOG(LogPhysicsControl, Verbose, TEXT("%s TeleportPhysics (Scale: %f %f %f)"), *Output.AnimInstanceProxy->GetAnimInstanceName(), CompWorldSpaceTM.GetScale3D().X, CompWorldSpaceTM.GetScale3D().Y, CompWorldSpaceTM.GetScale3D().Z);
 
 					// Teleport bodies.
-					for (const RigidBodyWithControl::FOutputBoneData& OutputData : OutputBoneData)
+					for (const UE::PhysicsControl::FOutputBoneData& OutputData : OutputBoneData)
 					{
 						const int32 BodyIndex = OutputData.BodyIndex;
 						BodyAnimData[BodyIndex].bBodyTransformInitialized = true;
@@ -657,7 +658,7 @@ void FAnimNode_RigidBodyWithControl::EvaluateSkeletalControl_AnyThread(FComponen
 					InitSimulationSpace(CompWorldSpaceTM, BaseBoneTM);
 
 					// Completely reset bodies.
-					for (const RigidBodyWithControl::FOutputBoneData& OutputData : OutputBoneData)
+					for (const UE::PhysicsControl::FOutputBoneData& OutputData : OutputBoneData)
 					{
 						const int32 BodyIndex = OutputData.BodyIndex;
 						BodyAnimData[BodyIndex].bBodyTransformInitialized = true;
@@ -697,7 +698,7 @@ void FAnimNode_RigidBodyWithControl::EvaluateSkeletalControl_AnyThread(FComponen
 			// Transfer bone velocities previously captured.
 			if (bTransferBoneVelocities && (CapturedBoneVelocityPose.GetPose().GetNumBones() > 0))
 			{
-				for (const RigidBodyWithControl::FOutputBoneData& OutputData : OutputBoneData)
+				for (const UE::PhysicsControl::FOutputBoneData& OutputData : OutputBoneData)
 				{
 					const int32 BodyIndex = OutputData.BodyIndex;
 					const FBodyAnimData& BodyData = BodyAnimData[BodyIndex];
@@ -730,7 +731,7 @@ void FAnimNode_RigidBodyWithControl::EvaluateSkeletalControl_AnyThread(FComponen
 					ApplyLinearAcc += WorldVectorToSpaceNoScale(SimulationSpace, -ComponentLinearAcceleration, CompWorldSpaceTM, BaseBoneTM) * ComponentLinearAccScale;
 
 					// Iterate over bodies
-					for (const RigidBodyWithControl::FOutputBoneData& OutputData : OutputBoneData)
+					for (const UE::PhysicsControl::FOutputBoneData& OutputData : OutputBoneData)
 					{
 						const int32 BodyIndex = OutputData.BodyIndex;
 						const FBodyAnimData& BodyData = BodyAnimData[BodyIndex];
@@ -785,7 +786,7 @@ void FAnimNode_RigidBodyWithControl::EvaluateSkeletalControl_AnyThread(FComponen
 			}
 
 			// Note that the simulation interpolates kinematic targets to handle substepping
-			for (const RigidBodyWithControl::FOutputBoneData& OutputData : OutputBoneData)
+			for (const UE::PhysicsControl::FOutputBoneData& OutputData : OutputBoneData)
 			{
 				const int32 BodyIndex = OutputData.BodyIndex;
 				BodyAnimData[BodyIndex].bIsSimulated = !Bodies[BodyIndex]->GetIsKinematic();
@@ -845,7 +846,7 @@ void FAnimNode_RigidBodyWithControl::EvaluateSkeletalControl_AnyThread(FComponen
 		
 		//write back to animation system
 		const FTransform& SimulationWorldSpaceTM = bUseDeferredSimulationTask ? PreviousCompWorldSpaceTM : CompWorldSpaceTM;
-		for (const RigidBodyWithControl::FOutputBoneData& OutputData : OutputBoneData)
+		for (const UE::PhysicsControl::FOutputBoneData& OutputData : OutputBoneData)
 		{
 			const int32 BodyIndex = OutputData.BodyIndex;
 			// Note that we always read back, whether kinematic or simulated
@@ -1877,7 +1878,8 @@ void FAnimNode_RigidBodyWithControl::InitializeBoneReferences(const FBoneContain
 		if (BodyIndex != INDEX_NONE)
 		{
 			//If we have a body we need to save it for later
-			RigidBodyWithControl::FOutputBoneData* OutputData = new (OutputBoneData) RigidBodyWithControl::FOutputBoneData();
+			UE::PhysicsControl::FOutputBoneData* OutputData = 
+				new (OutputBoneData) UE::PhysicsControl::FOutputBoneData();
 			OutputData->BodyIndex = BodyIndex;
 			OutputData->CompactPoseBoneIndex = CompactPoseBoneIndex;
 
diff --git a/Engine/Plugins/Experimental/PhysicsControl/Source/PhysicsControl/Private/PhysicsControlComponent.cpp b/Engine/Plugins/Experimental/PhysicsControl/Source/PhysicsControl/Private/PhysicsControlComponent.cpp
--- a/Engine/Plugins/Experimental/PhysicsControl/Source/PhysicsControl/Private/PhysicsControlComponent.cpp
+++ b/Engine/Plugins/Experimental/PhysicsControl/Source/PhysicsControl/Private/PhysicsControlComponent.cpp
@@ -2378,13 +2378,14 @@ TArray<FTransform> UPhysicsControlComponent::GetCachedBoneTransforms(
 {
 	TArray<FTransform> Result;
 	Result.Reserve(BoneNames.Num());
-	FCachedSkeletalMeshData::FBoneData BoneData;
+	UE::PhysicsControl::FBoneData BoneData;
+	const UE::PhysicsControl::FPhysicsControlPoseData* PoseData;
 
 	for (const FName& BoneName : BoneNames)
 	{
-		if (GetBoneData(BoneData, SkeletalMeshComponent, BoneName))
+		if (GetBoneData(BoneData, PoseData, SkeletalMeshComponent, BoneName))
 		{
-			FTransform BoneTransform(BoneData.Orientation, BoneData.Position);
+			FTransform BoneTransform(BoneData.CurrentTM.GetRotation(), BoneData.CurrentTM.GetTranslation());
 			Result.Add(BoneTransform);
 		}
 		else
@@ -2402,13 +2403,14 @@ TArray<FVector> UPhysicsControlComponent::GetCachedBonePositions(
 {
 	TArray<FVector> Result;
 	Result.Reserve(BoneNames.Num());
-	FCachedSkeletalMeshData::FBoneData BoneData;
+	UE::PhysicsControl::FBoneData BoneData;
+	const UE::PhysicsControl::FPhysicsControlPoseData* PoseData;
 
 	for (const FName& BoneName : BoneNames)
 	{
-		if (GetBoneData(BoneData, SkeletalMeshComponent, BoneName))
+		if (GetBoneData(BoneData, PoseData, SkeletalMeshComponent, BoneName))
 		{
-			Result.Add(BoneData.Position);
+			Result.Add(BoneData.CurrentTM.GetTranslation());
 		}
 		else
 		{
@@ -2425,13 +2427,14 @@ TArray<FRotator> UPhysicsControlComponent::GetCachedBoneOrientations(
 {
 	TArray<FRotator> Result;
 	Result.Reserve(BoneNames.Num());
-	FCachedSkeletalMeshData::FBoneData BoneData;
+	UE::PhysicsControl::FBoneData BoneData;
+	const UE::PhysicsControl::FPhysicsControlPoseData* PoseData;
 
 	for (const FName& BoneName : BoneNames)
 	{
-		if (GetBoneData(BoneData, SkeletalMeshComponent, BoneName))
+		if (GetBoneData(BoneData, PoseData, SkeletalMeshComponent, BoneName))
 		{
-			Result.Add(BoneData.Orientation.Rotator());
+			Result.Add(BoneData.CurrentTM.GetRotation().Rotator());
 		}
 		else
 		{
@@ -2441,63 +2444,16 @@ TArray<FRotator> UPhysicsControlComponent::GetCachedBoneOrientations(
 	return Result;
 }
 
-//======================================================================================================================
-TArray<FVector> UPhysicsControlComponent::GetCachedBoneVelocities(
-	const USkeletalMeshComponent* SkeletalMeshComponent,
-	const TArray<FName>&          BoneNames)
-{
-	TArray<FVector> Result;
-	Result.Reserve(BoneNames.Num());
-	FCachedSkeletalMeshData::FBoneData BoneData;
-
-	for (const FName& BoneName : BoneNames)
-	{
-		if (GetBoneData(BoneData, SkeletalMeshComponent, BoneName))
-		{
-			Result.Add(BoneData.Velocity);
-		}
-		else
-		{
-			Result.Add(FVector::ZeroVector);
-		}
-	}
-
-	return Result;
-}
-
-//======================================================================================================================
-TArray<FVector> UPhysicsControlComponent::GetCachedBoneAngularVelocities(
-	const USkeletalMeshComponent* SkeletalMeshComponent,
-	const TArray<FName>&          BoneNames)
-{
-	TArray<FVector> Result;
-	Result.Reserve(BoneNames.Num());
-	FCachedSkeletalMeshData::FBoneData BoneData;
-
-	for (const FName& BoneName : BoneNames)
-	{
-		if (GetBoneData(BoneData, SkeletalMeshComponent, BoneName))
-		{
-			Result.Add(BoneData.AngularVelocity);
-		}
-		else
-		{
-			Result.Add(FVector::ZeroVector);
-		}
-	}
-
-	return Result;
-}
-
 //======================================================================================================================
 FTransform UPhysicsControlComponent::GetCachedBoneTransform(
 	const USkeletalMeshComponent* SkeletalMeshComponent,
 	const FName                   Name)
 {
-	FCachedSkeletalMeshData::FBoneData BoneData;
-	if (GetBoneData(BoneData, SkeletalMeshComponent, Name))
+	UE::PhysicsControl::FBoneData BoneData;
+	const UE::PhysicsControl::FPhysicsControlPoseData* PoseData;
+	if (GetBoneData(BoneData, PoseData, SkeletalMeshComponent, Name))
 	{
-		return FTransform(BoneData.Orientation, BoneData.Position);
+		return FTransform(BoneData.CurrentTM.GetRotation(), BoneData.CurrentTM.GetTranslation());
 	}
 	if (bWarnAboutInvalidNames)
 	{
@@ -2512,10 +2468,11 @@ FVector UPhysicsControlComponent::GetCachedBonePosition(
 	const USkeletalMeshComponent* SkeletalMeshComponent,
 	const FName                   Name)
 {
-	FCachedSkeletalMeshData::FBoneData BoneData;
-	if (GetBoneData(BoneData, SkeletalMeshComponent, Name))
+	UE::PhysicsControl::FBoneData BoneData;
+	const UE::PhysicsControl::FPhysicsControlPoseData* PoseData;
+	if (GetBoneData(BoneData, PoseData, SkeletalMeshComponent, Name))
 	{
-		return BoneData.Position;
+		return BoneData.CurrentTM.GetTranslation();
 	}
 	if (bWarnAboutInvalidNames)
 	{
@@ -2530,10 +2487,11 @@ FRotator UPhysicsControlComponent::GetCachedBoneOrientation(
 	const USkeletalMeshComponent* SkeletalMeshComponent,
 	const FName                   Name)
 {
-	FCachedSkeletalMeshData::FBoneData BoneData;
-	if (GetBoneData(BoneData, SkeletalMeshComponent, Name))
+	UE::PhysicsControl::FBoneData BoneData;
+	const UE::PhysicsControl::FPhysicsControlPoseData* PoseData;
+	if (GetBoneData(BoneData, PoseData, SkeletalMeshComponent, Name))
 	{
-		return BoneData.Orientation.Rotator();
+		return BoneData.CurrentTM.GetRotation().Rotator();
 	}
 	if (bWarnAboutInvalidNames)
 	{
@@ -2543,57 +2501,16 @@ FRotator UPhysicsControlComponent::GetCachedBoneOrientation(
 	return FRotator::ZeroRotator;
 }
 
-//======================================================================================================================
-FVector UPhysicsControlComponent::GetCachedBoneVelocity(
-	const USkeletalMeshComponent* SkeletalMeshComponent,
-	const FName                   Name)
-{
-	FCachedSkeletalMeshData::FBoneData BoneData;
-	if (GetBoneData(BoneData, SkeletalMeshComponent, Name))
-	{
-		return BoneData.Velocity;
-	}
-	if (bWarnAboutInvalidNames)
-	{
-		UE_LOG(LogPhysicsControl, Warning,
-			TEXT("GetCachedBoneVelocity - invalid name %s"), *Name.ToString());
-	}
-	return FVector::Zero();
-}
-
-//======================================================================================================================
-FVector UPhysicsControlComponent::GetCachedBoneAngularVelocity(
-	const USkeletalMeshComponent* SkeletalMeshComponent,
-	const FName                   Name)
-{
-	FCachedSkeletalMeshData::FBoneData BoneData;
-	if (GetBoneData(BoneData, SkeletalMeshComponent, Name))
-	{
-		return BoneData.AngularVelocity;
-	}
-	if (bWarnAboutInvalidNames)
-	{
-		UE_LOG(LogPhysicsControl, Warning,
-			TEXT("GetCachedBoneAngularVelocity - invalid name %s"), *Name.ToString());
-	}
-	return FVector::Zero();
-}
-
 //======================================================================================================================
 bool UPhysicsControlComponent::SetCachedBoneData(
 	const USkeletalMeshComponent* SkeletalMeshComponent,
 	const FName                   Name,
-	const FTransform&             TM,
-	const FVector                 Velocity,
-	const FVector                 AngularVelocity)
+	const FTransform&             TM)
 {
-	FCachedSkeletalMeshData::FBoneData* BoneData;
+	UE::PhysicsControl::FBoneData* BoneData;
 	if (GetModifiableBoneData(BoneData, SkeletalMeshComponent, Name))
 	{
-		BoneData->Position = TM.GetLocation();
-		BoneData->Orientation = TM.GetRotation();
-		BoneData->Velocity = Velocity;
-		BoneData->AngularVelocity = AngularVelocity;
+		BoneData->CurrentTM = TM;
 		return true;
 	}
 	if (bWarnAboutInvalidNames)
@@ -2827,10 +2744,11 @@ void UPhysicsControlComponent::DebugDrawControl(
 			Record.ParentMeshComponent.Get(), Record.PhysicsControl.ParentBoneName);
 		const FTransform ParentBodyTM = ParentBodyInstance ? ParentBodyInstance->GetUnrealWorldTransform() : FTransform();
 
-		FTransform TargetTM;
+		FTransform TargetTM, SkeletalTargetTM;
 		FVector TargetVelocity;
 		FVector TargetAngularVelocity;
-		CalculateControlTargetData(TargetTM, TargetVelocity, TargetAngularVelocity, Record, true);
+		// Note that we want velocities, but there is a risk that they will be invalid, depending on the update times
+		CalculateControlTargetData(TargetTM, SkeletalTargetTM, TargetVelocity, TargetAngularVelocity, Record, true);
 
 		// WorldChildFrameTM is the world-space transform of the child (driven) constraint frame
 		const FTransform WorldChildFrameTM = ConstraintInstance->GetRefFrame(EConstraintFrame::Frame1) * ChildBodyTM;
diff --git a/Engine/Plugins/Experimental/PhysicsControl/Source/PhysicsControl/Private/PhysicsControlComponentHelpers.cpp b/Engine/Plugins/Experimental/PhysicsControl/Source/PhysicsControl/Private/PhysicsControlComponentHelpers.cpp
--- a/Engine/Plugins/Experimental/PhysicsControl/Source/PhysicsControl/Private/PhysicsControlComponentHelpers.cpp
+++ b/Engine/Plugins/Experimental/PhysicsControl/Source/PhysicsControl/Private/PhysicsControlComponentHelpers.cpp
@@ -141,6 +141,6 @@ FName GetPhysicalParentBone(USkeletalMeshComponent* SkeletalMeshComponent, FName
 	return FName();
 }
 
-} // namespace PhysicsControlComponent
+} // namespace PhysicsControl
 } // namespace UE
 
diff --git a/Engine/Plugins/Experimental/PhysicsControl/Source/PhysicsControl/Private/PhysicsControlComponentHelpers.h b/Engine/Plugins/Experimental/PhysicsControl/Source/PhysicsControl/Private/PhysicsControlComponentHelpers.h
--- a/Engine/Plugins/Experimental/PhysicsControl/Source/PhysicsControl/Private/PhysicsControlComponentHelpers.h
+++ b/Engine/Plugins/Experimental/PhysicsControl/Source/PhysicsControl/Private/PhysicsControlComponentHelpers.h
@@ -96,5 +96,5 @@ static void ConvertSpringToStrengthParams(
 	}
 }
 
-} // namespace PhysicsControlComponent
+} // namespace PhysicsControl
 } // namespace UE
diff --git a/Engine/Plugins/Experimental/PhysicsControl/Source/PhysicsControl/Private/PhysicsControlComponentImpl.cpp b/Engine/Plugins/Experimental/PhysicsControl/Source/PhysicsControl/Private/PhysicsControlComponentImpl.cpp
--- a/Engine/Plugins/Experimental/PhysicsControl/Source/PhysicsControl/Private/PhysicsControlComponentImpl.cpp
+++ b/Engine/Plugins/Experimental/PhysicsControl/Source/PhysicsControl/Private/PhysicsControlComponentImpl.cpp
@@ -47,11 +47,13 @@ static USkeletalMeshComponent* GetValidSkeletalMeshComponentFromBodyModifier(
 
 //======================================================================================================================
 bool UPhysicsControlComponent::GetBoneData(
-	FCachedSkeletalMeshData::FBoneData& OutBoneData,
-	const USkeletalMeshComponent*       InSkeletalMeshComponent,
-	const FName                         InBoneName) const
+	UE::PhysicsControl::FBoneData&                      OutBoneData,
+	const UE::PhysicsControl::FPhysicsControlPoseData*& OutPoseData,
+	const USkeletalMeshComponent*                       InSkeletalMeshComponent,
+	const FName                                         InBoneName) const
 {
 	check(InSkeletalMeshComponent);
+	OutPoseData = nullptr;
 	const FReferenceSkeleton& RefSkeleton = InSkeletalMeshComponent->GetSkeletalMeshAsset()->GetRefSkeleton();
 	const int32 BoneIndex = RefSkeleton.FindBoneIndex(InBoneName);
 
@@ -61,14 +63,15 @@ bool UPhysicsControlComponent::GetBoneData(
 		return false;
 	}
 
-	const FCachedSkeletalMeshData* CachedSkeletalMeshData = CachedSkeletalMeshDatas.Find(InSkeletalMeshComponent);
-	if (CachedSkeletalMeshData &&
-		CachedSkeletalMeshData->ReferenceCount > 0 &&
-		!CachedSkeletalMeshData->BoneData.IsEmpty())
+	const UE::PhysicsControl::FPhysicsControlPoseData* CachedSkeletalMeshData =
+		CachedPoseDatas.Find(InSkeletalMeshComponent);
+	if (CachedSkeletalMeshData && CachedSkeletalMeshData->ReferenceCount > 0 &&
+		!CachedSkeletalMeshData->BoneDatas.IsEmpty())
 	{
-		if (BoneIndex < CachedSkeletalMeshData->BoneData.Num())
+		if (BoneIndex < CachedSkeletalMeshData->BoneDatas.Num())
 		{
-			OutBoneData = CachedSkeletalMeshData->BoneData[BoneIndex];
+			OutBoneData = CachedSkeletalMeshData->BoneDatas[BoneIndex];
+			OutPoseData = CachedSkeletalMeshData;
 			return true;
 		}
 		UE_LOG(LogPhysicsControl, Warning, TEXT("BoneIndex is out of range"));
@@ -80,9 +83,9 @@ bool UPhysicsControlComponent::GetBoneData(
 
 //======================================================================================================================
 bool UPhysicsControlComponent::GetModifiableBoneData(
-	FCachedSkeletalMeshData::FBoneData*& OutBoneData,
-	const USkeletalMeshComponent*        InSkeletalMeshComponent,
-	const FName                          InBoneName)
+	UE::PhysicsControl::FBoneData*& OutBoneData,
+	const USkeletalMeshComponent*   InSkeletalMeshComponent,
+	const FName                     InBoneName)
 {
 	check(InSkeletalMeshComponent);
 	const FReferenceSkeleton& RefSkeleton = InSkeletalMeshComponent->GetSkeletalMeshAsset()->GetRefSkeleton();
@@ -94,14 +97,15 @@ bool UPhysicsControlComponent::GetModifiableBoneData(
 		return false;
 	}
 
-	FCachedSkeletalMeshData* CachedSkeletalMeshData = CachedSkeletalMeshDatas.Find(InSkeletalMeshComponent);
+	UE::PhysicsControl::FPhysicsControlPoseData* CachedSkeletalMeshData = 
+		CachedPoseDatas.Find(InSkeletalMeshComponent);
 	if (CachedSkeletalMeshData &&
 		CachedSkeletalMeshData->ReferenceCount > 0 &&
-		!CachedSkeletalMeshData->BoneData.IsEmpty())
+		!CachedSkeletalMeshData->BoneDatas.IsEmpty())
 	{
-		if (BoneIndex < CachedSkeletalMeshData->BoneData.Num())
+		if (BoneIndex < CachedSkeletalMeshData->BoneDatas.Num())
 		{
-			OutBoneData = &CachedSkeletalMeshData->BoneData[BoneIndex];
+			OutBoneData = &CachedSkeletalMeshData->BoneDatas[BoneIndex];
 			return true;
 		}
 		UE_LOG(LogPhysicsControl, Warning, TEXT("BoneIndex is out of range"));
@@ -153,42 +157,15 @@ const FPhysicsControl* UPhysicsControlComponent::FindControl(const FName Name) c
 }
 
 //======================================================================================================================
-bool UPhysicsControlComponent::DetectTeleport(
-	const FVector& OldPosition, const FQuat& OldOrientation,
-	const FVector& NewPosition, const FQuat& NewOrientation) const
+void UPhysicsControlComponent::UpdateCachedSkeletalBoneData(float DeltaTime)
 {
-	if (TeleportDistanceThreshold > 0)
-	{
-		const double Distance = FVector::Distance(OldPosition, NewPosition);
-		if (Distance > TeleportDistanceThreshold)
-		{
-			return true;
-		}
-	}
-	if (TeleportRotationThreshold > 0)
-	{
-		const double Radians = OldOrientation.AngularDistance(NewOrientation);
-		if (FMath::RadiansToDegrees(Radians) > TeleportRotationThreshold)
-		{
-			return true;
-		}
-	}
-	return false;
-}
-
-//======================================================================================================================
-bool UPhysicsControlComponent::DetectTeleport(const FTransform& OldTM, const FTransform& NewTM) const
-{
-	return DetectTeleport(OldTM.GetTranslation(), OldTM.GetRotation(), NewTM.GetTranslation(), NewTM.GetRotation());
-}
+	// Allow the target counter to wrap, after 2e19 frames
+	++CurrentUpdateCounter;
 
-//======================================================================================================================
-void UPhysicsControlComponent::UpdateCachedSkeletalBoneData(float Dt)
-{
-	for (TPair<TWeakObjectPtr<USkeletalMeshComponent>, FCachedSkeletalMeshData>& CachedSkeletalMeshDataPair :
-		CachedSkeletalMeshDatas)
+	for (TPair<TWeakObjectPtr<USkeletalMeshComponent>, UE::PhysicsControl::FPhysicsControlPoseData>&
+		CachedSkeletalMeshDataPair : CachedPoseDatas)
 	{
-		FCachedSkeletalMeshData& CachedSkeletalMeshData = CachedSkeletalMeshDataPair.Value;
+		UE::PhysicsControl::FPhysicsControlPoseData& CachedSkeletalMeshData = CachedSkeletalMeshDataPair.Value;
 		if (!CachedSkeletalMeshData.ReferenceCount)
 		{
 			continue;
@@ -197,44 +174,12 @@ void UPhysicsControlComponent::UpdateCachedSkeletalBoneData(float Dt)
 		TWeakObjectPtr<USkeletalMeshComponent> SkeletalMeshComponent = CachedSkeletalMeshDataPair.Key;
 		if (USkeletalMeshComponent* SkeletalMesh = SkeletalMeshComponent.Get())
 		{
-			const FTransform ComponentTM = SkeletalMesh->GetComponentTransform();
-			const TArray<FTransform>& TMs = SkeletalMesh->GetEditableComponentSpaceTransforms();
-			const int32 NumTMs = TMs.Num();
-			if (NumTMs == CachedSkeletalMeshData.BoneData.Num() &&
-				!DetectTeleport(CachedSkeletalMeshData.ComponentTM, ComponentTM))
-			{
-				// Avoid the Dt test on every bone
-				if (Dt > 0)
-				{
-					for (int32 Index = 0; Index != NumTMs; ++Index)
-					{
-						FTransform TM = TMs[Index] * ComponentTM;
-						CachedSkeletalMeshData.BoneData[Index].Update(TM.GetTranslation(), TM.GetRotation(), Dt);
-					}
-				}
-				else
-				{
-					for (int32 Index = 0; Index != NumTMs; ++Index)
-					{
-						FTransform TM = TMs[Index] * ComponentTM;
-						CachedSkeletalMeshData.BoneData[Index].Update(TM.GetTranslation(), TM.GetRotation());
-					}
-				}
-			}
-			else
-			{
-				CachedSkeletalMeshData.BoneData.Empty(NumTMs);
-				for (const FTransform& BoneTM : TMs)
-				{
-					FTransform TM = BoneTM * ComponentTM;
-					CachedSkeletalMeshData.BoneData.Emplace(TM.GetTranslation(), TM.GetRotation());
-				}
-			}
-			CachedSkeletalMeshData.ComponentTM = ComponentTM;
+			CachedSkeletalMeshData.Update(
+				SkeletalMesh, DeltaTime, TeleportDistanceThreshold, TeleportRotationThreshold);
 		}
 		else
 		{
-			CachedSkeletalMeshData.BoneData.Empty();
+			CachedSkeletalMeshData.Reset();
 		}
 	}
 }
@@ -273,22 +218,26 @@ void UPhysicsControlComponent::ApplyKinematicTarget(const FPhysicsBodyModifierRe
 		KinematicTarget.SetTranslation(Record.KinematicTargetPosition);
 		if (Record.BodyModifier.ModifierData.bUseSkeletalAnimation)
 		{
-			FCachedSkeletalMeshData::FBoneData BoneData;
-			if (GetBoneData(BoneData, SkeletalMeshComponent, Record.BodyModifier.BoneName))
+			UE::PhysicsControl::FBoneData BoneData;
+			const UE::PhysicsControl::FPhysicsControlPoseData* PoseData;
+			if (GetBoneData(BoneData, PoseData, SkeletalMeshComponent, Record.BodyModifier.BoneName))
 			{
-				FTransform BoneTM = BoneData.GetTM();
+				FTransform BoneTM = BoneData.CurrentTM.ToTransform();
 				KinematicTarget = KinematicTarget * BoneTM;
 			}
 		}
-		ETeleportType TT = DetectTeleport(TM, KinematicTarget) ? ETeleportType::ResetPhysics : ETeleportType::None;
+		ETeleportType TT = UE::PhysicsControl::DetectTeleport(
+			TM, KinematicTarget, TeleportRotationThreshold, TeleportRotationThreshold) ? 
+			ETeleportType::ResetPhysics : ETeleportType::None;
 		BodyInstance->SetBodyTransform(KinematicTarget, TT);
 	}
 	else
 	{
 		const FTransform TM = Record.MeshComponent->GetComponentToWorld();
-		const ETeleportType TT = DetectTeleport(
+		const ETeleportType TT = UE::PhysicsControl::DetectTeleport(
 			TM.GetTranslation(), TM.GetRotation(), 
-			Record.KinematicTargetPosition, Record.KinematicTargetOrientation)
+			Record.KinematicTargetPosition, Record.KinematicTargetOrientation,
+			TeleportDistanceThreshold, TeleportRotationThreshold)
 			? ETeleportType::ResetPhysics : ETeleportType::None;
 		// Note that calling BodyInstance->SetBodyTransform moves the physics, but not the mesh
 		Record.MeshComponent->SetWorldLocationAndRotation(
@@ -308,16 +257,17 @@ void UPhysicsControlComponent::ResetToCachedTarget(const FPhysicsBodyModifierRec
 
 	if (USkeletalMeshComponent* SkeletalMeshComponent = GetValidSkeletalMeshComponentFromBodyModifier(Record))
 	{
-		FCachedSkeletalMeshData::FBoneData BoneData;
-		if (GetBoneData(BoneData, SkeletalMeshComponent, Record.BodyModifier.BoneName))
+		UE::PhysicsControl::FBoneData BoneData;
+		const UE::PhysicsControl::FPhysicsControlPoseData* PoseData;
+		if (GetBoneData(BoneData, PoseData, SkeletalMeshComponent, Record.BodyModifier.BoneName))
 		{
 			FTransform BoneTM = BodyInstance->GetUnrealWorldTransform(); // Preserve scale
-			BoneTM.SetLocation(BoneData.Position);
-			BoneTM.SetRotation(BoneData.Orientation);
+			BoneTM.SetLocation(BoneData.CurrentTM.GetTranslation());
+			BoneTM.SetRotation(BoneData.CurrentTM.GetRotation());
 
 			BodyInstance->SetBodyTransform(BoneTM, ETeleportType::TeleportPhysics);
-			BodyInstance->SetLinearVelocity(BoneData.Velocity, false);
-			BodyInstance->SetAngularVelocityInRadians(BoneData.AngularVelocity, false);
+			BodyInstance->SetLinearVelocity(BoneData.CalculateLinearVelocity(PoseData->DeltaTime), false);
+			BodyInstance->SetAngularVelocityInRadians(BoneData.CalculateAngularVelocity(PoseData->DeltaTime), false);
 		}
 	}
 }
@@ -327,18 +277,18 @@ void UPhysicsControlComponent::AddSkeletalMeshReferenceForCaching(
 	USkeletalMeshComponent* InSkeletalMeshComponent)
 {
 	check(InSkeletalMeshComponent);
-	for (TPair<TWeakObjectPtr<USkeletalMeshComponent>, FCachedSkeletalMeshData>& CachedSkeletalMeshDataPair :
-		CachedSkeletalMeshDatas)
+	for (TPair<TWeakObjectPtr<USkeletalMeshComponent>, UE::PhysicsControl::FPhysicsControlPoseData>& 
+		CachedSkeletalMeshDataPair : CachedPoseDatas)
 	{
 		TWeakObjectPtr<USkeletalMeshComponent> SkeletalMeshComponent = CachedSkeletalMeshDataPair.Key;
 		if (SkeletalMeshComponent.Get() == InSkeletalMeshComponent)
 		{
-			FCachedSkeletalMeshData& CachedSkeletalMeshData = CachedSkeletalMeshDataPair.Value;
+			UE::PhysicsControl::FPhysicsControlPoseData& CachedSkeletalMeshData = CachedSkeletalMeshDataPair.Value;
 			++CachedSkeletalMeshData.ReferenceCount;
 			return;
 		}
 	}
-	FCachedSkeletalMeshData& Data = CachedSkeletalMeshDatas.Add(InSkeletalMeshComponent);
+	UE::PhysicsControl::FPhysicsControlPoseData& Data = CachedPoseDatas.Add(InSkeletalMeshComponent);
 	Data.ReferenceCount = 1;
 	PrimaryComponentTick.AddPrerequisite(InSkeletalMeshComponent, InSkeletalMeshComponent->PrimaryComponentTick);
 }
@@ -354,10 +304,10 @@ bool UPhysicsControlComponent::RemoveSkeletalMeshReferenceForCaching(
 		return false;
 	}
 
-	for (auto It = CachedSkeletalMeshDatas.CreateIterator(); It; ++It)
+	for (auto It = CachedPoseDatas.CreateIterator(); It; ++It)
 	{
 		TWeakObjectPtr<USkeletalMeshComponent> SkeletalMeshComponent = It.Key();
-		FCachedSkeletalMeshData& Data = It.Value();
+		UE::PhysicsControl::FPhysicsControlPoseData& Data = It.Value();
 		if (SkeletalMeshComponent.Get() == InSkeletalMeshComponent)
 		{
 			if (--Data.ReferenceCount == 0)
@@ -442,164 +392,118 @@ bool UPhysicsControlComponent::RemoveSkeletalMeshReferenceForModifier(
 //======================================================================================================================
 void UPhysicsControlComponent::CalculateControlTargetData(
 	FTransform&                  OutTargetTM,
+	FTransform&                  OutSkeletalTargetTM,
 	FVector&                     OutTargetVelocity,
 	FVector&                     OutTargetAngularVelocity,
 	const FPhysicsControlRecord& Record,
-	bool                         bCalculateVelocity) const
+	bool                         bUsePreviousSkeletalTargetTM) const
 {
 	OutTargetTM = FTransform();
+	OutSkeletalTargetTM = FTransform();
 	OutTargetVelocity.Set(0, 0, 0);
 	OutTargetAngularVelocity.Set(0, 0, 0);
+	float SkeletalDeltaTime = 0.0f;
 
 	bool bUsedSkeletalAnimation = false;
+	bool bHasJustTeleported = false;
 
 	// Set the target TM and velocities based on any skeletal action. Note that the targets from animation 
 	// should always account for the control point
 	if (Record.PhysicsControl.ControlData.bUseSkeletalAnimation)
 	{
-		FCachedSkeletalMeshData::FBoneData ChildBoneData, ParentBoneData;
+		UE::PhysicsControl::FBoneData ChildBoneData, ParentBoneData;
+		const UE::PhysicsControl::FPhysicsControlPoseData* ChildPoseData = nullptr;
+		const UE::PhysicsControl::FPhysicsControlPoseData* ParentPoseData = nullptr;
 		bool bHaveChildBoneData = false;
 		bool bHaveParentBoneData = false;
 
 		if (USkeletalMeshComponent* ChildSkeletalMeshComponent = GetValidSkeletalMeshComponentFromControlChild(Record))
 		{
 			bHaveChildBoneData = GetBoneData(
-				ChildBoneData, ChildSkeletalMeshComponent, Record.PhysicsControl.ChildBoneName);
+				ChildBoneData, ChildPoseData, ChildSkeletalMeshComponent, Record.PhysicsControl.ChildBoneName);
 		}
 
 		if (USkeletalMeshComponent* ParentSkeletalMeshComponent = GetValidSkeletalMeshComponentFromControlParent(Record))
 		{
 			bHaveParentBoneData = GetBoneData(
-				ParentBoneData, ParentSkeletalMeshComponent, Record.PhysicsControl.ParentBoneName);
+				ParentBoneData, ParentPoseData, ParentSkeletalMeshComponent, Record.PhysicsControl.ParentBoneName);
 		}
 		else if (Record.ParentMeshComponent.IsValid())
 		{
 			const FTransform ParentTM = Record.ParentMeshComponent->GetComponentTransform();
-			ParentBoneData.Position = ParentTM.GetLocation();
-			ParentBoneData.Orientation = ParentTM.GetRotation();
-			ParentBoneData.Velocity = Record.ParentMeshComponent->GetPhysicsLinearVelocity();
-			ParentBoneData.AngularVelocity = Record.ParentMeshComponent->GetPhysicsAngularVelocityInRadians();
+			ParentBoneData.CurrentTM = ParentTM;
 			bHaveParentBoneData = true;
-			bCalculateVelocity = false;
 		}
 
-		// Note that the TargetTM/velocity calculated so far are supposed to be interpreted as
+		// Note that the TargetTM calculated so far are supposed to be interpreted as
 		// expressed relative to the skeletal animation pose.
-		//
-		// Also note that the velocities calculated in the bone data are the strict rates of change
-		// of the transform position/orientation - not of the center of mass (which is what physics
-		// bodies often use for velocity).
 		if (bHaveChildBoneData)
 		{
 			bUsedSkeletalAnimation = true;
-			const FTransform ChildBoneTM = ChildBoneData.GetTM();
+			bHasJustTeleported = ChildPoseData->bHasJustTeleported;
+			SkeletalDeltaTime = ChildPoseData->DeltaTime;
+			const UE::PhysicsControl::FPosQuat ChildBoneTM = ChildBoneData.CurrentTM;
 			if (bHaveParentBoneData)
 			{
-				const FTransform ParentBoneTM = ParentBoneData.GetTM();
-				const FTransform SkeletalDeltaTM = ChildBoneTM * ParentBoneTM.Inverse();
+				const UE::PhysicsControl::FPosQuat ParentBoneTM = ParentBoneData.CurrentTM;
+				const UE::PhysicsControl::FPosQuat SkeletalDeltaTM = ChildBoneTM * ParentBoneTM.Inverse();
 				// This puts TargetTM in the space of the ParentBone
-				OutTargetTM = SkeletalDeltaTM;
-
-				// Add on the control point offset
-				OutTargetTM.AddToTranslation(OutTargetTM.GetRotation() * Record.GetControlPoint());
-
-				const FQuat ParentBoneQ = ParentBoneTM.GetRotation();
-				const FQuat ParentBoneInvQ = ParentBoneQ.Inverse();
-
-				if (bCalculateVelocity)
-				{
-					if (Record.PhysicsControl.ControlData.SkeletalAnimationVelocityMultiplier != 0)
-					{
-						// Offset of the control point from the target child bone TM, in world space.
-						const FVector WorldControlPointOffset = ChildBoneTM.GetRotation() * Record.GetControlPoint();
-						// World space position of the target control point
-						const FVector WorldChildControlPointPosition = ChildBoneTM.GetTranslation() + WorldControlPointOffset;
-
-						// World-space velocity of the control point due to the motion of the parent
-						// linear and angular velocity.
-						const FVector ChildTargetVelocityDueToParent =
-							ParentBoneData.Velocity + ParentBoneData.AngularVelocity.Cross(
-								WorldChildControlPointPosition - ParentBoneTM.GetTranslation());
-						// World-space velocity of the control point due to the motion of the child
-						// linear and angular velocity
-						const FVector ChildTargetVelocity =
-							ChildBoneData.Velocity + ChildBoneData.AngularVelocity.Cross(WorldControlPointOffset);
-
-						// Pull out just the motion in the child that isn't due to the parent
-						const FVector SkeletalTargetVelocity =
-							ParentBoneInvQ * (ChildTargetVelocity - ChildTargetVelocityDueToParent);
-						OutTargetVelocity += SkeletalTargetVelocity *
-							Record.PhysicsControl.ControlData.SkeletalAnimationVelocityMultiplier;
-
-						const FVector SkeletalTargetAngularVelocity =
-							ParentBoneInvQ * (ChildBoneData.AngularVelocity - ParentBoneData.AngularVelocity);
-						OutTargetAngularVelocity += SkeletalTargetAngularVelocity *
-							Record.PhysicsControl.ControlData.SkeletalAnimationVelocityMultiplier;
-					}
-				}
+				OutSkeletalTargetTM = SkeletalDeltaTM.ToTransform();
 			}
 			else
 			{
-				OutTargetTM = ChildBoneTM;
-
-				// Add on the control point offset
-				OutTargetTM.AddToTranslation(OutTargetTM.GetRotation()* Record.GetControlPoint());
-
-				if (bCalculateVelocity)
-				{
-					OutTargetVelocity = ChildBoneTM.GetRotation() * OutTargetVelocity;
-					OutTargetAngularVelocity = ChildBoneTM.GetRotation() * OutTargetAngularVelocity;
-
-					if (Record.PhysicsControl.ControlData.SkeletalAnimationVelocityMultiplier != 0)
-					{
-						const FVector WorldControlPointOffset = ChildBoneTM.GetRotation() * Record.GetControlPoint();
-						const FVector WorldChildControlPointPosition = ChildBoneTM.GetTranslation() + WorldControlPointOffset;
-
-						// const World-space velocity of the control point due to the motion of the child
-						const FVector ChildTargetVelocity =
-							ChildBoneData.Velocity + ChildBoneData.AngularVelocity.Cross(WorldControlPointOffset);
-
-						OutTargetVelocity += ChildTargetVelocity *
-							Record.PhysicsControl.ControlData.SkeletalAnimationVelocityMultiplier;
-
-						OutTargetAngularVelocity += ChildBoneData.AngularVelocity *
-							Record.PhysicsControl.ControlData.SkeletalAnimationVelocityMultiplier;
-					}
-				}
+				OutSkeletalTargetTM = ChildBoneTM.ToTransform();
 			}
+			// Add on the control point offset
+			OutSkeletalTargetTM.AddToTranslation(OutSkeletalTargetTM.GetRotation() * Record.GetControlPoint());
 		}
 	}
 
-	// Now apply the explicit target specified in the record. It operates in the space of the target
-	// transform we (may have) just calculated.
+	// Calculate the velocity targets due to skeletal animation
+	if (bUsePreviousSkeletalTargetTM && !bHasJustTeleported)
 	{
-		const FPhysicsControlTarget& Target = Record.ControlTarget;
-
-		// Calculate the authored target position/orientation - i.e. not using the skeletal animation
-		FQuat TargetOrientationQ = Target.TargetOrientation.Quaternion();
-		const FVector TargetPosition = Target.TargetPosition;
-
-		FVector ExtraTargetPosition(0);
-		// Incorporate the offset from the control point. If we used animation, then we don't need
-		// to do this.
-		if (!bUsedSkeletalAnimation && 
-			Record.ControlTarget.bApplyControlPointToTarget)
+		if (SkeletalDeltaTime * Record.PhysicsControl.ControlData.LinearTargetVelocityMultiplier != 0)
 		{
-			ExtraTargetPosition = TargetOrientationQ * Record.GetControlPoint();
+			OutTargetVelocity = (OutSkeletalTargetTM.GetTranslation() - Record.PreviousSkeletalTargetTM.GetTranslation()) *
+				(Record.PhysicsControl.ControlData.LinearTargetVelocityMultiplier / SkeletalDeltaTime);
 		}
-
-		// The record's target is specified in the space of the previously calculated/set OutTargetTM
-		OutTargetTM = FTransform(TargetOrientationQ, TargetPosition + ExtraTargetPosition) * OutTargetTM;
-
-		if (bCalculateVelocity)
+		if (SkeletalDeltaTime * Record.PhysicsControl.ControlData.AngularTargetVelocityMultiplier != 0)
 		{
-			// Note that Target.TargetAngularVelocity is in revs per second (as it's user-facing)
-			const FVector TargetAngularVelocity = Target.TargetAngularVelocity * UE_TWO_PI;
-			OutTargetAngularVelocity += TargetAngularVelocity;
-			const FVector ExtraVelocity = TargetAngularVelocity.Cross(ExtraTargetPosition);
-			OutTargetVelocity += Target.TargetVelocity + ExtraVelocity;
+			const FQuat Q = OutSkeletalTargetTM.GetRotation();
+			FQuat PrevQ = Record.PreviousSkeletalTargetTM.GetRotation();
+			PrevQ.EnforceShortestArcWith(Q);
+			const FQuat DeltaQ = Q * PrevQ.Inverse();
+			OutTargetAngularVelocity = DeltaQ.ToRotationVector() * (
+				Record.PhysicsControl.ControlData.AngularTargetVelocityMultiplier / SkeletalDeltaTime);
 		}
 	}
+
+	// Now apply the explicit target specified in the record. It operates in the space of the target
+	// transform we (may have) just calculated.
+	const FPhysicsControlTarget& Target = Record.ControlTarget;
+
+	// Calculate the authored target position/orientation - i.e. not using the skeletal animation
+	FQuat TargetOrientationQ = Target.TargetOrientation.Quaternion();
+	const FVector TargetPosition = Target.TargetPosition;
+
+	FVector ExtraTargetPosition(0);
+	// Incorporate the offset from the control point. If we used animation, then we don't need
+	// to do this.
+	if (!bUsedSkeletalAnimation && 
+		Record.ControlTarget.bApplyControlPointToTarget)
+	{
+		ExtraTargetPosition = TargetOrientationQ * Record.GetControlPoint();
+	}
+
+	// Note that Target.TargetAngularVelocity is in revs per second (as it's user-facing)
+	// Also, these need to be converted (rotated) from the skeletal target space 
+	const FVector TargetAngularVelocity = Target.TargetAngularVelocity * UE_TWO_PI;
+	OutTargetAngularVelocity += OutSkeletalTargetTM.GetRotation() * TargetAngularVelocity;
+	const FVector ExtraVelocity = TargetAngularVelocity.Cross(ExtraTargetPosition);
+	OutTargetVelocity += ExtraVelocity + OutSkeletalTargetTM.GetRotation() * Target.TargetVelocity;
+
+	// The record's target is specified in the space of the previously calculated/set OutSkeletalTargetTM
+	OutTargetTM = FTransform(TargetOrientationQ, TargetPosition + ExtraTargetPosition) * OutSkeletalTargetTM;
 }
 
 //======================================================================================================================
@@ -697,10 +601,14 @@ void UPhysicsControlComponent::ApplyControl(FPhysicsControlRecord& Record)
 	// Set strengths etc and then targets (if there were strengths)
 	if (ApplyControlStrengths(Record, ConstraintInstance))
 	{
-		FTransform TargetTM;
+		FTransform TargetTM, SkeletalTargetTM;
 		FVector TargetVelocity;
 		FVector TargetAngularVelocity;
-		CalculateControlTargetData(TargetTM, TargetVelocity, TargetAngularVelocity, Record, true);
+		bool bUsePreviousSkeletalTargetTM = CurrentUpdateCounter == Record.ExpectedUpdateCounter;
+		CalculateControlTargetData(
+			TargetTM, SkeletalTargetTM, TargetVelocity, TargetAngularVelocity, Record, bUsePreviousSkeletalTargetTM);
+		Record.PreviousSkeletalTargetTM = SkeletalTargetTM;
+		Record.ExpectedUpdateCounter = CurrentUpdateCounter + 1;
 
 		ConstraintInstance->SetLinearPositionTarget(TargetTM.GetTranslation());
 		ConstraintInstance->SetAngularOrientationTarget(TargetTM.GetRotation());
@@ -807,7 +715,7 @@ bool UPhysicsControlComponent::DestroyControl(const FName Name, const EDestroyBe
 			RemoveSkeletalMeshReferenceForCaching(SkeletalMeshComponent);
 		}
 
-		Record->ResetConstraint();
+		Record->ResetConstraint(); // This terminates the constraint
 		NameRecords.RemoveControl(Name);
 		if (DestroyBehavior == EDestroyBehavior::RemoveRecord)
 		{
diff --git a/Engine/Plugins/Experimental/PhysicsControl/Source/PhysicsControl/Private/PhysicsControlOperatorNameGeneration.cpp b/Engine/Plugins/Experimental/PhysicsControl/Source/PhysicsControl/Private/PhysicsControlOperatorNameGeneration.cpp
--- a/Engine/Plugins/Experimental/PhysicsControl/Source/PhysicsControl/Private/PhysicsControlOperatorNameGeneration.cpp
+++ b/Engine/Plugins/Experimental/PhysicsControl/Source/PhysicsControl/Private/PhysicsControlOperatorNameGeneration.cpp
@@ -392,5 +392,5 @@ void CreateAdditionalSets(
 	CreateAdditionalSets_Implementation(AdditionalSets, BodyModifierRecords, Controls, NameRecords);
 }
 
-} // namespace PhysicsControlComponent
+} // namespace PhysicsControl
 } // namespace UE
diff --git a/Engine/Plugins/Experimental/PhysicsControl/Source/PhysicsControl/Private/PhysicsControlPoseData.cpp b/Engine/Plugins/Experimental/PhysicsControl/Source/PhysicsControl/Private/PhysicsControlPoseData.cpp
new file mode 100644
--- /dev/null
+++ b/Engine/Plugins/Experimental/PhysicsControl/Source/PhysicsControl/Private/PhysicsControlPoseData.cpp
@@ -0,0 +1,108 @@
+// Copyright Epic Games, Inc. All Rights Reserved.
+
+#include "PhysicsControlPoseData.h"
+#include "Components/SkeletalMeshComponent.h"
+
+namespace UE
+{
+namespace PhysicsControl
+{
+
+//======================================================================================================================
+bool DetectTeleport(
+	const FVector& OldPosition, const FQuat& OldOrientation,
+	const FVector& NewPosition, const FQuat& NewOrientation,
+	float DistanceThreshold, float RotationThreshold)
+{
+	if (DistanceThreshold > 0)
+	{
+		const double Distance = FVector::Distance(OldPosition, NewPosition);
+		if (Distance > DistanceThreshold)
+		{
+			return true;
+		}
+	}
+	if (RotationThreshold > 0)
+	{
+		const double Radians = OldOrientation.AngularDistance(NewOrientation);
+		if (FMath::RadiansToDegrees(Radians) > RotationThreshold)
+		{
+			return true;
+		}
+	}
+	return false;
+}
+
+//======================================================================================================================
+bool DetectTeleport(const FTransform& OldTM, const FTransform& NewTM,
+	float DistanceThreshold, float RotationThreshold)
+{
+	return DetectTeleport(OldTM.GetTranslation(), OldTM.GetRotation(), NewTM.GetTranslation(), NewTM.GetRotation(),
+		DistanceThreshold, RotationThreshold);
+}
+
+//======================================================================================================================
+void FPhysicsControlPoseData::Update(
+	USkeletalMeshComponent* SkeletalMesh, 
+	const float             Dt,
+	const float             TeleportDistanceThreshold, 
+	const float             TeleportRotationThreshold)
+{
+	DeltaTime = Dt;
+	const FTransform SkeletalMeshComponentTM = SkeletalMesh->GetComponentTransform();
+	const TArray<FTransform>& TMs = SkeletalMesh->GetEditableComponentSpaceTransforms();
+	const int32 NumTMs = TMs.Num();
+	if (NumTMs == BoneDatas.Num())
+	{
+		for (int32 Index = 0; Index != NumTMs; ++Index)
+		{
+			BoneDatas[Index].PreviousTM = BoneDatas[Index].CurrentTM;
+			FTransform TM = TMs[Index] * SkeletalMeshComponentTM;
+			BoneDatas[Index].CurrentTM = TM;
+		}
+		bHasJustTeleported = DetectTeleport(
+			ComponentTM, SkeletalMeshComponentTM, TeleportDistanceThreshold, TeleportRotationThreshold);
+	}
+	else
+	{
+		BoneDatas.Empty(NumTMs);
+		for (const FTransform& BoneTM : TMs)
+		{
+			FTransform TM = BoneTM * SkeletalMeshComponentTM;
+			BoneDatas.Emplace(TM, TM);
+		}
+		bHasJustTeleported = true;
+	}
+	ComponentTM = SkeletalMeshComponentTM;
+}
+
+//======================================================================================================================
+FVector FBoneData::CalculateLinearVelocity(float Dt) const
+{
+	if (Dt <= 0)
+	{
+		return FVector::ZeroVector;
+	}
+
+	return (CurrentTM.GetTranslation() - PreviousTM.GetTranslation()) / Dt;
+}
+
+//======================================================================================================================
+FVector FBoneData::CalculateAngularVelocity(float Dt) const
+{
+	if (Dt <= 0)
+	{
+		return FVector::ZeroVector;
+	}
+
+	// Note that quats multiply in the opposite order to TMs, and must be in the same hemisphere.
+	const FQuat Q = CurrentTM.GetRotation();
+	FQuat PrevQ = PreviousTM.GetRotation();
+	PrevQ.EnforceShortestArcWith(Q);
+	const FQuat DeltaQ = Q * PrevQ.Inverse();
+	return DeltaQ.ToRotationVector() / Dt;
+}
+
+} // namespace PhysicsControl
+} // namespace UE
+
diff --git a/Engine/Plugins/Experimental/PhysicsControl/Source/PhysicsControl/Private/RigidBodyPoseData.cpp b/Engine/Plugins/Experimental/PhysicsControl/Source/PhysicsControl/Private/RigidBodyPoseData.cpp
--- a/Engine/Plugins/Experimental/PhysicsControl/Source/PhysicsControl/Private/RigidBodyPoseData.cpp
+++ b/Engine/Plugins/Experimental/PhysicsControl/Source/PhysicsControl/Private/RigidBodyPoseData.cpp
@@ -4,7 +4,9 @@
 #include "Animation/AnimNodeBase.h"
 #include "Animation/AnimInstanceProxy.h"
 
-namespace RigidBodyWithControl
+namespace UE
+{
+namespace PhysicsControl
 {
 
 //======================================================================================================================
@@ -47,4 +49,7 @@ void FRigidBodyPoseData::Update(
 	}
 }
 
-}
\ No newline at end of file
+} // namespace PhysicsControl
+} // namespace UE
+
+
diff --git a/Engine/Plugins/Experimental/PhysicsControl/Source/PhysicsControl/Private/RigidBodyWithControl.cpp b/Engine/Plugins/Experimental/PhysicsControl/Source/PhysicsControl/Private/RigidBodyWithControl.cpp
--- a/Engine/Plugins/Experimental/PhysicsControl/Source/PhysicsControl/Private/RigidBodyWithControl.cpp
+++ b/Engine/Plugins/Experimental/PhysicsControl/Source/PhysicsControl/Private/RigidBodyWithControl.cpp
@@ -557,19 +557,19 @@ static bool UpdateDriveSpringDamperSettings(
 }
 
 //======================================================================================================================
-static RigidBodyWithControl::FPosQuat CalculateTargetTM(
-	const Chaos::FPBDJointSettings&                 JointSettings, 
-	const RigidBodyWithControl::FRigidBodyPoseData& PoseData,
-	const int32                                     ParentBodyIndex, 
-	const int32                                     ChildBodyIndex)
+static UE::PhysicsControl::FPosQuat CalculateTargetTM(
+	const Chaos::FPBDJointSettings&               JointSettings, 
+	const UE::PhysicsControl::FRigidBodyPoseData& PoseData,
+	const int32                                   ParentBodyIndex, 
+	const int32                                   ChildBodyIndex)
 {
-	const RigidBodyWithControl::FPosQuat ChildTargetTM =
-		RigidBodyWithControl::FPosQuat(JointSettings.ConnectorTransforms[ConstraintChildIndex]) * 
+	const UE::PhysicsControl::FPosQuat ChildTargetTM =
+		UE::PhysicsControl::FPosQuat(JointSettings.ConnectorTransforms[ConstraintChildIndex]) * 
 		PoseData.GetTM(ChildBodyIndex);
 	if (ParentBodyIndex >= 0)
 	{
-		const RigidBodyWithControl::FPosQuat ParentTargetTM =
-			RigidBodyWithControl::FPosQuat(JointSettings.ConnectorTransforms[ConstraintParentIndex]) * 
+		const UE::PhysicsControl::FPosQuat ParentTargetTM =
+			UE::PhysicsControl::FPosQuat(JointSettings.ConnectorTransforms[ConstraintParentIndex]) * 
 			PoseData.GetTM(ParentBodyIndex);
 		return ChildTargetTM * ParentTargetTM.Inverse();
 	}
@@ -619,13 +619,13 @@ void FAnimNode_RigidBodyWithControl::ApplyControl(FRigidBodyControlRecord& Contr
 					checkSlow((ControlRecord.Control.ParentBoneName.IsNone() ? -1 : 
 						FindBodyIndexFromBoneName(ControlRecord.Control.ParentBoneName)) == ControlRecord.ParentBodyIndex);
 
-					RigidBodyWithControl::FPosQuat TargetTM(
+					UE::PhysicsControl::FPosQuat TargetTM(
 						ControlRecord.ControlTarget.TargetOrientation, 
 						ControlRecord.ControlTarget.TargetPosition);
 
 					if (ControlRecord.ControlData.bUseSkeletalAnimation)
 					{
-						RigidBodyWithControl::FPosQuat AnimTargetTM = CalculateTargetTM(
+						UE::PhysicsControl::FPosQuat AnimTargetTM = CalculateTargetTM(
 							JointSettings, PoseData, ControlRecord.ParentBodyIndex, ControlRecord.ChildBodyIndex);
 						TargetTM = TargetTM * AnimTargetTM;
 					}
@@ -635,7 +635,8 @@ void FAnimNode_RigidBodyWithControl::ApplyControl(FRigidBodyControlRecord& Contr
 
 					if ((DeltaTime * ControlRecord.ControlData.LinearTargetVelocityMultiplier) != 0)
 					{
-						FVector Velocity = (TargetTM.GetTranslation() - ControlRecord.PrevTargetTM.GetTranslation()) / DeltaTime;
+						FVector Velocity = 
+							(TargetTM.GetTranslation() - ControlRecord.PreviousTargetTM.GetTranslation()) / DeltaTime;
 						Constraint->SetLinearDriveVelocityTarget(
 							Velocity * ControlRecord.ControlData.LinearTargetVelocityMultiplier);
 					}
@@ -648,7 +649,7 @@ void FAnimNode_RigidBodyWithControl::ApplyControl(FRigidBodyControlRecord& Contr
 					{
 						// Note that quats multiply in the opposite order to TMs, and must be in the same hemisphere.
 						const FQuat Q = TargetTM.GetRotation();
-						FQuat PrevQ = ControlRecord.PrevTargetTM.GetRotation();
+						FQuat PrevQ = ControlRecord.PreviousTargetTM.GetRotation();
 						PrevQ.EnforceShortestArcWith(Q);
 						const FQuat DeltaQ = Q * PrevQ.Inverse();
 						const FVector AngularVelocity = DeltaQ.ToRotationVector() / DeltaTime;
@@ -662,7 +663,7 @@ void FAnimNode_RigidBodyWithControl::ApplyControl(FRigidBodyControlRecord& Contr
 					}
 
 
-					ControlRecord.PrevTargetTM = TargetTM;
+					ControlRecord.PreviousTargetTM = TargetTM;
 					ControlRecord.ExpectedUpdateCounter = PoseData.UpdateCounter;
 					ControlRecord.ExpectedUpdateCounter.Increment();
 				}
@@ -878,7 +879,7 @@ void FAnimNode_RigidBodyWithControl::ApplyKinematicTargets()
 				const int32 BodyIndex = FindBodyIndexFromBoneName(ModifierRecord->Modifier.BoneName);
 				if (ActorHandle->GetIsKinematic() && BodyIndex != INDEX_NONE)
 				{
-					RigidBodyWithControl::FPosQuat TM(Target.TargetOrientation, Target.TargetPosition);
+					UE::PhysicsControl::FPosQuat TM(Target.TargetOrientation, Target.TargetPosition);
 					if (Target.bUseSkeletalAnimation)
 					{
 						TM = TM * PoseData.GetTM(BodyIndex);
@@ -968,59 +969,65 @@ void FAnimNode_RigidBodyWithControl::ApplyCurrentControlProfile()
 
 //======================================================================================================================
 void FAnimNode_RigidBodyWithControl::TransformConstraintsToMatchSkeletalMesh(
-	const USkeletalMesh* const SkeletalMeshAsset,
+	const USkeletalMesh*             SkeletalMeshAsset,
 	const MapConstraintsBehaviorType PositionBehavior,
 	const MapConstraintsBehaviorType OrientationBehavior,
-	TArray<FConstraintInstance*>& ConstraintInstances)
+	TArray<FConstraintInstance*>&    ConstraintInstances)
 {
 	// Bone1 = Child
 	// Bone2 = Parent 
 
 	if (SkeletalMeshAsset != nullptr)
 	{
-		const bool bAuthoredTransformRequired = ((PositionBehavior == MapConstraintsBehaviorType::AuthoredSkeleton) || (OrientationBehavior == MapConstraintsBehaviorType::AuthoredSkeleton)) && (PhysicsAssetAuthoredSkeletalMesh != nullptr) && (SkeletalMeshAsset != PhysicsAssetAuthoredSkeletalMesh);
-		const bool bDefaultTransformRequired = (PositionBehavior == MapConstraintsBehaviorType::DefaultTransform) || (OrientationBehavior == MapConstraintsBehaviorType::DefaultTransform);
+		const bool bAuthoredTransformRequired = 
+			((PositionBehavior == MapConstraintsBehaviorType::AuthoredSkeleton) || 
+				(OrientationBehavior == MapConstraintsBehaviorType::AuthoredSkeleton)) && 
+			(PhysicsAssetAuthoredSkeletalMesh != nullptr) && 
+			(SkeletalMeshAsset != PhysicsAssetAuthoredSkeletalMesh);
+		const bool bDefaultTransformRequired = 
+			(PositionBehavior == MapConstraintsBehaviorType::DefaultTransform) || 
+			(OrientationBehavior == MapConstraintsBehaviorType::DefaultTransform);
 
 		if (bAuthoredTransformRequired || bDefaultTransformRequired)
 		{
 #if !NO_LOGGING
-			const FString AuthoredSkeletalMeshName = (PhysicsAssetAuthoredSkeletalMesh) ? PhysicsAssetAuthoredSkeletalMesh->GetName() : FString("UNDEFINED");
+			const FString AuthoredSkeletalMeshName = 
+				(PhysicsAssetAuthoredSkeletalMesh) ? PhysicsAssetAuthoredSkeletalMesh->GetName() : FString("UNDEFINED");
 			UE_LOGFMT(LogPhysicsControl, Log,
 				"Modify Constraint parent transforms to match the current skeleton \"{0}\". Settings: Authored Skeleton {1}, Position set from {2}, Orientation set from {3}",
 				SkeletalMeshAsset->GetName(), AuthoredSkeletalMeshName, MapConstraintsBehaviorTypeToString(PositionBehavior), MapConstraintsBehaviorTypeToString(OrientationBehavior));
 #endif
-			
 
-			for (FConstraintInstance* const ConstraintInstance : ConstraintInstances)
+			for (FConstraintInstance* ConstraintInstance : ConstraintInstances)
 			{
-				FTransform AuthoredCurrentRefFrame;
-				FTransform DefaultCurrentRefFrame;
+				const FReferenceSkeleton& SkeletalMeshReferenceSkeleton = SkeletalMeshAsset->GetRefSkeleton();
+				const FName ChildBoneName = ConstraintInstance->ConstraintBone1;
+				const FName ParentBoneName = ConstraintInstance->ConstraintBone2;
+
+				// This function might be overkill, but it handles the case that there are sketal
+				// bones missing in the physics hierarchy.
+				const FTransform CurrentChildRelParentTM = CalculateRelativeBoneTransform(
+					ChildBoneName, ParentBoneName, SkeletalMeshReferenceSkeleton);
 
+				FTransform AuthoredCurrentRefFrame;
 				if (bAuthoredTransformRequired)
 				{
-					const FTransform CurrentParentRelChildTM = CalculateRelativeBoneTransform(
-						ConstraintInstance->ConstraintBone1, ConstraintInstance->ConstraintBone2, SkeletalMeshAsset->GetRefSkeleton());
-					const FTransform OriginalParentRelChildTM = CalculateRelativeBoneTransform(
-						ConstraintInstance->ConstraintBone1, ConstraintInstance->ConstraintBone2, PhysicsAssetAuthoredSkeletalMesh->GetRefSkeleton());
+					const FTransform OriginalChildRelParentTM = CalculateRelativeBoneTransform(
+						ChildBoneName, ParentBoneName, 
+						PhysicsAssetAuthoredSkeletalMesh->GetRefSkeleton());
 
 					// Find the transform that maps the parent-bone-relative-to-the-child-bone transform
 					// in the original skeleton to the parent-bone-relative-to-the-child-bone transform
 					// in the current skeleton.
-					// Should be equivalent to CurrentParentRelChildTM * OriginalParentRelChildTM.Inverse()
+					// Should be equivalent to CurrentChildRelParentTM * OriginalChildRelParentTM.Inverse()
 					const FTransform OriginalToCurrentParentRelChildTM =
-						CurrentParentRelChildTM.GetRelativeTransform(OriginalParentRelChildTM);
+						CurrentChildRelParentTM.GetRelativeTransform(OriginalChildRelParentTM);
 
 					// Update the constraints transform relative to the parent bone.
 					const FTransform OriginalRefFrame = ConstraintInstance->GetRefFrame(EConstraintFrame::Frame2);
 					AuthoredCurrentRefFrame = OriginalToCurrentParentRelChildTM * OriginalRefFrame;
 				}
 
-				if (bDefaultTransformRequired)
-				{
-					DefaultCurrentRefFrame = CalculateRelativeBoneTransform(
-						ConstraintInstance->ConstraintBone1, ConstraintInstance->ConstraintBone2, SkeletalMeshAsset->GetRefSkeleton());
-				}
-
 #if !NO_LOGGING
 				const FTransform LogPreviousConstraintTransformRelParent = ConstraintInstance->GetRefFrame(EConstraintFrame::Frame2);
 #endif
@@ -1031,23 +1038,25 @@ void FAnimNode_RigidBodyWithControl::TransformConstraintsToMatchSkeletalMesh(
 				}
 				else if (PositionBehavior == MapConstraintsBehaviorType::DefaultTransform)
 				{
-					ConstraintInstance->SetRefPosition(EConstraintFrame::Frame2, DefaultCurrentRefFrame.GetTranslation());
+					ConstraintInstance->SetRefPosition(EConstraintFrame::Frame2, CurrentChildRelParentTM.GetTranslation());
 				}
 
 				if (OrientationBehavior == MapConstraintsBehaviorType::AuthoredSkeleton)
 				{
-					ConstraintInstance->SetRefOrientation(EConstraintFrame::Frame2, AuthoredCurrentRefFrame.GetUnitAxis(EAxis::X), AuthoredCurrentRefFrame.GetUnitAxis(EAxis::Y));
+					ConstraintInstance->SetRefOrientation(EConstraintFrame::Frame2, 
+						AuthoredCurrentRefFrame.GetUnitAxis(EAxis::X), AuthoredCurrentRefFrame.GetUnitAxis(EAxis::Y));
 				}
 				else if (OrientationBehavior == MapConstraintsBehaviorType::DefaultTransform)
 				{
-					ConstraintInstance->SetRefOrientation(EConstraintFrame::Frame2, DefaultCurrentRefFrame.GetUnitAxis(EAxis::X), DefaultCurrentRefFrame.GetUnitAxis(EAxis::Y));
+					ConstraintInstance->SetRefOrientation(EConstraintFrame::Frame2, 
+						CurrentChildRelParentTM.GetUnitAxis(EAxis::X), CurrentChildRelParentTM.GetUnitAxis(EAxis::Y));
 				}
 
 #if !NO_LOGGING
 				UE_LOGFMT(LogPhysicsControl, Log,
 					"Constraint {0} - {1}  - transform relative to parent was {2} now {3}.",
-					ConstraintInstance->ConstraintBone1.ToString(),
-					ConstraintInstance->ConstraintBone2.ToString(),
+					ChildBoneName.ToString(),
+					ParentBoneName.ToString(),
 					LogPreviousConstraintTransformRelParent.ToString(),
 					ConstraintInstance->GetRefFrame(EConstraintFrame::Frame2).ToString());
 #endif
diff --git a/Engine/Plugins/Experimental/PhysicsControl/Source/PhysicsControl/Public/AnimNode_RigidBodyWithControl.h b/Engine/Plugins/Experimental/PhysicsControl/Source/PhysicsControl/Public/AnimNode_RigidBodyWithControl.h
--- a/Engine/Plugins/Experimental/PhysicsControl/Source/PhysicsControl/Public/AnimNode_RigidBodyWithControl.h
+++ b/Engine/Plugins/Experimental/PhysicsControl/Source/PhysicsControl/Public/AnimNode_RigidBodyWithControl.h
@@ -589,7 +589,7 @@ private:
 		bool bExpired;
 	};
 
-	TArray<RigidBodyWithControl::FOutputBoneData> OutputBoneData;
+	TArray<UE::PhysicsControl::FOutputBoneData> OutputBoneData;
 	// Note that the Bodies and Joints arrays will be the same size - i.e. there will be a
 	// correspondence between every joint and body. The joint will correspond to the parent of the
 	// body. Note that some joints will not be set.
@@ -601,7 +601,7 @@ private:
 	friend class UAnimGraphNode_RigidBodyWithControl;
 
 	// Each update we cache the incoming pose transforms in whatever space the simulation is running in
-	RigidBodyWithControl::FRigidBodyPoseData PoseData;
+	UE::PhysicsControl::FRigidBodyPoseData PoseData;
 
 	// Map of control records - they will be referenced by name
 	TMap<FName, FRigidBodyControlRecord> ControlRecords;
diff --git a/Engine/Plugins/Experimental/PhysicsControl/Source/PhysicsControl/Public/PhysicsControlComponent.h b/Engine/Plugins/Experimental/PhysicsControl/Source/PhysicsControl/Public/PhysicsControlComponent.h
--- a/Engine/Plugins/Experimental/PhysicsControl/Source/PhysicsControl/Public/PhysicsControlComponent.h
+++ b/Engine/Plugins/Experimental/PhysicsControl/Source/PhysicsControl/Public/PhysicsControlComponent.h
@@ -7,6 +7,7 @@
 #include "PhysicsControlNameRecords.h"
 #include "PhysicsControlProfileAsset.h"
 #include "PhysicsControlRecord.h"
+#include "PhysicsControlPoseData.h"
 
 #include "UObject/ObjectMacros.h"
 #include "Components/SceneComponent.h"
@@ -1507,30 +1508,6 @@ public:
 		const USkeletalMeshComponent* SkeletalMeshComponent,
 		const TArray<FName>&          BoneNames);
 
-	/**
-	 * Gets the linear velocities of the requested bones that will be used as targets (in world space). Target 
-	 * velocities for bones that are not found will be set to zero. Note that these targets will have been 
-	 * calculated and cached at the start of the Physics Control Component, so if using the built in tick, 
-	 * may be too old to be useful. If you manually update the component then you can access these target 
-	 * velocities prior to applying your own targets.
-	 */
-	UFUNCTION(BlueprintCallable, Category = PhysicsControl)
-	TArray<FVector> GetCachedBoneVelocities(
-		const USkeletalMeshComponent* SkeletalMeshComponent,
-		const TArray<FName>&          BoneNames);
-
-	/**
-	 * Gets the angular velocities of the requested bones that will be used as targets (in world space). Target 
-	 * velocities for bones that are not found will be set to zero. Note that these targets will have been 
-	 * calculated and cached at the start of the Physics Control Component, so if using the built in tick, 
-	 * may be too old to be useful. If you manually update the component then you can access these target 
-	 * velocities prior to applying your own targets.
-	 */
-	UFUNCTION(BlueprintCallable, Category = PhysicsControl)
-	TArray<FVector> GetCachedBoneAngularVelocities(
-		const USkeletalMeshComponent* SkeletalMeshComponent,
-		const TArray<FName>&          BoneNames);
-
 	/**
 	 * Gets the transforms of the requested bone that will be used as a target (in world space). Targets for bones
 	 * that are not found will be set to Identity. Note that these targets will have been calculated and cached
@@ -1567,30 +1544,6 @@ public:
 		const USkeletalMeshComponent* SkeletalMeshComponent,
 		const FName                   BoneName);
 
-	/**
-	 * Gets the linear velocity of the requested bone that will be used as a target (in world space). Target 
-	 * velocities for bones that are not found will be set to zero. Note that these targets will have been 
-	 * calculated and cached at the start of the Physics Control Component, so if using the built in tick, 
-	 * may be too old to be useful. If you manually update the component then you can access these target 
-	 * velocities prior to applying your own targets.
-	 */
-	UFUNCTION(BlueprintCallable, Category = PhysicsControl)
-	FVector GetCachedBoneVelocity(
-		const USkeletalMeshComponent* SkeletalMeshComponent,
-		const FName                   BoneName);
-
-	/**
-	 * Gets the angular velocity of the requested bone that will be used as a target (in world space). Target 
-	 * velocities for bones that are not found will be set to zero. Note that these targets will have been 
-	 * calculated and cached at the start of the Physics Control Component update, so if using the built in tick, 
-	 * may be too old to be useful. If you manually update the component then you can access these target 
-	 * velocities prior to applying your own targets.
-	 */
-	UFUNCTION(BlueprintCallable, Category = PhysicsControl)
-	FVector GetCachedBoneAngularVelocity(
-		const USkeletalMeshComponent* SkeletalMeshComponent,
-		const FName                   BoneName);
-
 	/**
 	 * This allows the caller to override the target that will have been calculated and cached at the start of 
 	 * the Physics Control Component update. This is unlikely to be useful when using the built in tick, 
@@ -1603,9 +1556,7 @@ public:
 	bool SetCachedBoneData(
 		const USkeletalMeshComponent* SkeletalMeshComponent,
 		const FName                   BoneName, 
-		const FTransform&             TM,
-		const FVector                 Velocity,
-		const FVector                 AngularVelocity);
+		const FTransform&             TM);
 
 	/**
 	 * This flags the body associated with the modifier to set (using teleport) its position and velocity to 
@@ -1740,14 +1691,16 @@ protected:
 	 * Retrieves the bone data for the specified bone given the skeletal mesh component.
 	 * 
 	 * @param OutBoneData If successful, this will contain the output bone data.
+	 * @param OutPoseData If successful, this will contain the pose data structure, in case more info is required
 	 * @param InSkeletalMeshComponent Required to be a valid pointer to a skeletal mesh component
 	 * @param InBoneName The name of the bone to retrieve data for
 	 * @return true if the bone and data were found, false if not (in which case warnings will be logged)
 	 */
 	bool GetBoneData(
-		FCachedSkeletalMeshData::FBoneData& OutBoneData,
-		const USkeletalMeshComponent*       InSkeletalMeshComponent,
-		const FName                         InBoneName) const;
+		UE::PhysicsControl::FBoneData&                      OutBoneData,
+		const UE::PhysicsControl::FPhysicsControlPoseData*& OutPoseData,
+		const USkeletalMeshComponent*                       InSkeletalMeshComponent,
+		const FName                                         InBoneName) const;
 
 	/*
 	 * Retrieves the bone data for the specified bone given the skeletal mesh component, for modification
@@ -1758,7 +1711,7 @@ protected:
 	 * @return true if the bone and data were found, false if not (in which case warnings will be logged)
 	 */
 	bool GetModifiableBoneData(
-		FCachedSkeletalMeshData::FBoneData*& OutBoneData,
+		UE::PhysicsControl::FBoneData*&      OutBoneData,
 		const USkeletalMeshComponent*        InSkeletalMeshComponent,
 		const FName                          InBoneName);
 
@@ -1779,20 +1732,6 @@ protected:
 	 */
 	void UpdateCachedSkeletalBoneData(float DeltaTime);
 
-	/**
-	 * @return true if the difference between the old and new TMs exceeds the teleport
-	 * translation/rotation thresholds
-	 */
-	bool DetectTeleport(const FTransform& OldTM, const FTransform& NewTM) const;
-
-	/**
-	 * @return true if the difference between the old and new TMs exceeds the teleport
-	 * translation/rotation thresholds
-	 */
-	bool DetectTeleport(
-		const FVector& OldPosition, const FQuat& OldOrientation,
-		const FVector& NewPosition, const FQuat& NewOrientation) const;
-
 	/**
 	 * Terminates the underlying physical constraints, resets our internal stored state for each control,
 	 * and optionally deletes all record of the controls.
@@ -1836,10 +1775,11 @@ protected:
 	 */
 	void CalculateControlTargetData(
 		FTransform&                  OutTargetTM, 
+		FTransform&                  OutSkeletalTargetTM,
 		FVector&                     OutTargetVelocity,
 		FVector&                     OutTargetAngularVelocity,
 		const FPhysicsControlRecord& Record,
-		bool                         bCalculateVelocity) const;
+		bool                         bUsePreviousSkeletalTargetTM) const;
 
 	/** Updates the body based on the modifier */
 	void ApplyBodyModifier(FPhysicsBodyModifierRecord& BodyModifier);
@@ -1887,7 +1827,7 @@ protected:
 
 	// Cached transforms from each skeletal mesh we're working with. Will be updated at the
 	// beginning of each tick.
-	TMap<TWeakObjectPtr<USkeletalMeshComponent>, FCachedSkeletalMeshData> CachedSkeletalMeshDatas;
+	TMap<TWeakObjectPtr<USkeletalMeshComponent>, UE::PhysicsControl::FPhysicsControlPoseData> CachedPoseDatas;
 
 	// Track which skeletons have been affected by a body modifier - some settings get overridden
 	// and then need to be restored when the last body modifier is destroyed.
@@ -1899,4 +1839,7 @@ protected:
 	// Keep track of the names of everything we have created
 	FPhysicsControlNameRecords NameRecords;
 
+	// Update counter - incremented every tick, and used to check whether previous targets etc are
+	// valid when calculating velocities.
+	int64 CurrentUpdateCounter = 0;
 };
diff --git a/Engine/Plugins/Experimental/PhysicsControl/Source/PhysicsControl/Public/PhysicsControlPoseData.h b/Engine/Plugins/Experimental/PhysicsControl/Source/PhysicsControl/Public/PhysicsControlPoseData.h
new file mode 100644
--- /dev/null
+++ b/Engine/Plugins/Experimental/PhysicsControl/Source/PhysicsControl/Public/PhysicsControlPoseData.h
@@ -0,0 +1,122 @@
+// Copyright Epic Games, Inc. All Rights Reserved.
+#pragma once
+
+#include "CoreMinimal.h"
+#include "Engine/EngineTypes.h"
+
+namespace UE
+{
+namespace PhysicsControl
+{
+
+// @return true if the difference between the old and new TMs exceeds the teleport
+// translation/rotation thresholds
+bool DetectTeleport(
+	const FVector& OldPosition, const FQuat& OldOrientation,
+	const FVector& NewPosition, const FQuat& NewOrientation,
+	float DistanceThreshold, float RotationThreshold);
+
+// @return true if the difference between the old and new TMs exceeds the teleport
+// translation/rotation thresholds
+bool DetectTeleport(const FTransform& OldTM, const FTransform& NewTM,
+	float DistanceThreshold, float RotationThreshold);
+
+//======================================================================================================================
+// Simple minimal implementation of a "FTransform without scale"
+struct FPosQuat
+{
+	FPosQuat(const FVector& Pos, const FQuat& Quat) : Translation(Pos), Rotation(Quat) {}
+	FPosQuat(const FQuat& Quat, const FVector& Pos) : Translation(Pos), Rotation(Quat) {}
+	FPosQuat(const FRotator& Rotator, const FVector& Pos) : Translation(Pos), Rotation(Rotator) {}
+	FPosQuat() : Translation(EForceInit::ForceInitToZero), Rotation(EForceInit::ForceInit) {}
+	FPosQuat(const FTransform& TM) : Translation(TM.GetTranslation()), Rotation(TM.GetRotation()) {}
+	FPosQuat(ENoInit) {}
+
+	FORCEINLINE FVector GetTranslation() const { return Translation; }
+	FORCEINLINE FQuat GetRotation() const { return Rotation; }
+
+	FORCEINLINE FTransform ToTransform() const
+	{
+		return FTransform(Rotation, Translation);
+	}
+
+	// Note that multiplication operates in the same sense as FTransform (i.e. "backwards")
+	FORCEINLINE FPosQuat operator*(const FPosQuat& Other) const
+	{
+		FQuat OutRotation = Other.Rotation * Rotation;
+		FVector OutTranslation = Other.Rotation * (Translation) + Other.Translation;
+		return FPosQuat(OutTranslation, OutRotation);
+	}
+
+	FORCEINLINE FVector operator*(const FVector& Position) const
+	{
+		return Translation + Rotation * Position;
+	}
+
+	FORCEINLINE FPosQuat Inverse() const
+	{
+		const FQuat OutRotation = Rotation.Inverse();
+		return FPosQuat(OutRotation * -Translation, OutRotation);
+	}
+
+	FVector Translation;
+	FQuat Rotation;
+};
+
+//======================================================================================================================
+struct FBoneData
+{
+	FBoneData() {}
+	FBoneData(const FPosQuat& InCurrentTM, const FPosQuat& InPreviousTM)
+		: CurrentTM(InCurrentTM), PreviousTM(InPreviousTM) {}
+
+	FVector CalculateLinearVelocity(float Dt) const;
+	FVector CalculateAngularVelocity(float Dt) const;
+
+	FPosQuat CurrentTM;
+	FPosQuat PreviousTM;
+};
+
+//======================================================================================================================
+// Caches the pose for PhysicsControlComponent
+struct FPhysicsControlPoseData
+{
+public:
+	FPhysicsControlPoseData() : ReferenceCount(0), bHasJustTeleported(true), DeltaTime(0) {}
+
+public:
+	void Update(
+		USkeletalMeshComponent* SkeletalMesh,
+		const float             Dt,
+		const float             TeleportDistanceThreshold, 
+		const float             TeleportRotationThreshold);
+	void Reset() { BoneDatas.Empty(); bHasJustTeleported = true; }
+
+	const FBoneData& GetBoneData(int32 Index) const { return BoneDatas[Index]; }
+	FPosQuat GetCurrentTM(int32 Index) const { return BoneDatas[Index].CurrentTM; }
+	FPosQuat GetPreviousTM(int32 Index) const { return BoneDatas[Index].PreviousTM; }
+
+	bool IsValidIndex(const int32 Index) const { return BoneDatas.IsValidIndex(Index); }
+	bool IsEmpty() const { return BoneDatas.IsEmpty(); }
+
+	// The cached skeletal data, updated at the start of each tick
+	TArray<FBoneData> BoneDatas;
+
+	// Track when skeletal meshes are going to be used so this entry can be removed, and also so we
+	// can add a tick dependency.
+	int32 ReferenceCount;
+
+	// The component transform. This is only stored so we can detect teleports. 
+	FTransform ComponentTM;
+
+	// Whether the character has just teleported, in which case velocities should not be calculated
+	// based on deltas.
+	bool bHasJustTeleported;
+
+	// The DeltaTime used for the previous update (may be zero)
+	float DeltaTime;
+};
+
+} // namespace PhysicsControl
+} // namespace UE
+
diff --git a/Engine/Plugins/Experimental/PhysicsControl/Source/PhysicsControl/Public/PhysicsControlRecord.h b/Engine/Plugins/Experimental/PhysicsControl/Source/PhysicsControl/Public/PhysicsControlRecord.h
--- a/Engine/Plugins/Experimental/PhysicsControl/Source/PhysicsControl/Public/PhysicsControlRecord.h
+++ b/Engine/Plugins/Experimental/PhysicsControl/Source/PhysicsControl/Public/PhysicsControlRecord.h
@@ -3,6 +3,7 @@
 #pragma once
 
 #include "PhysicsControlData.h"
+#include "PhysicsControlPoseData.h"
 #include "Components/SkeletalMeshComponent.h"
 #include "Engine/EngineTypes.h"
 
@@ -22,6 +23,7 @@ struct FPhysicsControlRecord
 		UMeshComponent*              InChildMeshComponent)
 		: PhysicsControl(InControl)
 		, ControlTarget(InControlTarget)
+		, ExpectedUpdateCounter(-1)
 		, ParentMeshComponent(InParentMeshComponent)
 		, ChildMeshComponent(InChildMeshComponent)
 	{}
@@ -51,9 +53,18 @@ struct FPhysicsControlRecord
 	 * skeletal meshes have the option to use skeletal animation as well, in which case these targets are 
 	 * expressed as relative to that animation.
 	 */
-	UPROPERTY(EditAnywhere, BlueprintReadWrite, Category = PhysicsControl)
 	FPhysicsControlTarget ControlTarget;
 
+	// The previous skeletal control target. This will have been set at the end of a previous update
+	// (but only if the control was enabled etc), so to check if it is valid, check the update
+	// counter. Note that explicit control targets (which contain their own velocity) will be added
+	// onto this.
+	UE::PhysicsControl::FPosQuat PreviousSkeletalTargetTM;
+
+	// Only use the previous target TM if the current counter is equal to this expected counter. The
+	// expected update counter will be set when the control/previous target TM has just been updated.
+	int64 ExpectedUpdateCounter;
+
 	/**  The mesh that will be doing the driving. Blank/non-existent means it will happen in world space */
 	TWeakObjectPtr<UMeshComponent> ParentMeshComponent;
 
@@ -128,81 +139,3 @@ public:
 	int32 ReferenceCount;
 };
 
-/**
- * Used internally/only at runtime to cache skeletal transforms at the start of the tick, to avoid
- * calculating them separately for every control.
- */
-struct FCachedSkeletalMeshData
-{
-public:
-	FCachedSkeletalMeshData() : ReferenceCount(0) {}
-
-public:
-	struct FBoneData
-	{
-		FBoneData()
-			: Position(FVector::ZeroVector), Orientation(FQuat::Identity)
-			, Velocity(FVector::ZeroVector), AngularVelocity(FVector::ZeroVector) {}
-		FBoneData(const FVector& InPosition, const FQuat& InOrientation)
-			: Position(InPosition), Orientation(InOrientation),
-			Velocity(FVector::ZeroVector), AngularVelocity(FVector::ZeroVector) {}
-
-		/**
-		 * Sets position/orientation and calculates velocity/angular velocity - requires Dt > 0 
-		 */
-		void Update(const FVector& InPosition, const FQuat& InOrientation, float Dt);
-
-		/**
-		 * Sets position/orientation, and sets velocity to zero
-		 */
-		void Update(const FVector& InPosition, const FQuat& InOrientation);
-
-		FTransform GetTM() const { return FTransform(Orientation, Position); }
-
-		FVector Position;
-		FQuat   Orientation;
-		FVector Velocity;
-		FVector AngularVelocity;
-	};
-
-public:
-
-	/**
-	 * The cached skeletal data, updated at the start of each tick
-	 */
-	TArray<FBoneData> BoneData;
-
-	/**
-	 * The component transform. This is only stored so we can detect teleports
-	 */
-	FTransform ComponentTM;
-
-	/**
-	 * Track when skeletal meshes are going to be used so this entry can be removed, and also so we
-	 * can add a tick dependency
-	 */
-	int32 ReferenceCount;
-};
-
-//======================================================================================================================
-inline void FCachedSkeletalMeshData::FBoneData::Update(const FVector& InPosition, const FQuat& InOrientation, float Dt)
-{
-	check(Dt > 0);
-
-	Velocity = (InPosition - Position) / Dt;
-	Orientation.EnforceShortestArcWith(InOrientation);
-	// Note that quats multiply in the opposite order to TMs
-	const FQuat DeltaQ = InOrientation * Orientation.Inverse();
-	AngularVelocity = DeltaQ.ToRotationVector() / Dt;
-	Position = InPosition;
-	Orientation = InOrientation;
-}
-
-//======================================================================================================================
-inline void FCachedSkeletalMeshData::FBoneData::Update(const FVector& InPosition, const FQuat& InOrientation)
-{
-	Position = InPosition;
-	Orientation = InOrientation;
-	Velocity = FVector::ZeroVector;
-	AngularVelocity = FVector::ZeroVector;
-}
diff --git a/Engine/Plugins/Experimental/PhysicsControl/Source/PhysicsControl/Public/RigidBodyControlData.h b/Engine/Plugins/Experimental/PhysicsControl/Source/PhysicsControl/Public/RigidBodyControlData.h
--- a/Engine/Plugins/Experimental/PhysicsControl/Source/PhysicsControl/Public/RigidBodyControlData.h
+++ b/Engine/Plugins/Experimental/PhysicsControl/Source/PhysicsControl/Public/RigidBodyControlData.h
@@ -5,7 +5,7 @@
 
 #include "PhysicsControlLimbData.h"
 #include "PhysicsControlData.h"
-#include "RigidBodyPoseData.h"
+#include "PhysicsControlPoseData.h"
 
 #include "Physics/ImmediatePhysics/ImmediatePhysicsDeclares.h"
 #include "Animation/AnimTypes.h"
@@ -120,7 +120,7 @@ struct FRigidBodyControlRecord
 
 	// The previous control target. This will have been set at the end of a previous update (but
 	// only if the control was enabled etc), so to check if it is valid, check the update counter.
-	RigidBodyWithControl::FPosQuat PrevTargetTM;
+	UE::PhysicsControl::FPosQuat PreviousTargetTM;
 
 	// Update counter set when the control was last updated.
 	// TODO just store the count we're interested in rather than the whole structure
diff --git a/Engine/Plugins/Experimental/PhysicsControl/Source/PhysicsControl/Public/RigidBodyPoseData.h b/Engine/Plugins/Experimental/PhysicsControl/Source/PhysicsControl/Public/RigidBodyPoseData.h
--- a/Engine/Plugins/Experimental/PhysicsControl/Source/PhysicsControl/Public/RigidBodyPoseData.h
+++ b/Engine/Plugins/Experimental/PhysicsControl/Source/PhysicsControl/Public/RigidBodyPoseData.h
@@ -2,15 +2,16 @@
 #pragma once
 
 #include "CoreMinimal.h"
-#include "Engine/EngineTypes.h"
-#include "BonePose.h"
+#include "PhysicsControlPoseData.h"
 
 // Use the simulation space functions from the RBAN
 #include "BoneControllers/AnimNode_RigidBody.h" 
 
 struct FComponentSpacePoseContext;
 
-namespace RigidBodyWithControl
+namespace UE
+{
+namespace PhysicsControl
 {
 
 //======================================================================================================================
@@ -27,66 +28,30 @@ struct FOutputBoneData
 	int32 ParentBodyIndex;
 };
 
-//======================================================================================================================
-// Simple minimal implementation of a "FTransform without scale"
-struct FPosQuat
-{
-	FPosQuat(const FVector& Pos, const FQuat& Quat) : Translation(Pos), Rotation(Quat) {}
-	FPosQuat(const FQuat& Quat, const FVector& Pos) : Translation(Pos), Rotation(Quat) {}
-	FPosQuat(const FRotator& Rotator, const FVector& Pos) : Translation(Pos), Rotation(Rotator) {}
-	FPosQuat() : Translation(EForceInit::ForceInitToZero), Rotation(EForceInit::ForceInit) {}
-	FPosQuat(const FTransform& TM) : Translation(TM.GetTranslation()), Rotation(TM.GetRotation()) {}
-	FPosQuat(ENoInit) {}
-
-	FORCEINLINE FVector GetTranslation() const { return Translation; }
-	FORCEINLINE FQuat GetRotation() const { return Rotation; }
-
-	FORCEINLINE FTransform ToTransform() const
-	{
-		return FTransform(Rotation, Translation);
-	}
-
-	// Note that multiplication operates in the same sense as FTransform (i.e. "backwards")
-	FORCEINLINE FPosQuat operator*(const FPosQuat& Other) const
-	{
-		FQuat OutRotation = Other.Rotation * Rotation;
-		FVector OutTranslation = Other.Rotation * (Translation) + Other.Translation;
-		return FPosQuat(OutTranslation, OutRotation);
-	}
-
-	FORCEINLINE FVector operator*(const FVector& Position) const
-	{
-		return Translation + Rotation * Position;
-	}
-
-	FORCEINLINE FPosQuat Inverse() const
-	{
-		const FQuat OutRotation = Rotation.Inverse();
-		return FPosQuat(OutRotation * -Translation, OutRotation);
-	}
-
-	FVector Translation;
-	FQuat Rotation;
-};
 
 //======================================================================================================================
+// Caches the pose for RigidBodyWithControl
 struct FRigidBodyPoseData
 {
+public:
+	FRigidBodyPoseData() {}
+
+public:
 	void Update(
-		FComponentSpacePoseContext&    ComponentSpacePoseContext,
+		FComponentSpacePoseContext& ComponentSpacePoseContext,
 		const TArray<FOutputBoneData>& OutputBoneData,
 		const ESimulationSpace         SimulationSpace,
-		const FBoneReference&          BaseBoneRef,
-		const FGraphTraversalCounter&  InUpdateCounter);
+		const FBoneReference& BaseBoneRef,
+		const FGraphTraversalCounter& InUpdateCounter);
 
-	FPosQuat GetTM(int32 Index) const { return BoneTMs[Index]; }
+	UE::PhysicsControl::FPosQuat GetTM(int32 Index) const { return BoneTMs[Index]; }
 	bool IsValidIndex(const int32 Index) const { return BoneTMs.IsValidIndex(Index); }
 	bool IsEmpty() const { return BoneTMs.IsEmpty(); }
 
 	/**
 	 * The cached skeletal data, updated at the start of each tick
 	 */
-	TArray<FPosQuat> BoneTMs;
+	TArray<UE::PhysicsControl::FPosQuat> BoneTMs;
 
 	// Track when we were currently/last updated so the user can detect missing updates if calculating
 	// velocity etc
@@ -96,4 +61,5 @@ struct FRigidBodyPoseData
 	FGraphTraversalCounter ExpectedUpdateCounter;
 };
 
-}
\ No newline at end of file
+} // namespace PhysicsControl
+} // namespace UE
