Make PhysicsControl work with PrimitiveComponent, not just meshes. Also adds a function that can be used to zero the implied velocities when moving a character. Some help text improvements.

diff --git a/Engine/Plugins/Experimental/PhysicsControl/Config/BasePhysicsControl.ini b/Engine/Plugins/Experimental/PhysicsControl/Config/BasePhysicsControl.ini
--- a/Engine/Plugins/Experimental/PhysicsControl/Config/BasePhysicsControl.ini
+++ b/Engine/Plugins/Experimental/PhysicsControl/Config/BasePhysicsControl.ini
@@ -1,30 +1,15 @@
 [CoreRedirects]
-+StructRedirects=(OldName="/Script/PhysicsControl.PhysicsControlNameArray", NewName="/Script/PhysicsControl.PhysicsControlNames")
-+StructRedirects=(OldName="/Script/PhysicsControl.FPhysicsControlMultipliers", NewName="/Script/PhysicsControl.FPhysicsControlMultiplier")
++PropertyRedirects=(OldName="PhysicsControlComponent.CreateControl.ParentMeshComponent",NewName="ParentComponent")
++PropertyRedirects=(OldName="PhysicsControlComponent.CreateControl.ChildMeshComponent",NewName="ChildComponent")
+
++PropertyRedirects=(OldName="PhysicsControlComponent.CreateNamedControl.ParentMeshComponent",NewName="ParentComponent")
++PropertyRedirects=(OldName="PhysicsControlComponent.CreateNamedControl.ChildMeshComponent",NewName="ChildComponent")
+
++PropertyRedirects=(OldName="PhysicsControlComponent.SetControlParent.ParentMeshComponent",NewName="ParentComponent")
++PropertyRedirects=(OldName="PhysicsControlComponent.SetControlParents.ParentMeshComponent",NewName="ParentComponent")
++PropertyRedirects=(OldName="PhysicsControlComponent.SetControlParentsInSet.ParentMeshComponent",NewName="ParentComponent")
+
++PropertyRedirects=(OldName="PhysicsControlComponent.CreateBodyModifier.MeshComponent",NewName="Component")
++PropertyRedirects=(OldName="PhysicsControlComponent.CreateNamedBodyModifier.MeshComponent",NewName="Component")
 
-+FunctionRedirects=(OldName="PhysicsControlComponent.MakeControl",NewName="PhysicsControlComponent.CreateControl")
-+FunctionRedirects=(OldName="PhysicsControlComponent.MakeNamedControl",NewName="PhysicsControlComponent.CreateNamedControl")
-+FunctionRedirects=(OldName="PhysicsControlComponent.MakeControlsFromSkeletalMeshBelow",NewName="PhysicsControlComponent.CreateControlsFromSkeletalMeshBelow")
-+FunctionRedirects=(OldName="PhysicsControlComponent.MakeControlsFromSkeletalMeshAndConstraintProfileBelow",NewName="PhysicsControlComponent.CreateControlsFromSkeletalMeshAndConstraintProfileBelow")
-+FunctionRedirects=(OldName="PhysicsControlComponent.MakeControlsFromSkeletalMesh",NewName="PhysicsControlComponent.CreateControlsFromSkeletalMesh")
-+FunctionRedirects=(OldName="PhysicsControlComponent.MakeControlsFromSkeletalMeshAndConstraintProfile",NewName="PhysicsControlComponent.CreateControlsFromSkeletalMeshAndConstraintProfile")
-+FunctionRedirects=(OldName="PhysicsControlComponent.MakeControlsFromLimbBones",NewName="PhysicsControlComponent.CreateControlsFromLimbBones")
-+FunctionRedirects=(OldName="PhysicsControlComponent.MakeControlsFromLimbBonesAndConstraintProfile",NewName="PhysicsControlComponent.CreateControlsFromLimbBonesAndConstraintProfile")
-+FunctionRedirects=(OldName="PhysicsControlComponent.MakeBodyModifier",NewName="PhysicsControlComponent.CreateBodyModifier")
-+FunctionRedirects=(OldName="PhysicsControlComponent.MakeNamedBodyModifier",NewName="PhysicsControlComponent.CreateNamedBodyModifier")
-+FunctionRedirects=(OldName="PhysicsControlComponent.MakeBodyModifiersFromSkeletalMeshBelow",NewName="PhysicsControlComponent.CreateBodyModifiersFromSkeletalMeshBelow")
-+FunctionRedirects=(OldName="PhysicsControlComponent.MakeBodyModifiersFromLimbBones",NewName="PhysicsControlComponent.CreateBodyModifiersFromLimbBones")
-+FunctionRedirects=(OldName="PhysicsControlComponent.SetAllBodyModifierMovementType",NewName="PhysicsControlComponent.SetBodyModifiersMovementType")
-+FunctionRedirects=(OldName="PhysicsControlComponent.SetAllBodyModifierGravityMultipliers",NewName="PhysicsControlComponent.SetBodyModifiersGravityMultiplier")
-+FunctionRedirects=(OldName="PhysicsControlComponent.DestroyAllControls",NewName="PhysicsControlComponent.DestroyControls")
-+FunctionRedirects=(OldName="PhysicsControlComponent.SetAllControlTargets",NewName="PhysicsControlComponent.SetControlTargets")
-+FunctionRedirects=(OldName="PhysicsControlComponent.SetAllControlDatas",NewName="PhysicsControlComponent.SetControlDatas")
-+FunctionRedirects=(OldName="PhysicsControlComponent.SetAllControlsUseSkeletalAnimation",NewName="PhysicsControlComponent.SetControlsUseSkeletalAnimation")
-+FunctionRedirects=(OldName="PhysicsControlComponent.SetAllControlsInSetUseSkeletalAnimation",NewName="PhysicsControlComponent.SetControlsInSetUseSkeletalAnimation")
-+FunctionRedirects=(OldName="PhysicsControlComponent.SetAllControlsEnabled",NewName="PhysicsControlComponent.SetControlsEnabled")
-+FunctionRedirects=(OldName="PhysicsControlComponent.DestroyAllBodyModifiers",NewName="PhysicsControlComponent.DestroyBodyModifiers")
-+FunctionRedirects=(OldName="PhysicsControlComponent.SetAllBodyModifiersMovementType",NewName="PhysicsControlComponent.SetBodyModifiersMovementType")
-+FunctionRedirects=(OldName="PhysicsControlComponent.SetAllBodyModifiersCollisionType",NewName="PhysicsControlComponent.SetBodyModifiersCollisionType")
-+FunctionRedirects=(OldName="PhysicsControlComponent.SetAllBodyModifiersGravityMultiplier",NewName="PhysicsControlComponent.SetBodyModifiersGravityMultiplier")
-+FunctionRedirects=(OldName="PhysicsControlComponent.SetAllBodyModifiersUseSkeletalAnimation",NewName="PhysicsControlComponent.SetBodyModifiersUseSkeletalAnimation")
 
diff --git a/Engine/Plugins/Experimental/PhysicsControl/Source/PhysicsControl/Private/PhysicsControlComponent.cpp b/Engine/Plugins/Experimental/PhysicsControl/Source/PhysicsControl/Private/PhysicsControlComponent.cpp
--- a/Engine/Plugins/Experimental/PhysicsControl/Source/PhysicsControl/Private/PhysicsControlComponent.cpp
+++ b/Engine/Plugins/Experimental/PhysicsControl/Source/PhysicsControl/Private/PhysicsControlComponent.cpp
@@ -211,9 +211,9 @@ TMap<FName, FPhysicsControlLimbBones> UPhysicsControlComponent::GetLimbBonesFrom
 //======================================================================================================================
 // Note - some params passed by value to allow them to be unconnected and enable inline BP editing
 FName UPhysicsControlComponent::CreateControl(
-	UMeshComponent*               ParentMeshComponent,
+	UPrimitiveComponent*          ParentComponent,
 	const FName                   ParentBoneName,
-	UMeshComponent*               ChildMeshComponent,
+	UPrimitiveComponent*          ChildComponent,
 	const FName                   ChildBoneName,
 	const FPhysicsControlData     ControlData, 
 	const FPhysicsControlTarget   ControlTarget, 
@@ -222,7 +222,7 @@ FName UPhysicsControlComponent::CreateControl(
 {
 	const FName Name = UE::PhysicsControl::GetUniqueControlName(ParentBoneName, ChildBoneName, ControlRecords, NamePrefix);
 	if (CreateNamedControl(
-		Name, ParentMeshComponent, ParentBoneName, ChildMeshComponent, ChildBoneName, ControlData, ControlTarget, Set))
+		Name, ParentComponent, ParentBoneName, ChildComponent, ChildBoneName, ControlData, ControlTarget, Set))
 	{
 		return Name;
 	}
@@ -233,9 +233,9 @@ FName UPhysicsControlComponent::CreateControl(
 // Note - some params passed by value to allow them to be unconnected and enable inline BP editing
 bool UPhysicsControlComponent::CreateNamedControl(
 	const FName                   Name, 
-	UMeshComponent*               ParentMeshComponent,
+	UPrimitiveComponent*          ParentComponent,
 	const FName                   ParentBoneName,
-	UMeshComponent*               ChildMeshComponent,
+	UPrimitiveComponent*          ChildComponent,
 	const FName                   ChildBoneName,
 	const FPhysicsControlData     ControlData, 
 	const FPhysicsControlTarget   ControlTarget, 
@@ -248,18 +248,18 @@ bool UPhysicsControlComponent::CreateNamedControl(
 		return false;
 	}
 
-	if (!ChildMeshComponent)
+	if (!ChildComponent)
 	{
 		UE_LOG(LogPhysicsControl, Warning,
 			TEXT("Unable to make a Control as the child mesh component has not been set"));
 		return false;
 	}
 
-	if (USkeletalMeshComponent* SkeletalMeshComponent = Cast<USkeletalMeshComponent>(ParentMeshComponent))
+	if (USkeletalMeshComponent* SkeletalMeshComponent = Cast<USkeletalMeshComponent>(ParentComponent))
 	{
 		AddSkeletalMeshReferenceForCaching(SkeletalMeshComponent);
 	}
-	if (USkeletalMeshComponent* SkeletalMeshComponent = Cast<USkeletalMeshComponent>(ChildMeshComponent))
+	if (USkeletalMeshComponent* SkeletalMeshComponent = Cast<USkeletalMeshComponent>(ChildComponent))
 	{
 		AddSkeletalMeshReferenceForCaching(SkeletalMeshComponent);
 	}
@@ -267,7 +267,7 @@ bool UPhysicsControlComponent::CreateNamedControl(
 	FPhysicsControlRecord& NewRecord = ControlRecords.Add(
 		Name, FPhysicsControlRecord(
 			FPhysicsControl(ParentBoneName, ChildBoneName, ControlData), 
-			ControlTarget, ParentMeshComponent, ChildMeshComponent));
+			ControlTarget, ParentComponent, ChildComponent));
 	NewRecord.ResetControlPoint();
 
 	NameRecords.AddControl(Name, Set);
@@ -293,14 +293,13 @@ TArray<FName> UPhysicsControlComponent::CreateControlsFromSkeletalMeshBelow(
 		return Result;
 	}
 
-	int32 NumBones = SkeletalMeshComponent->GetNumBones();
-	UMeshComponent* ParentMeshComponent = 
+	UPrimitiveComponent* ParentComponent = 
 		(ControlType == EPhysicsControlType::ParentSpace) ? SkeletalMeshComponent : nullptr;
 
 	SkeletalMeshComponent->ForEachBodyBelow(
 		BoneName, bIncludeSelf, /*bSkipCustomType=*/false, 
 		[
-			this, PhysicsAsset, ParentMeshComponent, SkeletalMeshComponent, 
+			this, PhysicsAsset, ParentComponent, SkeletalMeshComponent, 
 			ControlType, &ControlData, Set, &Result
 		](const FBodyInstance* BI)
 		{
@@ -309,7 +308,7 @@ TArray<FName> UPhysicsControlComponent::CreateControlsFromSkeletalMeshBelow(
 				const FName ChildBoneName = PhysicsAsset->SkeletalBodySetups[BI->InstanceBodyIndex]->BoneName;
 
 				FName ParentBoneName;
-				if (ParentMeshComponent)
+				if (ParentComponent)
 				{
 					ParentBoneName = UE::PhysicsControl::GetPhysicalParentBone(
 						SkeletalMeshComponent, ChildBoneName);
@@ -319,7 +318,7 @@ TArray<FName> UPhysicsControlComponent::CreateControlsFromSkeletalMeshBelow(
 					}
 				}
 				const FName ControlName = CreateControl(
-					ParentMeshComponent, ParentBoneName, SkeletalMeshComponent, ChildBoneName,
+					ParentComponent, ParentBoneName, SkeletalMeshComponent, ChildBoneName,
 					ControlData, FPhysicsControlTarget(), 
 					FName(GetControlTypeName(ControlType).ToString().Append("_").Append(Set.ToString())));
 				if (!ControlName.IsNone())
@@ -429,13 +428,13 @@ TArray<FName> UPhysicsControlComponent::CreateControlsFromSkeletalMesh(
 		return Result;
 	}
 
-	UMeshComponent* ParentMeshComponent =
+	UPrimitiveComponent* ParentComponent =
 		(ControlType == EPhysicsControlType::ParentSpace) ? SkeletalMeshComponent : nullptr;
 
 	for (FName ChildBoneName : BoneNames)
 	{
 		FName ParentBoneName;
-		if (ParentMeshComponent)
+		if (ParentComponent)
 		{
 			ParentBoneName = UE::PhysicsControl::GetPhysicalParentBone(
 				SkeletalMeshComponent, ChildBoneName);
@@ -445,7 +444,7 @@ TArray<FName> UPhysicsControlComponent::CreateControlsFromSkeletalMesh(
 			}
 		}
 		const FName ControlName = CreateControl(
-			ParentMeshComponent, ParentBoneName, SkeletalMeshComponent, ChildBoneName,
+			ParentComponent, ParentBoneName, SkeletalMeshComponent, ChildBoneName,
 			ControlData, FPhysicsControlTarget(), 
 			FName(GetControlTypeName(ControlType).ToString().Append("_").Append(Set.ToString())));
 		if (!ControlName.IsNone())
@@ -529,7 +528,7 @@ TMap<FName, FPhysicsControlNames> UPhysicsControlComponent::CreateControlsFromLi
 	const TMap<FName, FPhysicsControlLimbBones>& LimbBones,
 	const EPhysicsControlType                    ControlType,
 	const FPhysicsControlData                    ControlData,
-	UMeshComponent*                              WorldComponent,
+	UPrimitiveComponent*                         WorldComponent,
 	FName                                        WorldBoneName,
 	FString                                      NamePrefix)
 {
@@ -586,15 +585,15 @@ TMap<FName, FPhysicsControlNames> UPhysicsControlComponent::CreateControlsFromLi
 				}
 			}
 
-			UMeshComponent* ParentMeshComponent = ParentSkeletalMeshComponent;
-			if (!ParentMeshComponent && WorldComponent)
+			UPrimitiveComponent* ParentComponent = ParentSkeletalMeshComponent;
+			if (!ParentComponent && WorldComponent)
 			{
-				ParentMeshComponent = WorldComponent;
+				ParentComponent = WorldComponent;
 				ParentBoneName = WorldBoneName;
 			}
 
 			const FName ControlName = CreateControl(
-				ParentMeshComponent, ParentBoneName, BonesInLimb.SkeletalMeshComponent.Get(), ChildBoneName,
+				ParentComponent, ParentBoneName, BonesInLimb.SkeletalMeshComponent.Get(), ChildBoneName,
 				ControlData, FPhysicsControlTarget(), FName(SetName));
 
 			if (!ControlName.IsNone())
@@ -757,25 +756,25 @@ void UPhysicsControlComponent::SetControlsInSetEnabled(FName SetName, bool bEnab
 
 //======================================================================================================================
 bool UPhysicsControlComponent::SetControlParent(
-	const FName     Name,
-	UMeshComponent* ParentMeshComponent,
-	const FName     ParentBoneName)
+	const FName          Name,
+	UPrimitiveComponent* ParentComponent,
+	const FName          ParentBoneName)
 {
 	FPhysicsControlRecord* Record = FindControlRecord(Name);
 	if (Record)
 	{
 		if (USkeletalMeshComponent* SkeletalMeshComponent = 
-			Cast<USkeletalMeshComponent>(Record->ParentMeshComponent.Get()))
+			Cast<USkeletalMeshComponent>(Record->ParentComponent.Get()))
 		{
 			RemoveSkeletalMeshReferenceForCaching(SkeletalMeshComponent);
 		}
 
-		if (USkeletalMeshComponent* SkeletalMeshComponent = Cast<USkeletalMeshComponent>(ParentMeshComponent))
+		if (USkeletalMeshComponent* SkeletalMeshComponent = Cast<USkeletalMeshComponent>(ParentComponent))
 		{
 			AddSkeletalMeshReferenceForCaching(SkeletalMeshComponent);
 		}
 
-		Record->ParentMeshComponent = ParentMeshComponent;
+		Record->ParentComponent = ParentComponent;
 		Record->PhysicsControl.ParentBoneName = ParentBoneName;
 		return Record->InitConstraint(this, Name);
 	}
@@ -789,23 +788,23 @@ bool UPhysicsControlComponent::SetControlParent(
 
 //======================================================================================================================
 void UPhysicsControlComponent::SetControlParents(
-		const TArray<FName>& Names,
-		UMeshComponent*      ParentMeshComponent,
-		const FName          ParentBoneName)
+	const TArray<FName>& Names,
+	UPrimitiveComponent* ParentComponent,
+	const FName          ParentBoneName)
 {
 	for (FName Name : Names)
 	{
-		SetControlParent(Name, ParentMeshComponent, ParentBoneName);
+		SetControlParent(Name, ParentComponent, ParentBoneName);
 	}
 }
 
 //======================================================================================================================
 void UPhysicsControlComponent::SetControlParentsInSet(
-	const FName     SetName,
-	UMeshComponent* ParentMeshComponent,
-	const FName     ParentBoneName)
+	const FName          SetName,
+	UPrimitiveComponent* ParentComponent,
+	const FName          ParentBoneName)
 {
-	SetControlParents(GetControlNamesInSet(SetName), ParentMeshComponent, ParentBoneName);
+	SetControlParents(GetControlNamesInSet(SetName), ParentComponent, ParentBoneName);
 }
 
 //======================================================================================================================
@@ -1577,7 +1576,7 @@ bool UPhysicsControlComponent::GetControlEnabled(const FName Name) const
 
 //======================================================================================================================
 FName UPhysicsControlComponent::CreateBodyModifier(
-	UMeshComponent*                   MeshComponent,
+	UPrimitiveComponent*              Component,
 	const FName                       BoneName,
 	const FName                       Set,
 	const FPhysicsControlModifierData BodyModifierData)
@@ -1585,7 +1584,7 @@ FName UPhysicsControlComponent::CreateBodyModifier(
 	TSet<FName> Keys;
 	BodyModifierRecords.GetKeys(Keys);
 	const FName Name = UE::PhysicsControl::GetUniqueBodyModifierName(BoneName, Keys, TEXT(""));
-	if (CreateNamedBodyModifier(Name, MeshComponent, BoneName, Set, BodyModifierData))
+	if (CreateNamedBodyModifier(Name, Component, BoneName, Set, BodyModifierData))
 	{
 		return Name;
 	}
@@ -1595,7 +1594,7 @@ FName UPhysicsControlComponent::CreateBodyModifier(
 //======================================================================================================================
 bool UPhysicsControlComponent::CreateNamedBodyModifier(
 	const FName                       Name,
-	UMeshComponent*                   MeshComponent,
+	UPrimitiveComponent*              Component,
 	const FName                       BoneName,
 	const FName                       Set,
 	const FPhysicsControlModifierData BodyModifierData)
@@ -1607,7 +1606,7 @@ bool UPhysicsControlComponent::CreateNamedBodyModifier(
 		return false;
 	}
 
-	if (!MeshComponent)
+	if (!Component)
 	{
 		UE_LOG(LogPhysicsControl, Warning,
 			TEXT("Unable to make a PhysicsBodyModifier as the mesh component has not been set"));
@@ -1615,9 +1614,9 @@ bool UPhysicsControlComponent::CreateNamedBodyModifier(
 	}
 
 	FPhysicsBodyModifierRecord& Modifier = BodyModifierRecords.Add(
-		Name, FPhysicsBodyModifierRecord(MeshComponent, BoneName, BodyModifierData));
+		Name, FPhysicsBodyModifierRecord(Component, BoneName, BodyModifierData));
 
-	USkeletalMeshComponent* SkeletalMeshComponent = Cast<USkeletalMeshComponent>(MeshComponent);
+	USkeletalMeshComponent* SkeletalMeshComponent = Cast<USkeletalMeshComponent>(Component);
 	if (SkeletalMeshComponent)
 	{
 		AddSkeletalMeshReferenceForCaching(SkeletalMeshComponent);
@@ -2081,7 +2080,7 @@ void UPhysicsControlComponent::CreateControlsAndBodyModifiersFromLimbBones(
 	const FPhysicsControlData                   WorldSpaceControlData,
 	const FPhysicsControlData                   ParentSpaceControlData,
 	const FPhysicsControlModifierData           BodyModifierData,
-	UMeshComponent*                             WorldComponent,
+	UPrimitiveComponent*                        WorldComponent,
 	FName                                       WorldBoneName)
 {
 	TMap<FName, FPhysicsControlLimbBones> LimbBones = 
@@ -2101,7 +2100,7 @@ void UPhysicsControlComponent::CreateControlsAndBodyModifiersFromLimbBones(
 //======================================================================================================================
 void UPhysicsControlComponent::CreateControlsAndBodyModifiersFromControlProfileAsset(
 	USkeletalMeshComponent* SkeletalMeshComponent,
-	UMeshComponent*         WorldComponent,
+	UPrimitiveComponent*    WorldComponent,
 	FName                   WorldBoneName)
 {
 	if (!PhysicsControlProfileAsset.IsValid())
@@ -2521,6 +2520,17 @@ bool UPhysicsControlComponent::SetCachedBoneData(
 	return false;
 }
 
+//======================================================================================================================
+void UPhysicsControlComponent::SetCachedBoneVelocitiesToZero()
+{
+	for (TPair<TWeakObjectPtr<USkeletalMeshComponent>, UE::PhysicsControl::FPhysicsControlPoseData>&
+		CachedSkeletalMeshDataPair : CachedPoseDatas)
+	{
+		UE::PhysicsControl::FPhysicsControlPoseData& CachedSkeletalMeshData = CachedSkeletalMeshDataPair.Value;
+		CachedSkeletalMeshData.BoneDatas.Reset(); // Doesn't change memory allocations
+	}
+}
+
 //======================================================================================================================
 bool UPhysicsControlComponent::ResetBodyModifierToCachedBoneTransform(
 	const FName                        Name,
@@ -2620,12 +2630,12 @@ void UPhysicsControlComponent::DebugDraw(FPrimitiveDrawInterface* PDI) const
 			{
 				const FPhysicsControlRecord& Record = PhysicsControlRecordPair.Value;
 
-				FString ParentComponentName = Record.ParentMeshComponent.IsValid() ?
-					Record.ParentMeshComponent->GetName() : TEXT("NoParent");
-				FString ChildComponentName = Record.ChildMeshComponent.IsValid() ?
-					Record.ChildMeshComponent->GetName() : TEXT("NoChild");
+				const FString ParentComponentName = Record.ParentComponent.IsValid() ?
+					Record.ParentComponent->GetName() : TEXT("NoParent");
+				const FString ChildComponentName = Record.ChildComponent.IsValid() ?
+					Record.ChildComponent->GetName() : TEXT("NoChild");
 
-				FString Text = FString::Printf(
+				const FString Text = FString::Printf(
 					TEXT("%s: Parent %s (%s) Child %s (%s): Linear strength %f Angular strength %f"),
 					*Name.ToString(),
 					*ParentComponentName,
@@ -2676,7 +2686,7 @@ void UPhysicsControlComponent::DebugDraw(FPrimitiveDrawInterface* PDI) const
 			{
 				const FPhysicsBodyModifierRecord& Record = PhysicsBodyModifierPair.Value;
 
-				FString ComponentName = Record.MeshComponent.IsValid() ? Record.MeshComponent->GetName() : TEXT("None");
+				FString ComponentName = Record.Component.IsValid() ? Record.Component->GetName() : TEXT("None");
 
 				FString Text = FString::Printf(
 					TEXT("%s: %s: %s %s GravityMultiplier %f BlendWeight %f"),
@@ -2733,7 +2743,7 @@ void UPhysicsControlComponent::DebugDrawControl(
 	if (Record.PhysicsControl.IsEnabled() && ConstraintInstance)
 	{
 		FBodyInstance* ChildBodyInstance = UE::PhysicsControl::GetBodyInstance(
-			Record.ChildMeshComponent.Get(), Record.PhysicsControl.ChildBoneName);
+			Record.ChildComponent.Get(), Record.PhysicsControl.ChildBoneName);
 		if (!ChildBodyInstance)
 		{
 			return;
@@ -2741,7 +2751,7 @@ void UPhysicsControlComponent::DebugDrawControl(
 		FTransform ChildBodyTM = ChildBodyInstance->GetUnrealWorldTransform();
 
 		FBodyInstance* ParentBodyInstance = UE::PhysicsControl::GetBodyInstance(
-			Record.ParentMeshComponent.Get(), Record.PhysicsControl.ParentBoneName);
+			Record.ParentComponent.Get(), Record.PhysicsControl.ParentBoneName);
 		const FTransform ParentBodyTM = ParentBodyInstance ? ParentBodyInstance->GetUnrealWorldTransform() : FTransform();
 
 		FTransform TargetTM, SkeletalTargetTM;
diff --git a/Engine/Plugins/Experimental/PhysicsControl/Source/PhysicsControl/Private/PhysicsControlComponentHelpers.cpp b/Engine/Plugins/Experimental/PhysicsControl/Source/PhysicsControl/Private/PhysicsControlComponentHelpers.cpp
--- a/Engine/Plugins/Experimental/PhysicsControl/Source/PhysicsControl/Private/PhysicsControlComponentHelpers.cpp
+++ b/Engine/Plugins/Experimental/PhysicsControl/Source/PhysicsControl/Private/PhysicsControlComponentHelpers.cpp
@@ -4,7 +4,7 @@
 #include "PhysicsControlData.h"
 #include "PhysicsEngine/ConstraintInstance.h"
 
-#include "Components/StaticMeshComponent.h"
+#include "Components/PrimitiveComponent.h"
 #include "Components/SkeletalMeshComponent.h"
 
 namespace UE
@@ -109,15 +109,15 @@ void ConvertConstraintProfileToControlData(
 }
 
 //======================================================================================================================
-FBodyInstance* GetBodyInstance(UMeshComponent* MeshComponent, const FName BoneName)
+FBodyInstance* GetBodyInstance(UPrimitiveComponent* Component, const FName BoneName)
 {
-	if (UStaticMeshComponent* StaticMeshComponent = Cast<UStaticMeshComponent>(MeshComponent))
+	if (USkeletalMeshComponent* SkeletalMeshComponent = Cast<USkeletalMeshComponent>(Component))
 	{
-		return StaticMeshComponent->GetBodyInstance();
+		return SkeletalMeshComponent->GetBodyInstance(BoneName);
 	}
-	if (USkeletalMeshComponent* SkeletalMeshComponent = Cast<USkeletalMeshComponent>(MeshComponent))
+	if (Component)
 	{
-		return SkeletalMeshComponent->GetBodyInstance(BoneName);
+		return Component->GetBodyInstance();
 	}
 	return nullptr;
 }
diff --git a/Engine/Plugins/Experimental/PhysicsControl/Source/PhysicsControl/Private/PhysicsControlComponentHelpers.h b/Engine/Plugins/Experimental/PhysicsControl/Source/PhysicsControl/Private/PhysicsControlComponentHelpers.h
--- a/Engine/Plugins/Experimental/PhysicsControl/Source/PhysicsControl/Private/PhysicsControlComponentHelpers.h
+++ b/Engine/Plugins/Experimental/PhysicsControl/Source/PhysicsControl/Private/PhysicsControlComponentHelpers.h
@@ -9,7 +9,7 @@ struct FBodyInstance;
 struct FPhysicsControlData;
 struct FConstraintProfileProperties;
 
-class UMeshComponent;
+class UPrimitiveComponent;
 class USkeletalMeshComponent;
 
 namespace UE
@@ -33,18 +33,19 @@ void ConvertConstraintProfileToControlData(
 	FPhysicsControlData& OutControlData, const FConstraintProfileProperties& InProfileProperties);
 
 /** 
- * Attempts to find a BodyInstance from the mesh. If it is a static mesh the single body instance
- * will be returned. If it is a skeletal mesh then if BoneName can be found, the body instance corresponding
- * to that bone will be returned. Otherwise it will return nullptr if the bone can't be found.
+ * Attempts to find a BodyInstance from the component. If it is a primitive or static mesh component the 
+ * single body instance will be returned. If it is a skeletal mesh then if BoneName can be found, the 
+ * body instance corresponding to that bone will be returned. Otherwise it will return nullptr if 
+ * the bone can't be found.
  */
-FBodyInstance* GetBodyInstance(UMeshComponent* MeshComponent, const FName BoneName);
+FBodyInstance* GetBodyInstance(UPrimitiveComponent* Component, const FName BoneName);
 
 /**
  * Attempts to find the parent physical bone given a skeletal mesh and starting bone. This walks up
  * the hierarchy, ignoring non-physical bones, until either a physical bone is found, or it has reached
  * the root without finding a physical bone.
  */
-FName GetPhysicalParentBone(USkeletalMeshComponent* SkeletalMeshComponent, FName BoneName);
+FName GetPhysicalParentBone(USkeletalMeshComponent* SkeletalComponent, FName BoneName);
 
 /**
  * Converts strength/damping ratio/extra damping into spring stiffness/damping.
diff --git a/Engine/Plugins/Experimental/PhysicsControl/Source/PhysicsControl/Private/PhysicsControlComponentImpl.cpp b/Engine/Plugins/Experimental/PhysicsControl/Source/PhysicsControl/Private/PhysicsControlComponentImpl.cpp
--- a/Engine/Plugins/Experimental/PhysicsControl/Source/PhysicsControl/Private/PhysicsControlComponentImpl.cpp
+++ b/Engine/Plugins/Experimental/PhysicsControl/Source/PhysicsControl/Private/PhysicsControlComponentImpl.cpp
@@ -26,7 +26,7 @@
 static USkeletalMeshComponent* GetValidSkeletalMeshComponentFromControlParent(
 	const FPhysicsControlRecord& Record)
 {
-	return Cast<USkeletalMeshComponent>(Record.ParentMeshComponent.Get());
+	return Cast<USkeletalMeshComponent>(Record.ParentComponent.Get());
 }
 
 //======================================================================================================================
@@ -34,7 +34,7 @@ static USkeletalMeshComponent* GetValidSkeletalMeshComponentFromControlParent(
 static USkeletalMeshComponent* GetValidSkeletalMeshComponentFromControlChild(
 	const FPhysicsControlRecord& Record)
 {
-	return Cast<USkeletalMeshComponent>(Record.ChildMeshComponent.Get());
+	return Cast<USkeletalMeshComponent>(Record.ChildComponent.Get());
 }
 
 //======================================================================================================================
@@ -42,7 +42,7 @@ static USkeletalMeshComponent* GetValidSkeletalMeshComponentFromControlChild(
 static USkeletalMeshComponent* GetValidSkeletalMeshComponentFromBodyModifier(
 	const FPhysicsBodyModifierRecord& PhysicsBodyModifier)
 {
-	return Cast<USkeletalMeshComponent>(PhysicsBodyModifier.MeshComponent.Get());
+	return Cast<USkeletalMeshComponent>(PhysicsBodyModifier.Component.Get());
 }
 
 //======================================================================================================================
@@ -206,7 +206,7 @@ void UPhysicsControlComponent::ApplyKinematicTarget(const FPhysicsBodyModifierRe
 	if (USkeletalMeshComponent* SkeletalMeshComponent = GetValidSkeletalMeshComponentFromBodyModifier(Record))
 	{
 		FBodyInstance* BodyInstance = UE::PhysicsControl::GetBodyInstance(
-			Record.MeshComponent.Get(), Record.BodyModifier.BoneName);
+			Record.Component.Get(), Record.BodyModifier.BoneName);
 		if (!BodyInstance)
 		{
 			return;
@@ -233,14 +233,14 @@ void UPhysicsControlComponent::ApplyKinematicTarget(const FPhysicsBodyModifierRe
 	}
 	else
 	{
-		const FTransform TM = Record.MeshComponent->GetComponentToWorld();
+		const FTransform TM = Record.Component->GetComponentToWorld();
 		const ETeleportType TT = UE::PhysicsControl::DetectTeleport(
 			TM.GetTranslation(), TM.GetRotation(), 
 			Record.KinematicTargetPosition, Record.KinematicTargetOrientation,
 			TeleportDistanceThreshold, TeleportRotationThreshold)
 			? ETeleportType::ResetPhysics : ETeleportType::None;
 		// Note that calling BodyInstance->SetBodyTransform moves the physics, but not the mesh
-		Record.MeshComponent->SetWorldLocationAndRotation(
+		Record.Component->SetWorldLocationAndRotation(
 			Record.KinematicTargetPosition, Record.KinematicTargetOrientation, false, nullptr, TT);
 	}
 }
@@ -249,7 +249,7 @@ void UPhysicsControlComponent::ApplyKinematicTarget(const FPhysicsBodyModifierRe
 void UPhysicsControlComponent::ResetToCachedTarget(const FPhysicsBodyModifierRecord& Record) const
 {
 	FBodyInstance* BodyInstance = UE::PhysicsControl::GetBodyInstance(
-		Record.MeshComponent.Get(), Record.BodyModifier.BoneName);
+		Record.Component.Get(), Record.BodyModifier.BoneName);
 	if (!BodyInstance)
 	{
 		return;
@@ -428,9 +428,9 @@ void UPhysicsControlComponent::CalculateControlTargetData(
 			bHaveParentBoneData = GetBoneData(
 				ParentBoneData, ParentPoseData, ParentSkeletalMeshComponent, Record.PhysicsControl.ParentBoneName);
 		}
-		else if (Record.ParentMeshComponent.IsValid())
+		else if (Record.ParentComponent.IsValid())
 		{
-			const FTransform ParentTM = Record.ParentMeshComponent->GetComponentTransform();
+			const FTransform ParentTM = Record.ParentComponent->GetComponentTransform();
 			ParentBoneData.CurrentTM = ParentTM;
 			bHaveParentBoneData = true;
 		}
@@ -588,10 +588,10 @@ void UPhysicsControlComponent::ApplyControl(FPhysicsControlRecord& Record)
 	ConstraintInstance->SetDisableCollision(Record.PhysicsControl.ControlData.bDisableCollision);
 
 	FBodyInstance* ParentBodyInstance = UE::PhysicsControl::GetBodyInstance(
-		Record.ParentMeshComponent.Get(), Record.PhysicsControl.ParentBoneName);
+		Record.ParentComponent.Get(), Record.PhysicsControl.ParentBoneName);
 
 	FBodyInstance* ChildBodyInstance = UE::PhysicsControl::GetBodyInstance(
-		Record.ChildMeshComponent.Get(), Record.PhysicsControl.ChildBoneName);
+		Record.ChildComponent.Get(), Record.PhysicsControl.ChildBoneName);
 
 	if (!ParentBodyInstance && !ChildBodyInstance)
 	{
@@ -632,7 +632,7 @@ void UPhysicsControlComponent::ApplyBodyModifier(FPhysicsBodyModifierRecord& Rec
 {
 	USkeletalMeshComponent* SKM = GetValidSkeletalMeshComponentFromBodyModifier(Record);
 	FBodyInstance* BodyInstance = UE::PhysicsControl::GetBodyInstance(
-		Record.MeshComponent.Get(), Record.BodyModifier.BoneName);
+		Record.Component.Get(), Record.BodyModifier.BoneName);
 	if (BodyInstance)
 	{
 		switch (Record.BodyModifier.ModifierData.MovementType)
diff --git a/Engine/Plugins/Experimental/PhysicsControl/Source/PhysicsControl/Private/PhysicsControlRecord.cpp b/Engine/Plugins/Experimental/PhysicsControl/Source/PhysicsControl/Private/PhysicsControlRecord.cpp
--- a/Engine/Plugins/Experimental/PhysicsControl/Source/PhysicsControl/Private/PhysicsControlRecord.cpp
+++ b/Engine/Plugins/Experimental/PhysicsControl/Source/PhysicsControl/Private/PhysicsControlRecord.cpp
@@ -25,8 +25,8 @@ FVector FPhysicsControlRecord::GetControlPoint() const
 		return PhysicsControl.ControlData.CustomControlPoint;
 	}
 
-	FBodyInstance* ChildBodyInstance = UE::PhysicsControl::GetBodyInstance(
-		ChildMeshComponent.Get(), PhysicsControl.ChildBoneName);
+	const FBodyInstance* ChildBodyInstance = UE::PhysicsControl::GetBodyInstance(
+		ChildComponent.Get(), PhysicsControl.ChildBoneName);
 
 	return ChildBodyInstance ? ChildBodyInstance->GetMassSpaceLocal().GetTranslation() : FVector::ZeroVector;
 }
@@ -41,11 +41,11 @@ bool FPhysicsControlRecord::InitConstraint(UObject* ConstraintDebugOwner, FName
 	check(ConstraintInstance);
 
 	FBodyInstance* ParentBody = UE::PhysicsControl::GetBodyInstance(
-		ParentMeshComponent.Get(), PhysicsControl.ParentBoneName);
+		ParentComponent.Get(), PhysicsControl.ParentBoneName);
 	FBodyInstance* ChildBody = UE::PhysicsControl::GetBodyInstance(
-		ChildMeshComponent.Get(), PhysicsControl.ChildBoneName);
+		ChildComponent.Get(), PhysicsControl.ChildBoneName);
 
-	if (ParentMeshComponent.IsValid() && !PhysicsControl.ParentBoneName.IsNone() && !ParentBody)
+	if (ParentComponent.IsValid() && !PhysicsControl.ParentBoneName.IsNone() && !ParentBody)
 	{
 		UE_LOG(LogPhysicsControl, Warning,
 			TEXT("Failed to find expected parent body %s when making constraint for control %s"),
@@ -53,7 +53,7 @@ bool FPhysicsControlRecord::InitConstraint(UObject* ConstraintDebugOwner, FName
 			*ControlName.ToString());
 		return false;
 	}
-	if (ChildMeshComponent.IsValid() && !PhysicsControl.ChildBoneName.IsNone() && !ChildBody)
+	if (ChildComponent.IsValid() && !PhysicsControl.ChildBoneName.IsNone() && !ChildBody)
 	{
 		UE_LOG(LogPhysicsControl, Warning,
 			TEXT("Failed to find expected child body %s when making constraint for control %s"),
diff --git a/Engine/Plugins/Experimental/PhysicsControl/Source/PhysicsControl/Public/PhysicsControlComponent.h b/Engine/Plugins/Experimental/PhysicsControl/Source/PhysicsControl/Public/PhysicsControlComponent.h
--- a/Engine/Plugins/Experimental/PhysicsControl/Source/PhysicsControl/Public/PhysicsControlComponent.h
+++ b/Engine/Plugins/Experimental/PhysicsControl/Source/PhysicsControl/Public/PhysicsControlComponent.h
@@ -101,9 +101,9 @@ public:
 	 */
 	UFUNCTION(BlueprintCallable, Category = PhysicsControl)
 	FName CreateControl(
-		UMeshComponent*               ParentMeshComponent,
+		UPrimitiveComponent*          ParentComponent,
 		FName                         ParentBoneName,
-		UMeshComponent*               ChildMeshComponent,
+		UPrimitiveComponent*          ChildComponent,
 		const FName                   ChildBoneName,
 		const FPhysicsControlData     ControlData,
 		const FPhysicsControlTarget   ControlTarget, 
@@ -126,9 +126,9 @@ public:
 	UFUNCTION(BlueprintCallable, Category = PhysicsControl)
 	bool CreateNamedControl(
 		FName                         Name,
-		UMeshComponent*               ParentMeshComponent,
+		UPrimitiveComponent*          ParentComponent,
 		const FName                   ParentBoneName,
-		UMeshComponent*               ChildMeshComponent,
+		UPrimitiveComponent*          ChildComponent,
 		const FName                   ChildBoneName,
 		const FPhysicsControlData     ControlData, 
 		const FPhysicsControlTarget   ControlTarget, 
@@ -277,7 +277,7 @@ public:
 		const TMap<FName, FPhysicsControlLimbBones>& LimbBones,
 		const EPhysicsControlType                    ControlType,
 		const FPhysicsControlData                    ControlData,
-		UMeshComponent*                              WorldComponent = nullptr,
+		UPrimitiveComponent*                         WorldComponent = nullptr,
 		FName                                        WorldBoneName = NAME_None,
 		FString                                      NamePrefix = TEXT(""));
 
@@ -342,9 +342,9 @@ public:
 	 */
 	UFUNCTION(BlueprintCallable, Category = PhysicsControl)
 	bool SetControlParent(
-		const FName     Name, 
-		UMeshComponent* ParentMeshComponent,
-		const FName     ParentBoneName);
+		const FName          Name, 
+		UPrimitiveComponent* ParentComponent,
+		const FName          ParentBoneName);
 
 	/**
 	 * Updates the parent object part of controls. Note that this won't change the name of the controls (which may
@@ -353,14 +353,14 @@ public:
 	UFUNCTION(BlueprintCallable, Category = PhysicsControl)
 	void SetControlParents(
 		const TArray<FName>& Names,
-		UMeshComponent*      ParentMeshComponent,
+		UPrimitiveComponent* ParentComponent,
 		const FName          ParentBoneName);
 
 	UFUNCTION(BlueprintCallable, Category = PhysicsControl)
 	void SetControlParentsInSet(
-		const FName     Set,
-		UMeshComponent* ParentMeshComponent,
-		const FName     ParentBoneName);
+		const FName          Set,
+		UPrimitiveComponent* ParentComponent,
+		const FName          ParentBoneName);
 
 	/**
 	 * Modifies an existing control data - i.e. the strengths etc of the control driving towards the target
@@ -943,7 +943,7 @@ public:
 	/**
 	 * Creates a new body modifier for mesh components
 	 * 
-	 * @param MeshComponent The Mesh Component used as a target for the modifier
+	 * @param Component The Mesh Component used as a target for the modifier
 	 * @param BoneName The bone name, if a skeletal mesh is used
 	 * @param Set Which set to include the body modifier in (optional). Note that it automatically 
 	 *        gets added to the set "All"
@@ -951,7 +951,7 @@ public:
 	 */
 	UFUNCTION(BlueprintCallable, Category = PhysicsControl)
 	FName CreateBodyModifier(
-		UMeshComponent*                   MeshComponent,
+		UPrimitiveComponent*              Component,
 		const FName                       BoneName,
 		const FName                       Set,
 		const FPhysicsControlModifierData BodyModifierData);
@@ -960,7 +960,7 @@ public:
 	 * Creates a new body modifier for mesh components
 	 * 
 	 * @param The name of the body modifier that will be created. Creation will fail if this name is already in use.
-	 * @param MeshComponent The Mesh Component used as a target for the modifier
+	 * @param Component The Mesh Component used as a target for the modifier
 	 * @param BoneName The bone name, if a skeletal mesh is used
 	 * @param Set Which set to include the body modifier in (optional). Note that it automatically
 	 *        gets added to the set "All"
@@ -969,7 +969,7 @@ public:
 	UFUNCTION(BlueprintCallable, Category = PhysicsControl)
 	bool CreateNamedBodyModifier(
 		const FName                       Name,
-		UMeshComponent*                   MeshComponent,
+		UPrimitiveComponent*              Component,
 		const FName                       BoneName,
 		const FName                       Set,
 		const FPhysicsControlModifierData BodyModifierData);
@@ -1375,7 +1375,7 @@ public:
 		const FPhysicsControlData                   WorldSpaceControlData,
 		const FPhysicsControlData                   ParentSpaceControlData,
 		const FPhysicsControlModifierData           BodyModifierData,
-		UMeshComponent*                             WorldComponent = nullptr,
+		UPrimitiveComponent*                        WorldComponent = nullptr,
 		FName                                       WorldBoneName = NAME_None);
 
 	/**
@@ -1387,7 +1387,7 @@ public:
 	UFUNCTION(BlueprintCallable, Category = PhysicsControl)
 	void CreateControlsAndBodyModifiersFromControlProfileAsset(
 		USkeletalMeshComponent* SkeletalMeshComponent,
-		UMeshComponent*         WorldComponent,
+		UPrimitiveComponent*    WorldComponent,
 		FName                   WorldBoneName);
 
 	/**
@@ -1558,6 +1558,13 @@ public:
 		const FName                   BoneName, 
 		const FTransform&             TM);
 
+	/**
+	 * This sets cached bone velocities to zero, by forgetting any previously used transform. This may be useful if 
+	 * the pose is changed suddenly and you don't want this change to show up in the simulated velocities.
+	 */
+	UFUNCTION(BlueprintCallable, Category = PhysicsControl)
+	void SetCachedBoneVelocitiesToZero();
+
 	/**
 	 * This flags the body associated with the modifier to set (using teleport) its position and velocity to 
 	 * the cached animation target. This will only affect skeletal mesh component bodies. 
diff --git a/Engine/Plugins/Experimental/PhysicsControl/Source/PhysicsControl/Public/PhysicsControlRecord.h b/Engine/Plugins/Experimental/PhysicsControl/Source/PhysicsControl/Public/PhysicsControlRecord.h
--- a/Engine/Plugins/Experimental/PhysicsControl/Source/PhysicsControl/Public/PhysicsControlRecord.h
+++ b/Engine/Plugins/Experimental/PhysicsControl/Source/PhysicsControl/Public/PhysicsControlRecord.h
@@ -9,7 +9,7 @@
 
 struct FConstraintInstance;
 struct FBodyInstance;
-class UMeshComponent;
+class UPrimitiveComponent;
 
 /**
  * There will be a PhysicsControlRecord created at runtime for every Control that has been created
@@ -19,13 +19,13 @@ struct FPhysicsControlRecord
 	FPhysicsControlRecord(
 		const FPhysicsControl&       InControl,
 		const FPhysicsControlTarget& InControlTarget,
-		UMeshComponent*              InParentMeshComponent,
-		UMeshComponent*              InChildMeshComponent)
+		UPrimitiveComponent*              InParentComponent,
+		UPrimitiveComponent*              InChildComponent)
 		: PhysicsControl(InControl)
 		, ControlTarget(InControlTarget)
 		, ExpectedUpdateCounter(-1)
-		, ParentMeshComponent(InParentMeshComponent)
-		, ChildMeshComponent(InChildMeshComponent)
+		, ParentComponent(InParentComponent)
+		, ChildComponent(InChildComponent)
 	{}
 
 	/** Removes any constraint and resets the state */
@@ -66,10 +66,10 @@ struct FPhysicsControlRecord
 	int64 ExpectedUpdateCounter;
 
 	/**  The mesh that will be doing the driving. Blank/non-existent means it will happen in world space */
-	TWeakObjectPtr<UMeshComponent> ParentMeshComponent;
+	TWeakObjectPtr<UPrimitiveComponent> ParentComponent;
 
 	/** The mesh that the control will be driving. */
-	TWeakObjectPtr<UMeshComponent> ChildMeshComponent;
+	TWeakObjectPtr<UPrimitiveComponent> ChildComponent;
 
 	/** The underlying constraint used to implement the control. */
 	TSharedPtr<FConstraintInstance> ConstraintInstance;
@@ -81,10 +81,10 @@ struct FPhysicsControlRecord
 struct FPhysicsBodyModifierRecord
 {
 	FPhysicsBodyModifierRecord(
-		TWeakObjectPtr<UMeshComponent>  InMeshComponent, 
+		TWeakObjectPtr<UPrimitiveComponent>  InComponent, 
 		const FName&                    InBoneName, 
 		FPhysicsControlModifierData     InBodyModifierData)
-		: MeshComponent(InMeshComponent)
+		: Component(InComponent)
 		, BodyModifier(InBoneName, InBodyModifierData)
 		, KinematicTargetPosition(FVector::ZeroVector)
 		, KinematicTargetOrientation(FQuat::Identity)
@@ -92,7 +92,7 @@ struct FPhysicsBodyModifierRecord
 	{}
 
 	/**  The mesh that will be modified. */
-	TWeakObjectPtr<UMeshComponent> MeshComponent;
+	TWeakObjectPtr<UPrimitiveComponent> Component;
 
 	// The core data
 	FPhysicsBodyModifier BodyModifier;
