General tidyup of PhysicsControl. This also makes FPosQuat operate from left to right (to match FQuat rather then FTransform)

diff --git a/Engine/Plugins/Experimental/PhysicsControl/Source/PhysicsControl/Private/AnimNode_RigidBodyWithControl.cpp b/Engine/Plugins/Experimental/PhysicsControl/Source/PhysicsControl/Private/AnimNode_RigidBodyWithControl.cpp
--- a/Engine/Plugins/Experimental/PhysicsControl/Source/PhysicsControl/Private/AnimNode_RigidBodyWithControl.cpp
+++ b/Engine/Plugins/Experimental/PhysicsControl/Source/PhysicsControl/Private/AnimNode_RigidBodyWithControl.cpp
@@ -794,7 +794,7 @@ void FAnimNode_RigidBodyWithControl::EvaluateSkeletalControl_AnyThread(FComponen
 				BodyAnimData[BodyIndex].bIsSimulated = !Bodies[BodyIndex]->GetIsKinematic();
 				if (!BodyAnimData[BodyIndex].bIsSimulated)
 				{
-					// TODO support explicit targets for kinematics, like we do for controls
+					// Custom kinematic targets are set subsequently, overriding these.
 					Bodies[BodyIndex]->SetKinematicTarget(PoseData.GetTM(BodyIndex).ToTransform());
 				}
 			}
diff --git a/Engine/Plugins/Experimental/PhysicsControl/Source/PhysicsControl/Private/PhysicsControlBPLibrary.cpp b/Engine/Plugins/Experimental/PhysicsControl/Source/PhysicsControl/Private/PhysicsControlBPLibrary.cpp
--- a/Engine/Plugins/Experimental/PhysicsControl/Source/PhysicsControl/Private/PhysicsControlBPLibrary.cpp
+++ b/Engine/Plugins/Experimental/PhysicsControl/Source/PhysicsControl/Private/PhysicsControlBPLibrary.cpp
@@ -75,9 +75,6 @@ template<typename TNamedParameters> void BlendParametersThroughSet(
 	const TArray<FName>&                               Names,
 	FPhysicsControlControlAndModifierParameters&       OutParametersContainer)
 {
-	// TODO - Limbs can include branches. Would need a more sophisticated approach to deal with
-	// those properly, perhaps include a depth index in the control name ?
-
 	OutParametersContainer = InParametersContainer;
 
 	const float WeightDelta = 1.0f / static_cast<float>(Names.Num() - 1);
diff --git a/Engine/Plugins/Experimental/PhysicsControl/Source/PhysicsControl/Private/PhysicsControlComponent.cpp b/Engine/Plugins/Experimental/PhysicsControl/Source/PhysicsControl/Private/PhysicsControlComponent.cpp
--- a/Engine/Plugins/Experimental/PhysicsControl/Source/PhysicsControl/Private/PhysicsControlComponent.cpp
+++ b/Engine/Plugins/Experimental/PhysicsControl/Source/PhysicsControl/Private/PhysicsControlComponent.cpp
@@ -72,17 +72,15 @@ void UPhysicsControlComponent::InitializeComponent()
 //======================================================================================================================
 void UPhysicsControlComponent::BeginDestroy()
 {
-	for (TPair<FName, FPhysicsControlRecord>& PhysicsControlRecordPair : ControlRecords)
+	for (TPair<FName, FPhysicsControlRecord>& ControlRecordPair : ControlRecords)
 	{
-		DestroyControl(
-			PhysicsControlRecordPair.Key, EDestroyBehavior::KeepRecord);
+		DestroyControl(ControlRecordPair.Key, EDestroyBehavior::KeepRecord);
 	}
 	ControlRecords.Empty();
 
-	for (TPair<FName, FPhysicsBodyModifierRecord>& PhysicsBodyModifierPair : BodyModifierRecords)
+	for (TPair<FName, FPhysicsBodyModifierRecord>& BodyModifierPair : BodyModifierRecords)
 	{
-		DestroyBodyModifier(
-			PhysicsBodyModifierPair.Key, EDestroyBehavior::KeepRecord);
+		DestroyBodyModifier(BodyModifierPair.Key, EDestroyBehavior::KeepRecord);
 	}
 	BodyModifierRecords.Empty();
 
@@ -2610,11 +2608,10 @@ void UPhysicsControlComponent::DebugDraw(FPrimitiveDrawInterface* PDI) const
 	// Draw gizmos
 	if (bShowDebugVisualization && VisualizationSizeScale > 0)
 	{
-		for (const TPair<FName, FPhysicsControlRecord>& PhysicsControlRecordPair : 
-			ControlRecords)
+		for (const TPair<FName, FPhysicsControlRecord>& ControlRecordPair : ControlRecords)
 		{
-			const FName Name = PhysicsControlRecordPair.Key;
-			const FPhysicsControlRecord& Record = PhysicsControlRecordPair.Value;
+			const FName Name = ControlRecordPair.Key;
+			const FPhysicsControlRecord& Record = ControlRecordPair.Value;
 			DebugDrawControl(PDI, Record, Name);
 		} 
 	}
@@ -2622,14 +2619,13 @@ void UPhysicsControlComponent::DebugDraw(FPrimitiveDrawInterface* PDI) const
 	// Detailed controls - if there's a filter
 	if (!DebugControlDetailFilter.IsEmpty())
 	{
-		for (const TPair<FName, FPhysicsControlRecord>& PhysicsControlRecordPair :
-			ControlRecords)
+		for (const TPair<FName, FPhysicsControlRecord>& ControlRecordPair : ControlRecords)
 		{
-			const FName Name = PhysicsControlRecordPair.Key;
+			const FName Name = ControlRecordPair.Key;
 
 			if (Name.ToString().Contains(DebugControlDetailFilter))
 			{
-				const FPhysicsControlRecord& Record = PhysicsControlRecordPair.Value;
+				const FPhysicsControlRecord& Record = ControlRecordPair.Value;
 
 				const FString ParentComponentName = Record.ParentComponent.IsValid() ?
 					Record.ParentComponent->GetName() : TEXT("NoParent");
@@ -2657,11 +2653,9 @@ void UPhysicsControlComponent::DebugDraw(FPrimitiveDrawInterface* PDI) const
 	if (bShowDebugControlList)
 	{
 		FString AllNames;
-
-		for (const TPair<FName, FPhysicsControlRecord>& PhysicsControlRecordPair :
-			ControlRecords)
+		for (const TPair<FName, FPhysicsControlRecord>& ControlRecordPair : ControlRecords)
 		{
-			const FName Name = PhysicsControlRecordPair.Key;
+			const FName Name = ControlRecordPair.Key;
 			AllNames += Name.ToString() + TEXT(" ");
 			if (AllNames.Len() > 256)
 			{
@@ -2670,22 +2664,19 @@ void UPhysicsControlComponent::DebugDraw(FPrimitiveDrawInterface* PDI) const
 			}
 		}
 		GEngine->AddOnScreenDebugMessage(
-			-1, 0.0f, FColor::White,
-			FString::Printf(TEXT("%d Controls: %s"), ControlRecords.Num(), *AllNames));
-
+			-1, 0.0f, FColor::White, FString::Printf(TEXT("%d Controls: %s"), ControlRecords.Num(), *AllNames));
 	}
 
 	// Detailed body modifiers - if there's a filter
 	if (!DebugBodyModifierDetailFilter.IsEmpty())
 	{
-		for (const TPair<FName, FPhysicsBodyModifierRecord>& PhysicsBodyModifierPair :
-			BodyModifierRecords)
+		for (const TPair<FName, FPhysicsBodyModifierRecord>& BodyModifierPair : BodyModifierRecords)
 		{
-			const FName Name = PhysicsBodyModifierPair.Key;
+			const FName Name = BodyModifierPair.Key;
 
 			if (Name.ToString().Contains(DebugBodyModifierDetailFilter))
 			{
-				const FPhysicsBodyModifierRecord& Record = PhysicsBodyModifierPair.Value;
+				const FPhysicsBodyModifierRecord& Record = BodyModifierPair.Value;
 
 				FString ComponentName = Record.Component.IsValid() ? Record.Component->GetName() : TEXT("None");
 
@@ -2710,9 +2701,9 @@ void UPhysicsControlComponent::DebugDraw(FPrimitiveDrawInterface* PDI) const
 	{
 		FString AllNames;
 
-		for (const TPair<FName, FPhysicsBodyModifierRecord>& PhysicsBodyModifierPair : BodyModifierRecords)
+		for (const TPair<FName, FPhysicsBodyModifierRecord>& BodyModifierPair : BodyModifierRecords)
 		{
-			const FName Name = PhysicsBodyModifierPair.Key;
+			const FName Name = BodyModifierPair.Key;
 			AllNames += Name.ToString() + TEXT(" ");
 			if (AllNames.Len() > 256)
 			{
diff --git a/Engine/Plugins/Experimental/PhysicsControl/Source/PhysicsControl/Private/PhysicsControlComponentImpl.cpp b/Engine/Plugins/Experimental/PhysicsControl/Source/PhysicsControl/Private/PhysicsControlComponentImpl.cpp
--- a/Engine/Plugins/Experimental/PhysicsControl/Source/PhysicsControl/Private/PhysicsControlComponentImpl.cpp
+++ b/Engine/Plugins/Experimental/PhysicsControl/Source/PhysicsControl/Private/PhysicsControlComponentImpl.cpp
@@ -159,8 +159,7 @@ const FPhysicsControl* UPhysicsControlComponent::FindControl(const FName Name) c
 //======================================================================================================================
 void UPhysicsControlComponent::UpdateCachedSkeletalBoneData(float DeltaTime)
 {
-	// Allow the target counter to wrap, after 2e19 frames
-	++CurrentUpdateCounter;
+	CurrentUpdateCounter.Increment();
 
 	for (TPair<TWeakObjectPtr<USkeletalMeshComponent>, UE::PhysicsControl::FPhysicsControlPoseData>&
 		CachedSkeletalMeshDataPair : CachedPoseDatas)
@@ -446,7 +445,7 @@ void UPhysicsControlComponent::CalculateControlTargetData(
 			if (bHaveParentBoneData)
 			{
 				const UE::PhysicsControl::FPosQuat ParentBoneTM = ParentBoneData.CurrentTM;
-				const UE::PhysicsControl::FPosQuat SkeletalDeltaTM = ChildBoneTM * ParentBoneTM.Inverse();
+				const UE::PhysicsControl::FPosQuat SkeletalDeltaTM = ParentBoneTM.Inverse() * ChildBoneTM;
 				// This puts TargetTM in the space of the ParentBone
 				OutSkeletalTargetTM = SkeletalDeltaTM.ToTransform();
 			}
@@ -604,11 +603,14 @@ void UPhysicsControlComponent::ApplyControl(FPhysicsControlRecord& Record)
 		FTransform TargetTM, SkeletalTargetTM;
 		FVector TargetVelocity;
 		FVector TargetAngularVelocity;
-		bool bUsePreviousSkeletalTargetTM = CurrentUpdateCounter == Record.ExpectedUpdateCounter;
+		bool bUsePreviousSkeletalTargetTM = 
+			CurrentUpdateCounter.HasEverBeenUpdated() && 
+			CurrentUpdateCounter.Get() == Record.ExpectedUpdateCounter.Get();
 		CalculateControlTargetData(
 			TargetTM, SkeletalTargetTM, TargetVelocity, TargetAngularVelocity, Record, bUsePreviousSkeletalTargetTM);
 		Record.PreviousSkeletalTargetTM = SkeletalTargetTM;
-		Record.ExpectedUpdateCounter = CurrentUpdateCounter + 1;
+		Record.ExpectedUpdateCounter = CurrentUpdateCounter;
+		Record.ExpectedUpdateCounter.Increment();
 
 		ConstraintInstance->SetLinearPositionTarget(TargetTM.GetTranslation());
 		ConstraintInstance->SetAngularOrientationTarget(TargetTM.GetRotation());
diff --git a/Engine/Plugins/Experimental/PhysicsControl/Source/PhysicsControl/Private/PhysicsControlOperatorNameGeneration.cpp b/Engine/Plugins/Experimental/PhysicsControl/Source/PhysicsControl/Private/PhysicsControlOperatorNameGeneration.cpp
--- a/Engine/Plugins/Experimental/PhysicsControl/Source/PhysicsControl/Private/PhysicsControlOperatorNameGeneration.cpp
+++ b/Engine/Plugins/Experimental/PhysicsControl/Source/PhysicsControl/Private/PhysicsControlOperatorNameGeneration.cpp
@@ -260,7 +260,7 @@ TMap<FName, FPhysicsControlLimbBones> GetLimbBones(
 	const FReferenceSkeleton& RefSkeleton,
 	UPhysicsAsset* const PhysicsAsset)
 {
-	// TODO - Output limb bones are not in the order specified in the skeleton - would be better if they were
+	// Note - Output limb bones are not in the order specified in the skeleton - would be better if they were
 
 	TMap<FName, FPhysicsControlLimbBones> Result;
 
diff --git a/Engine/Plugins/Experimental/PhysicsControl/Source/PhysicsControl/Private/RigidBodyWithControl.cpp b/Engine/Plugins/Experimental/PhysicsControl/Source/PhysicsControl/Private/RigidBodyWithControl.cpp
--- a/Engine/Plugins/Experimental/PhysicsControl/Source/PhysicsControl/Private/RigidBodyWithControl.cpp
+++ b/Engine/Plugins/Experimental/PhysicsControl/Source/PhysicsControl/Private/RigidBodyWithControl.cpp
@@ -196,8 +196,12 @@ ImmediatePhysics::FJointHandle* FAnimNode_RigidBodyWithControl::CreateConstraint
 		Settings.bAngularSwingVelocityDriveEnabled = false;
 		Settings.AngularDriveForceMode = Chaos::EJointForceMode::Acceleration;
 
-		Settings.bMassConditioningEnabled = false; // TODO needed/wanted?
-		Settings.bCollisionEnabled = false; // TODO needed?
+		// For control, we shouldn't be in situations where mass conditioning is needed.
+		Settings.bMassConditioningEnabled = false;
+		// It's not our job to change collision settings - that should come from the physics asset.
+		// However, the naming of this is unclear - if collisions are disabled in the physics asset,
+		// trust that this doesn't enable them.
+		Settings.bCollisionEnabled = true;
 
 		FVector ChildCoMPositionOffset = ChildActorHandle->GetLocalCoMTransform().GetLocation();
 		Settings.ConnectorTransforms[ConstraintChildIndex].SetLocation(ChildCoMPositionOffset);
@@ -366,16 +370,16 @@ void FAnimNode_RigidBodyWithControl::InitControlsAndBodyModifiers(const FReferen
 	TMap<FName, FPhysicsControlLimbBones> AllLimbBones =
 		UE::PhysicsControl::GetLimbBones(SetupData.LimbSetupData, RefSkeleton, GetPhysicsAsset());
 
-	FPhysicsControlAndBodyModifierCreationDatas ControlAndBodyModifierCreationDatas;
+	FPhysicsControlAndBodyModifierCreationDatas AdditionalControlAndBodyModifierCreationDatas;
 	if (IsValid(PhysicsControlProfileAsset))
 	{
-		ControlAndBodyModifierCreationDatas = PhysicsControlProfileAsset->AdditionalControlsAndModifiers;
+		AdditionalControlAndBodyModifierCreationDatas = PhysicsControlProfileAsset->AdditionalControlsAndModifiers;
 	}
-	ControlAndBodyModifierCreationDatas += AdditionalControlsAndBodyModifiers;
+	AdditionalControlAndBodyModifierCreationDatas += AdditionalControlsAndBodyModifiers;
 
 	// An "operator" is a control or a body modifier. This will also add them to sets etc.
 	UE::PhysicsControl::CreateOperatorsForNode(
-		this, SetupData, ControlAndBodyModifierCreationDatas, 
+		this, SetupData, AdditionalControlAndBodyModifierCreationDatas,
 		AllLimbBones, RefSkeleton, GetPhysicsAsset(), NameRecords);
 
 	for (TMap<FName, FRigidBodyControlRecord>::ElementType& NameRecordPair : ControlRecords)
@@ -565,14 +569,14 @@ static UE::PhysicsControl::FPosQuat CalculateTargetTM(
 	const int32                                   ChildBodyIndex)
 {
 	const UE::PhysicsControl::FPosQuat ChildTargetTM =
-		UE::PhysicsControl::FPosQuat(JointSettings.ConnectorTransforms[ConstraintChildIndex]) * 
-		PoseData.GetTM(ChildBodyIndex);
+		PoseData.GetTM(ChildBodyIndex) * 
+		UE::PhysicsControl::FPosQuat(JointSettings.ConnectorTransforms[ConstraintChildIndex]);
 	if (ParentBodyIndex >= 0)
 	{
 		const UE::PhysicsControl::FPosQuat ParentTargetTM =
-			UE::PhysicsControl::FPosQuat(JointSettings.ConnectorTransforms[ConstraintParentIndex]) * 
-			PoseData.GetTM(ParentBodyIndex);
-		return ChildTargetTM * ParentTargetTM.Inverse();
+			PoseData.GetTM(ParentBodyIndex) * 
+			UE::PhysicsControl::FPosQuat(JointSettings.ConnectorTransforms[ConstraintParentIndex]);
+		return ParentTargetTM.Inverse() * ChildTargetTM;
 	}
 	return ChildTargetTM;
 }
@@ -588,7 +592,8 @@ void FAnimNode_RigidBodyWithControl::ApplyControl(FRigidBodyControlRecord& Contr
 		Chaos::FPBDJointConstraintHandle* Constraint = JointHandle->GetConstraint();
 		if (Constraint)
 		{
-			if (ControlRecord.ExpectedUpdateCounter.Get() != PoseData.UpdateCounter.Get())
+			if (!PoseData.UpdateCounter.HasEverBeenUpdated() || 
+				ControlRecord.ExpectedUpdateCounter.Get() != PoseData.UpdateCounter.Get())
 			{
 				// If we missed some intermediate updates, then we don't want to use the previous
 				// positions etc to calculate velocities. This will mean velocity/damping will be
@@ -628,7 +633,7 @@ void FAnimNode_RigidBodyWithControl::ApplyControl(FRigidBodyControlRecord& Contr
 					{
 						UE::PhysicsControl::FPosQuat AnimTargetTM = CalculateTargetTM(
 							JointSettings, PoseData, ControlRecord.ParentBodyIndex, ControlRecord.ChildBodyIndex);
-						TargetTM = TargetTM * AnimTargetTM;
+						TargetTM = AnimTargetTM * TargetTM;
 					}
 
 					Constraint->SetLinearDrivePositionTarget(TargetTM.GetTranslation());
@@ -824,7 +829,7 @@ void FAnimNode_RigidBodyWithControl::ApplyControlsAndModifiers(const FVector& Si
 	// If we've skipped a frame then we need to avoid doing any velocity calculations. Simplest
 	// method is to set DeltaTime to zero.
 	{
-		if (PoseData.UpdateCounter.Get() == INDEX_NONE || 
+		if (!PoseData.UpdateCounter.HasEverBeenUpdated() || 
 			PoseData.UpdateCounter.Get() != PoseData.ExpectedUpdateCounter.Get())
 		{
 			DeltaTime = 0.0f;
@@ -883,7 +888,7 @@ void FAnimNode_RigidBodyWithControl::ApplyKinematicTargets()
 					UE::PhysicsControl::FPosQuat TM(Target.TargetOrientation, Target.TargetPosition);
 					if (Target.bUseSkeletalAnimation)
 					{
-						TM = TM * PoseData.GetTM(BodyIndex);
+						TM = PoseData.GetTM(BodyIndex) * TM;
 					}
 					ActorHandle->SetKinematicTarget(TM.ToTransform());
 				}
diff --git a/Engine/Plugins/Experimental/PhysicsControl/Source/PhysicsControl/Public/PhysicsControlBPLibrary.h b/Engine/Plugins/Experimental/PhysicsControl/Source/PhysicsControl/Public/PhysicsControlBPLibrary.h
--- a/Engine/Plugins/Experimental/PhysicsControl/Source/PhysicsControl/Public/PhysicsControlBPLibrary.h
+++ b/Engine/Plugins/Experimental/PhysicsControl/Source/PhysicsControl/Public/PhysicsControlBPLibrary.h
@@ -27,7 +27,7 @@ class PHYSICSCONTROL_API UPhysicsControlBPLibrary : public UBlueprintFunctionLib
 {
 	GENERATED_BODY()
 
-	/** Add a single control parameter Parameters to a set of parameter Parameters. */
+	// Add a single control parameter to a container of parameters.
 	UFUNCTION(BlueprintPure, Category = "Animation|PhysicsControl", meta = (BlueprintThreadSafe, AutoCreateRefTerm = "LinearStrengthMultiplier, LinearExtraDampingMultiplier, AngularStrengthMultiplier, AngularExtraDampingMultiplier, bEnabled"))
 	static void AddControlParameters(
 		UPARAM(ref, DisplayName = "Parameters") FPhysicsControlControlAndModifierParameters& InParameters, 
@@ -35,7 +35,7 @@ class PHYSICSCONTROL_API UPhysicsControlBPLibrary : public UBlueprintFunctionLib
 		const FName                                                                          Name,
 		const FPhysicsControlSparseData&                                                     ControlData);
 
-	/** Add a single control parameter Parameters to a set of parameter Parameters. */
+	// Add an array of control parameters (all with the same data) to a container of parameters.
 	UFUNCTION(BlueprintPure, Category = "Animation|PhysicsControl", meta = (BlueprintThreadSafe, AutoCreateRefTerm = "LinearStrengthMultiplier, LinearExtraDampingMultiplier, AngularStrengthMultiplier, AngularExtraDampingMultiplier, bEnabled"))
 	static void AddMultipleControlParameters(
 		UPARAM(ref, DisplayName = "Parameters") FPhysicsControlControlAndModifierParameters& InParameters, 
@@ -43,7 +43,7 @@ class PHYSICSCONTROL_API UPhysicsControlBPLibrary : public UBlueprintFunctionLib
 		const TArray<FName>&                                                                 Names,
 		const FPhysicsControlSparseData&                                                     ControlData);
 
-	/** Add a single body modifier parameter Parameters to a set of parameter Parameters. */
+	// Add a single body modifier parameter to a container of parameters.
 	UFUNCTION(BlueprintPure, Category = "Animation|PhysicsControl", meta = (BlueprintThreadSafe))
 	static void AddModifierParameters(
 		UPARAM(ref, DisplayName = "Parameters") FPhysicsControlControlAndModifierParameters& InParameters,
@@ -51,7 +51,7 @@ class PHYSICSCONTROL_API UPhysicsControlBPLibrary : public UBlueprintFunctionLib
 		const FName                                                                          Name, 
 		const FPhysicsControlModifierSparseData&                                             ModifierData);
 
-	/** Add a single body modifier parameter Parameters to a set of parameter Parameters. */
+	// Add an array of body modifier parameters to a container of parameter.
 	UFUNCTION(BlueprintPure, Category = "Animation|PhysicsControl", meta = (BlueprintThreadSafe))
 	static void AddMultipleModifierParameters(
 		UPARAM(ref, DisplayName = "Parameters") FPhysicsControlControlAndModifierParameters& InParameters, 
@@ -59,7 +59,8 @@ class PHYSICSCONTROL_API UPhysicsControlBPLibrary : public UBlueprintFunctionLib
 		const TArray<FName>&                                                                 Names,
 		const FPhysicsControlModifierSparseData&                                             ModifierData);
 
-	/** Returns the linear interpolation of two sets of parameter Parameters. Any Parameters that exists in one of the input sets but not the other will be added to the output with a weight of 1. */
+	// Returns the linear interpolation of two sets of parameters. Any parameters that exist in one
+	// of the input sets but not the other will be added to the output with a weight of 1.
 	UFUNCTION(BlueprintPure, Category = "Animation|PhysicsControl", meta = (BlueprintThreadSafe))
 	static void BlendParameters(
 		UPARAM(ref, DisplayName = "From Parameters") FPhysicsControlControlAndModifierParameters& InParametersA,
@@ -67,7 +68,11 @@ class PHYSICSCONTROL_API UPhysicsControlBPLibrary : public UBlueprintFunctionLib
 		UPARAM(DisplayName = "Weight") const float                                                InInterpolationWeight,
 		UPARAM(DisplayName = "Parameters") FPhysicsControlControlAndModifierParameters&           OutParameters);
 	
-	/** Adds an Parameters to the output parameters for each control name in the supplied set. The values in each Parameters will be a linear interpolation of the two supplied Parameters, blending from the start Parameters to the end Parameters across the elements in the set. */
+	// Adds control parameters to the output parameters for each control name in the supplied array.
+	// The values in each control parameters added will be a linear interpolation of the two
+	// supplied Parameters, blending from the start parameters to the end parameters across the
+	// elements in the array. Note that this is most likely only useful when the control names are
+	// in order - for example going down a limb.
 	UFUNCTION(BlueprintPure, Category = "Animation|PhysicsControl", meta = (BlueprintThreadSafe))
 	static void BlendControlParametersThroughSet(
 		UPARAM(ref, DisplayName = "Parameters") FPhysicsControlControlAndModifierParameters&       InParameters,
@@ -76,7 +81,11 @@ class PHYSICSCONTROL_API UPhysicsControlBPLibrary : public UBlueprintFunctionLib
 		const TArray<FName>&                                                                       ControlNames,
 		FPhysicsControlControlAndModifierParameters&                                               OutParameters);
 
-	/** Adds an Parameters to the output parameters for each modifier name in the supplied set. The values in each Parameters will be a linear interpolation of the two supplied Parameters, blending from the start Parameters to the end Parameters across the elements in the set. */
+	// Adds body modifier parameters to the output parameters for each modifier name in the supplied
+	// array. The values in each modifier parameters added will be a linear interpolation of the two
+	// supplied parameters, blending from the start parameters to the end parameters across the
+	// elements in the array. Note that this is most likely only useful when the modifier names are
+	// in order - for example going down a limb.
 	UFUNCTION(BlueprintPure, Category = "Animation|PhysicsControl", meta = (BlueprintThreadSafe))
 	static void BlendModifierParametersThroughSet(
 		UPARAM(ref) FPhysicsControlControlAndModifierParameters&                        InParameters,
@@ -85,13 +94,13 @@ class PHYSICSCONTROL_API UPhysicsControlBPLibrary : public UBlueprintFunctionLib
 		const TArray<FName>&                                                            ModifierNames,
 		UPARAM(DisplayName = "Parameters") FPhysicsControlControlAndModifierParameters& OutParameters);
 
-	/** Get a Rigid Body With Control node reference from an anim node reference */
+	// Get a Rigid Body With Control node reference from an anim node reference 
 	UFUNCTION(BlueprintCallable, Category = "Animation|PhysicsControl", meta = (BlueprintThreadSafe, ExpandEnumAsExecs = "Result"))
 	static FRigidBodyWithControlReference ConvertToRigidBodyWithControl(
 		const FAnimNodeReference&           Node,
 		EAnimNodeReferenceConversionResult& Result);
 
-	/** Get a Rigid Body With Control node from an anim node (pure) */
+	// Get a Rigid Body With Control node from an anim node (pure)
 	UFUNCTION(BlueprintPure, Category = "Animation|PhysicsControl", meta = (BlueprintThreadSafe, DisplayName = "Convert to Rigid Body With Control"))
 	static void ConvertToRigidBodyWithControlPure(
 		const FAnimNodeReference&       Node, 
@@ -103,17 +112,18 @@ class PHYSICSCONTROL_API UPhysicsControlBPLibrary : public UBlueprintFunctionLib
 		Result = (ConversionResult == EAnimNodeReferenceConversionResult::Succeeded);
 	}
 
-	/** Set the physics asset on the rigid body with control anim graph node (RBWC). */
+	// Set the physics asset on the Rigid Body With Control anim graph node.
 	UFUNCTION(BlueprintCallable, Category = "Animation|PhysicsControl", meta = (BlueprintThreadSafe))
 	static FRigidBodyWithControlReference SetOverridePhysicsAsset(const FRigidBodyWithControlReference& Node, UPhysicsAsset* PhysicsAsset);
 
-	/** Get the names of all the controls in a specified set managed by this Rigid Body With Control node. */
+	// Get the names of all the controls in a specified set managed by this Rigid Body With Control node.
 	UFUNCTION(BlueprintPure, Category = "Animation|PhysicsControl", meta = (BlueprintThreadSafe))
 	static TArray<FName> GetControlNamesInSet(
 		UPARAM(DisplayName = "Node") const FRigidBodyWithControlReference& RigidBodyWithControl, 
 		const FName                                                        SetName);
 
-	/** Get the names of all the body modifiers in a specified set managed by this Rigid Body With Control node. */
+	// Get the names of all the body modifiers in a specified set managed by this Rigid Body With
+	// Control node.
 	UFUNCTION(BlueprintPure, Category = "Animation|PhysicsControl", meta = (BlueprintThreadSafe))
 	static TArray<FName> GetBodyModifierNamesInSet(
 		UPARAM(DisplayName = "Node") const FRigidBodyWithControlReference& RigidBodyWithControl, 
diff --git a/Engine/Plugins/Experimental/PhysicsControl/Source/PhysicsControl/Public/PhysicsControlComponent.h b/Engine/Plugins/Experimental/PhysicsControl/Source/PhysicsControl/Public/PhysicsControlComponent.h
--- a/Engine/Plugins/Experimental/PhysicsControl/Source/PhysicsControl/Public/PhysicsControlComponent.h
+++ b/Engine/Plugins/Experimental/PhysicsControl/Source/PhysicsControl/Public/PhysicsControlComponent.h
@@ -1848,5 +1848,5 @@ protected:
 
 	// Update counter - incremented every tick, and used to check whether previous targets etc are
 	// valid when calculating velocities.
-	int64 CurrentUpdateCounter = 0;
+	FGraphTraversalCounter CurrentUpdateCounter;
 };
diff --git a/Engine/Plugins/Experimental/PhysicsControl/Source/PhysicsControl/Public/PhysicsControlPoseData.h b/Engine/Plugins/Experimental/PhysicsControl/Source/PhysicsControl/Public/PhysicsControlPoseData.h
--- a/Engine/Plugins/Experimental/PhysicsControl/Source/PhysicsControl/Public/PhysicsControlPoseData.h
+++ b/Engine/Plugins/Experimental/PhysicsControl/Source/PhysicsControl/Public/PhysicsControlPoseData.h
@@ -40,11 +40,12 @@ struct FPosQuat
 		return FTransform(Rotation, Translation);
 	}
 
-	// Note that multiplication operates in the same sense as FTransform (i.e. "backwards")
+	// Note that multiplication operates in the same sense as FQuat - i.e. in reverse compared to FTransform
+	// WorldChildTM = WorldParentTM * ChildRelParentTM
 	FORCEINLINE FPosQuat operator*(const FPosQuat& Other) const
 	{
-		FQuat OutRotation = Other.Rotation * Rotation;
-		FVector OutTranslation = Other.Rotation * (Translation) + Other.Translation;
+		FQuat OutRotation = Rotation * Other.Rotation;
+		FVector OutTranslation = (Rotation * Other.Translation) + Translation;
 		return FPosQuat(OutTranslation, OutRotation);
 	}
 
diff --git a/Engine/Plugins/Experimental/PhysicsControl/Source/PhysicsControl/Public/PhysicsControlRecord.h b/Engine/Plugins/Experimental/PhysicsControl/Source/PhysicsControl/Public/PhysicsControlRecord.h
--- a/Engine/Plugins/Experimental/PhysicsControl/Source/PhysicsControl/Public/PhysicsControlRecord.h
+++ b/Engine/Plugins/Experimental/PhysicsControl/Source/PhysicsControl/Public/PhysicsControlRecord.h
@@ -23,7 +23,6 @@ struct FPhysicsControlRecord
 		UPrimitiveComponent*              InChildComponent)
 		: PhysicsControl(InControl)
 		, ControlTarget(InControlTarget)
-		, ExpectedUpdateCounter(-1)
 		, ParentComponent(InParentComponent)
 		, ChildComponent(InChildComponent)
 	{}
@@ -63,7 +62,7 @@ struct FPhysicsControlRecord
 
 	// Only use the previous target TM if the current counter is equal to this expected counter. The
 	// expected update counter will be set when the control/previous target TM has just been updated.
-	int64 ExpectedUpdateCounter;
+	FGraphTraversalCounter ExpectedUpdateCounter;
 
 	/**  The mesh that will be doing the driving. Blank/non-existent means it will happen in world space */
 	TWeakObjectPtr<UPrimitiveComponent> ParentComponent;
diff --git a/Engine/Plugins/Experimental/PhysicsControl/Source/PhysicsControl/Public/RigidBodyControlData.h b/Engine/Plugins/Experimental/PhysicsControl/Source/PhysicsControl/Public/RigidBodyControlData.h
--- a/Engine/Plugins/Experimental/PhysicsControl/Source/PhysicsControl/Public/RigidBodyControlData.h
+++ b/Engine/Plugins/Experimental/PhysicsControl/Source/PhysicsControl/Public/RigidBodyControlData.h
@@ -123,7 +123,6 @@ struct FRigidBodyControlRecord
 	UE::PhysicsControl::FPosQuat PreviousTargetTM;
 
 	// Update counter set when the control was last updated.
-	// TODO just store the count we're interested in rather than the whole structure
 	FGraphTraversalCounter ExpectedUpdateCounter;
 
 	// This is the currently active control data. It will be updated just prior to applying
diff --git a/Engine/Plugins/Experimental/PhysicsControl/Source/PhysicsControlUncookedOnly/Private/AnimGraphNode_RigidBodyWithControl.cpp b/Engine/Plugins/Experimental/PhysicsControl/Source/PhysicsControlUncookedOnly/Private/AnimGraphNode_RigidBodyWithControl.cpp
--- a/Engine/Plugins/Experimental/PhysicsControl/Source/PhysicsControlUncookedOnly/Private/AnimGraphNode_RigidBodyWithControl.cpp
+++ b/Engine/Plugins/Experimental/PhysicsControl/Source/PhysicsControlUncookedOnly/Private/AnimGraphNode_RigidBodyWithControl.cpp
@@ -10,6 +10,7 @@
 #include "IPhysicsControlOperatorEditorInterface.h"
 #include "Kismet2/BlueprintEditorUtils.h"
 #include "PhysicsControlOperatorNameGeneration.h"
+#include "PhysicsControlProfileAsset.h"
 
 // Details includes
 #include "PropertyHandle.h"
@@ -105,27 +106,6 @@ void UAnimGraphNode_RigidBodyWithControl::Draw(
 				PhysicsAssetRenderInterface.DebugDrawConstraints(PreviewSkelMeshComp, PhysicsAsset, PDI);
 			}
 		}
-		
-		
-		// Draw Space Controls
-		if (bIsSelected || 
-			(bIsPoseWatchEnabled && 
-				PoseWatchElementWorldSpaceControls.IsValid() && PoseWatchElementWorldSpaceControls->GetIsVisible()))
-		{
-			for (const TMap<FName, FRigidBodyControlRecord>::ElementType& NameRecordPair : RuntimeRigidBodyNode->ControlRecords)
-			{
-				const FRigidBodyControlRecord& ControlRecord = NameRecordPair.Value;
-				if (!ControlRecord.JointHandle || !ControlRecord.IsEnabled())
-				{
-					continue;
-				}
-#if 0
-				// TODO get the targets from the control record and draw them with respect to the physical bodies
-				PDI->DrawPoint(ControlRootTransform.GetTranslation(), FLinearColor::Blue, 5.0f, SDPG_Foreground);
-				PDI->DrawLine(BodyPosition, TargetPosition, FLinearColor::Yellow, SDPG_Foreground);
-#endif
-			}
-		}
 	}
 }
 
@@ -360,12 +340,35 @@ TArray<TPair<FName, TArray<FName>>> UAnimGraphNode_RigidBodyWithControl::Generat
 		TSet<FName> ControlNames;
 		FPhysicsControlNameRecords NameRecords;
 
-		// TODO This needs to account for the control profile asset
+		// Note that controls can come from the setup data in the node and/or from a profile asset
+		FPhysicsControlCharacterSetupData SetupData;
+		if (IsValid(Node.PhysicsControlProfileAsset))
+		{
+			SetupData = Node.PhysicsControlProfileAsset->CharacterSetupData;
+		}
+		if (Node.bEnableCharacterSetupData)
+		{
+			SetupData += Node.CharacterSetupData;
+		}
+
+		FPhysicsControlAndBodyModifierCreationDatas AdditionalControlAndBodyModifierCreationDatas;
+		if (IsValid(Node.PhysicsControlProfileAsset))
+		{
+			AdditionalControlAndBodyModifierCreationDatas = Node.PhysicsControlProfileAsset->AdditionalControlsAndModifiers;
+		}
+		AdditionalControlAndBodyModifierCreationDatas += Node.AdditionalControlsAndBodyModifiers;
+
+		// Get the list of modifier and control names, based on the setup data
 		UE::PhysicsControl::CollectOperatorNames(
-			&Node, Node.CharacterSetupData, Node.AdditionalControlsAndBodyModifiers,
+			&Node, SetupData, AdditionalControlAndBodyModifierCreationDatas,
 			AllLimbBones, RefSkeleton, Node.OverridePhysicsAsset.Get(), BodyModifierNames, ControlNames, NameRecords);
 
 		// Create any additional sets that have been requested
+		if (IsValid(Node.PhysicsControlProfileAsset))
+		{
+			UE::PhysicsControl::CreateAdditionalSets(
+				Node.PhysicsControlProfileAsset->AdditionalSets, BodyModifierNames, ControlNames, NameRecords);
+		}
 		UE::PhysicsControl::CreateAdditionalSets(Node.AdditionalSets, BodyModifierNames, ControlNames, NameRecords);
 
 		auto TransformOperatorNamesAndTags = [&GeneratedOperatorNames](
