Change some calls to use ensure rather than check

diff --git a/Engine/Plugins/Experimental/PhysicsControl/Source/PhysicsControl/Private/PhysicsControlComponentImpl.cpp b/Engine/Plugins/Experimental/PhysicsControl/Source/PhysicsControl/Private/PhysicsControlComponentImpl.cpp
--- a/Engine/Plugins/Experimental/PhysicsControl/Source/PhysicsControl/Private/PhysicsControlComponentImpl.cpp
+++ b/Engine/Plugins/Experimental/PhysicsControl/Source/PhysicsControl/Private/PhysicsControlComponentImpl.cpp
@@ -721,7 +721,7 @@ bool UPhysicsControlComponent::DestroyControl(const FName Name, const EDestroyBe
 		NameRecords.RemoveControl(Name);
 		if (DestroyBehavior == EDestroyBehavior::RemoveRecord)
 		{
-			check(ControlRecords.Remove(Name) == 1);
+			ensure(ControlRecords.Remove(Name) == 1);
 		}
 		return true;
 	}
@@ -747,7 +747,7 @@ bool UPhysicsControlComponent::DestroyBodyModifier(const FName Name, const EDest
 		NameRecords.RemoveBodyModifier(Name);
 		if (DestroyBehavior == EDestroyBehavior::RemoveRecord)
 		{
-			check(BodyModifierRecords.Remove(Name) == 1);
+			ensure(BodyModifierRecords.Remove(Name) == 1);
 		}
 		return true;
 	}
