Rename PhysicsControlProfilesAsset to PhysicsControlAsset - no functional changes, and the redirector should fix up any existing assets

diff --git a/Engine/Plugins/Experimental/PhysicsControl/Config/BasePhysicsControl.ini b/Engine/Plugins/Experimental/PhysicsControl/Config/BasePhysicsControl.ini
--- a/Engine/Plugins/Experimental/PhysicsControl/Config/BasePhysicsControl.ini
+++ b/Engine/Plugins/Experimental/PhysicsControl/Config/BasePhysicsControl.ini
@@ -12,4 +12,5 @@
 +PropertyRedirects=(OldName="PhysicsControlComponent.CreateBodyModifier.MeshComponent",NewName="Component")
 +PropertyRedirects=(OldName="PhysicsControlComponent.CreateNamedBodyModifier.MeshComponent",NewName="Component")
 
-
++PropertyRedirects=(OldName="PhysicsControlProfileAsset",NewName="PhysicsControlAsset")
++ClassRedirects=(OldName="PhysicsControlProfileAsset",NewName="PhysicsControlAsset")
diff --git a/Engine/Plugins/Experimental/PhysicsControl/Source/PhysicsControl/Private/PhysicsControlAsset.cpp b/Engine/Plugins/Experimental/PhysicsControl/Source/PhysicsControl/Private/PhysicsControlAsset.cpp
new file mode 100644
--- /dev/null
+++ b/Engine/Plugins/Experimental/PhysicsControl/Source/PhysicsControl/Private/PhysicsControlAsset.cpp
@@ -0,0 +1,163 @@
+// Copyright Epic Games, Inc. All Rights Reserved.
+
+#include "PhysicsControlAsset.h"
+#include "Engine/SkeletalMesh.h"
+
+//#ifdef WITH_EDITOR
+//#include "Editor.h"
+//#endif // WITH_EDITOR
+
+
+//======================================================================================================================
+UPhysicsControlAsset::UPhysicsControlAsset()
+{
+	// This needs to be explored further - a possible way to hook into the BP compilation. However, we would also
+	// need to make sure we don't compile (and get marked as dirty) if there are no changes.
+//#ifdef WITH_EDITOR
+//	GEditor->OnBlueprintCompiled().AddUObject(this, &UPhysicsControlAsset::Compile);
+//#endif
+}
+
+#if WITH_EDITOR
+//======================================================================================================================
+void UPhysicsControlAsset::ShowCompiledData() const
+{
+	UE_LOG(LogTemp, Log, TEXT("Character setup data:"));
+	for (const FPhysicsControlLimbSetupData& LimbSetupData : CharacterSetupData.LimbSetupData)
+	{
+		UE_LOG(LogTemp, Log, TEXT("Limb %s"), *LimbSetupData.LimbName.ToString());
+		UE_LOG(LogTemp, Log, TEXT("  Start bone %s"), *LimbSetupData.StartBone.ToString());
+		UE_LOG(LogTemp, Log, TEXT("  Include parent bone %d"), LimbSetupData.bIncludeParentBone);
+		UE_LOG(LogTemp, Log, TEXT("  Create world space controls %d"), LimbSetupData.bCreateWorldSpaceControls);
+		UE_LOG(LogTemp, Log, TEXT("  Create parent space controls %d"), LimbSetupData.bCreateParentSpaceControls);
+		UE_LOG(LogTemp, Log, TEXT("  Create body modifiers %d"), LimbSetupData.bCreateBodyModifiers);
+	}
+	UE_LOG(LogTemp, Log, TEXT("Additional controls and modifiers:"));
+	UE_LOG(LogTemp, Log, TEXT("  Additional controls:"));
+	for (TMap<FName, FPhysicsControlCreationData>::ElementType ControlPair : AdditionalControlsAndModifiers.Controls)
+	{
+		UE_LOG(LogTemp, Log, TEXT("    %s:"), *ControlPair.Key.ToString());
+		UE_LOG(LogTemp, Log, TEXT("      Parent bone %s:"), *ControlPair.Value.Control.ParentBoneName.ToString());
+		UE_LOG(LogTemp, Log, TEXT("      Child bone %s:"), *ControlPair.Value.Control.ParentBoneName.ToString());
+	}
+	UE_LOG(LogTemp, Log, TEXT("  Additional modifiers:"));
+	for (TMap<FName, FPhysicsBodyModifierCreationData>::ElementType ModifierPair : AdditionalControlsAndModifiers.Modifiers)
+	{
+		UE_LOG(LogTemp, Log, TEXT("    %s:"), *ModifierPair.Key.ToString());
+		UE_LOG(LogTemp, Log, TEXT("      Bone %s:"), *ModifierPair.Value.Modifier.BoneName.ToString());
+	}
+	UE_LOG(LogTemp, Log, TEXT("Profiles:"));
+	for (TMap<FName, FPhysicsControlControlAndModifierUpdates>::ElementType ProfilePair : Profiles)
+	{
+		UE_LOG(LogTemp, Log, TEXT("  %s:"), *ProfilePair.Key.ToString());
+	}
+
+}
+
+//======================================================================================================================
+void UPhysicsControlAsset::Compile()
+{
+	CharacterSetupData = GetCharacterSetupData();
+	AdditionalControlsAndModifiers = GetAdditionalControlsAndModifiers();
+	AdditionalSets = GetAdditionalSets();
+	InitialControlAndModifierUpdates = GetInitialControlAndModifierUpdates();
+	Profiles = GetProfiles();
+
+	Modify();
+}
+
+//======================================================================================================================
+FPhysicsControlCharacterSetupData UPhysicsControlAsset::GetCharacterSetupData() const
+{
+	FPhysicsControlCharacterSetupData CompiledCharacterSetupData;
+	if (ParentAsset.LoadSynchronous())
+	{
+		CompiledCharacterSetupData = ParentAsset->GetCharacterSetupData();
+	}
+	CompiledCharacterSetupData += MyCharacterSetupData;
+	return CompiledCharacterSetupData;
+}
+
+//======================================================================================================================
+FPhysicsControlAndBodyModifierCreationDatas UPhysicsControlAsset::GetAdditionalControlsAndModifiers() const
+{
+	FPhysicsControlAndBodyModifierCreationDatas CompiledAdditionalControlsAndModifiers;
+	if (ParentAsset.LoadSynchronous())
+	{
+		CompiledAdditionalControlsAndModifiers = ParentAsset->GetAdditionalControlsAndModifiers();
+	}
+	// this will overwrite duplicates with our value
+	CompiledAdditionalControlsAndModifiers += MyAdditionalControlsAndModifiers; 
+	return CompiledAdditionalControlsAndModifiers;
+}
+
+//======================================================================================================================
+FPhysicsControlSetUpdates UPhysicsControlAsset::GetAdditionalSets() const
+{
+	FPhysicsControlSetUpdates CompiledAdditionalSets;
+	if (ParentAsset.LoadSynchronous())
+	{
+		CompiledAdditionalSets = ParentAsset->GetAdditionalSets();
+	}
+	CompiledAdditionalSets += MyAdditionalSets;
+	return CompiledAdditionalSets;
+}
+
+//======================================================================================================================
+TArray<FPhysicsControlControlAndModifierUpdates> UPhysicsControlAsset::GetInitialControlAndModifierUpdates() const
+{
+	TArray<FPhysicsControlControlAndModifierUpdates> CompiledInitialControlAndModifierUpdates;
+	if (ParentAsset.LoadSynchronous())
+	{
+		CompiledInitialControlAndModifierUpdates = ParentAsset->GetInitialControlAndModifierUpdates();
+	}
+	CompiledInitialControlAndModifierUpdates.Append(MyInitialControlAndModifierUpdates);
+	return CompiledInitialControlAndModifierUpdates;
+}
+
+//======================================================================================================================
+TMap<FName, FPhysicsControlControlAndModifierUpdates> UPhysicsControlAsset::GetProfiles() const
+{
+	TMap<FName, FPhysicsControlControlAndModifierUpdates> CompiledProfiles;
+	if (ParentAsset.LoadSynchronous())
+	{
+		CompiledProfiles = ParentAsset->GetProfiles();
+	}
+	for (const TSoftObjectPtr<UPhysicsControlAsset>& AdditionalProfileAsset : AdditionalProfileAssets)
+	{
+		if (AdditionalProfileAsset.LoadSynchronous())
+		{
+			CompiledProfiles.Append(AdditionalProfileAsset->GetProfiles());
+		}
+	}
+	CompiledProfiles.Append(MyProfiles); // this will overwrite duplicates with our value
+	return CompiledProfiles;
+}
+
+#endif
+
+#if WITH_EDITOR
+//======================================================================================================================
+const FName UPhysicsControlAsset::GetPreviewMeshPropertyName()
+{
+	return GET_MEMBER_NAME_STRING_CHECKED(UPhysicsControlAsset, PreviewSkeletalMesh);
+};
+#endif
+
+//======================================================================================================================
+void UPhysicsControlAsset::SetPreviewMesh(USkeletalMesh* PreviewMesh, bool bMarkAsDirty)
+{
+#if WITH_EDITOR
+	PreviewSkeletalMesh = PreviewMesh;
+#endif
+}
+
+//======================================================================================================================
+USkeletalMesh* UPhysicsControlAsset::GetPreviewMesh() const
+{
+#if WITH_EDITOR
+	return PreviewSkeletalMesh.LoadSynchronous();
+#else
+	return nullptr;
+#endif
+}
diff --git a/Engine/Plugins/Experimental/PhysicsControl/Source/PhysicsControl/Private/PhysicsControlComponent.cpp b/Engine/Plugins/Experimental/PhysicsControl/Source/PhysicsControl/Private/PhysicsControlComponent.cpp
--- a/Engine/Plugins/Experimental/PhysicsControl/Source/PhysicsControl/Private/PhysicsControlComponent.cpp
+++ b/Engine/Plugins/Experimental/PhysicsControl/Source/PhysicsControl/Private/PhysicsControlComponent.cpp
@@ -2102,7 +2102,7 @@ void UPhysicsControlComponent::CreateControlsAndBodyModifiersFromControlProfileA
 	UPrimitiveComponent*    WorldComponent,
 	FName                   WorldBoneName)
 {
-	if (!PhysicsControlProfileAsset.IsValid())
+	if (!PhysicsControlAsset.IsValid())
 	{
 		UE_LOG(LogPhysicsControl, Warning,
 			TEXT("CreateControlsAndBodyModifiersFromControlProfile - unable to get/load the control profile asset"));
@@ -2120,16 +2120,16 @@ void UPhysicsControlComponent::CreateControlsAndBodyModifiersFromControlProfileA
 		AllWorldSpaceControls, LimbWorldSpaceControls, AllParentSpaceControls, LimbParentSpaceControls, 
 		AllBodyModifiers, LimbBodyModifiers,
 		SkeletalMeshComponent,
-		PhysicsControlProfileAsset->CharacterSetupData.LimbSetupData,
-		PhysicsControlProfileAsset->CharacterSetupData.DefaultWorldSpaceControlData,
-		PhysicsControlProfileAsset->CharacterSetupData.DefaultParentSpaceControlData,
-		PhysicsControlProfileAsset->CharacterSetupData.DefaultBodyModifierData,
+		PhysicsControlAsset->CharacterSetupData.LimbSetupData,
+		PhysicsControlAsset->CharacterSetupData.DefaultWorldSpaceControlData,
+		PhysicsControlAsset->CharacterSetupData.DefaultParentSpaceControlData,
+		PhysicsControlAsset->CharacterSetupData.DefaultBodyModifierData,
 		WorldComponent,
 		WorldBoneName);
 
 	// Create additional controls
 	for (const TPair<FName, FPhysicsControlCreationData>& ControlPair : 
-		PhysicsControlProfileAsset->AdditionalControlsAndModifiers.Controls)
+		PhysicsControlAsset->AdditionalControlsAndModifiers.Controls)
 	{
 		FName ControlName = ControlPair.Key;
 		const FPhysicsControlCreationData& ControlCreationData = ControlPair.Value;
@@ -2148,7 +2148,7 @@ void UPhysicsControlComponent::CreateControlsAndBodyModifiersFromControlProfileA
 
 	// Create additional modifiers
 	for (const TPair<FName, FPhysicsBodyModifierCreationData>& ModifierPair : 
-		PhysicsControlProfileAsset->AdditionalControlsAndModifiers.Modifiers)
+		PhysicsControlAsset->AdditionalControlsAndModifiers.Modifiers)
 	{
 		FName ModifierName = ModifierPair.Key;
 		const FPhysicsBodyModifierCreationData& ModifierCreationData = ModifierPair.Value;
@@ -2165,9 +2165,9 @@ void UPhysicsControlComponent::CreateControlsAndBodyModifiersFromControlProfileA
 
 	// Create any additional sets that have been requested
 	UE::PhysicsControl::CreateAdditionalSets(
-		PhysicsControlProfileAsset->AdditionalSets, BodyModifierRecords, ControlRecords, NameRecords);
+		PhysicsControlAsset->AdditionalSets, BodyModifierRecords, ControlRecords, NameRecords);
 
-	for (FPhysicsControlControlAndModifierUpdates& Updates : PhysicsControlProfileAsset->InitialControlAndModifierUpdates)
+	for (FPhysicsControlControlAndModifierUpdates& Updates : PhysicsControlAsset->InitialControlAndModifierUpdates)
 	{
 		ApplyControlAndModifierUpdates(Updates);
 	}
@@ -2176,7 +2176,7 @@ void UPhysicsControlComponent::CreateControlsAndBodyModifiersFromControlProfileA
 //======================================================================================================================
 void UPhysicsControlComponent::InvokeControlProfile(FName ProfileName)
 {
-	if (!PhysicsControlProfileAsset.IsValid())
+	if (!PhysicsControlAsset.IsValid())
 	{
 		if (bWarnAboutInvalidNames)
 		{
@@ -2187,7 +2187,7 @@ void UPhysicsControlComponent::InvokeControlProfile(FName ProfileName)
 	}
 
 	const FPhysicsControlControlAndModifierUpdates* ControlAndModifierUpdates =
-		PhysicsControlProfileAsset->Profiles.Find(ProfileName);
+		PhysicsControlAsset->Profiles.Find(ProfileName);
 
 	if (!ControlAndModifierUpdates)
 	{
diff --git a/Engine/Plugins/Experimental/PhysicsControl/Source/PhysicsControl/Private/PhysicsControlProfileAsset.cpp b/Engine/Plugins/Experimental/PhysicsControl/Source/PhysicsControl/Private/PhysicsControlProfileAsset.cpp
deleted file mode 100644
diff --git a/Engine/Plugins/Experimental/PhysicsControl/Source/PhysicsControl/Private/RigidBodyWithControl.cpp b/Engine/Plugins/Experimental/PhysicsControl/Source/PhysicsControl/Private/RigidBodyWithControl.cpp
--- a/Engine/Plugins/Experimental/PhysicsControl/Source/PhysicsControl/Private/RigidBodyWithControl.cpp
+++ b/Engine/Plugins/Experimental/PhysicsControl/Source/PhysicsControl/Private/RigidBodyWithControl.cpp
@@ -4,7 +4,7 @@
 
 #include "PhysicsControlOperatorNameGeneration.h"
 #include "PhysicsControlLog.h"
-#include "PhysicsControlProfileAsset.h"
+#include "PhysicsControlAsset.h"
 #include "PhysicsControlComponentHelpers.h"
 
 #include "Animation/AnimInstance.h"
@@ -357,9 +357,9 @@ void FAnimNode_RigidBodyWithControl::InitControlsAndBodyModifiers(const FReferen
 	check(ControlRecords.IsEmpty()); // Controls should not exist when this function is called.
 
 	FPhysicsControlCharacterSetupData SetupData;
-	if (IsValid(PhysicsControlProfileAsset))
+	if (IsValid(PhysicsControlAsset))
 	{
-		SetupData = PhysicsControlProfileAsset->CharacterSetupData;
+		SetupData = PhysicsControlAsset->CharacterSetupData;
 	}
 	if (bEnableCharacterSetupData)
 	{
@@ -371,9 +371,9 @@ void FAnimNode_RigidBodyWithControl::InitControlsAndBodyModifiers(const FReferen
 		UE::PhysicsControl::GetLimbBones(SetupData.LimbSetupData, RefSkeleton, GetPhysicsAsset());
 
 	FPhysicsControlAndBodyModifierCreationDatas AdditionalControlAndBodyModifierCreationDatas;
-	if (IsValid(PhysicsControlProfileAsset))
+	if (IsValid(PhysicsControlAsset))
 	{
-		AdditionalControlAndBodyModifierCreationDatas = PhysicsControlProfileAsset->AdditionalControlsAndModifiers;
+		AdditionalControlAndBodyModifierCreationDatas = PhysicsControlAsset->AdditionalControlsAndModifiers;
 	}
 	AdditionalControlAndBodyModifierCreationDatas += AdditionalControlsAndBodyModifiers;
 
@@ -389,10 +389,10 @@ void FAnimNode_RigidBodyWithControl::InitControlsAndBodyModifiers(const FReferen
 	}
 
 	// Create any additional sets that have been requested
-	if (IsValid(PhysicsControlProfileAsset))
+	if (IsValid(PhysicsControlAsset))
 	{
 		UE::PhysicsControl::CreateAdditionalSets(
-			PhysicsControlProfileAsset->AdditionalSets, ModifierRecords, ControlRecords, NameRecords);
+			PhysicsControlAsset->AdditionalSets, ModifierRecords, ControlRecords, NameRecords);
 	}
 	UE::PhysicsControl::CreateAdditionalSets(AdditionalSets, ModifierRecords, ControlRecords, NameRecords);
 
@@ -951,10 +951,10 @@ void FAnimNode_RigidBodyWithControl::InvokeControlProfile(FName ControlProfileNa
 void FAnimNode_RigidBodyWithControl::ApplyCurrentControlProfile()
 {
 	// We shouldn't have a hitch here, since the asset (if set) will already have been loaded 
-	if (IsValid(PhysicsControlProfileAsset))
+	if (IsValid(PhysicsControlAsset))
 	{
 		const FPhysicsControlControlAndModifierUpdates* Updates = 
-			PhysicsControlProfileAsset->Profiles.Find(CurrentControlProfile);
+			PhysicsControlAsset->Profiles.Find(CurrentControlProfile);
 		if (Updates)
 		{
 			ApplyControlAndBodyModifierDatas(
diff --git a/Engine/Plugins/Experimental/PhysicsControl/Source/PhysicsControl/Public/AnimNode_RigidBodyWithControl.h b/Engine/Plugins/Experimental/PhysicsControl/Source/PhysicsControl/Public/AnimNode_RigidBodyWithControl.h
--- a/Engine/Plugins/Experimental/PhysicsControl/Source/PhysicsControl/Public/AnimNode_RigidBodyWithControl.h
+++ b/Engine/Plugins/Experimental/PhysicsControl/Source/PhysicsControl/Public/AnimNode_RigidBodyWithControl.h
@@ -26,7 +26,7 @@ struct FBodyInstance;
 struct FConstraintInstance;
 
 class FEvent;
-class UPhysicsControlProfileAsset;
+class UPhysicsControlAsset;
 class UPhysicsAsset;
 
 extern TAutoConsoleVariable<int32> CVarEnableRigidBodyNodeWithControl;
@@ -281,12 +281,12 @@ public:
 	bool bEnableControls;
 
 	/** 
-	 * Physics control profile asset to use. This is optional, but if it has been set, then it will be used 
+	 * Physics control asset to use. This is optional, but if it has been set, then it will be used 
 	 * during creation of controls/body modifiers. Note that this can be modified at runtime in the Animation 
 	 * Blueprint via the Animation Node functions.
 	 */
 	UPROPERTY(EditAnywhere, Category = ControlSetup)
-	TObjectPtr<UPhysicsControlProfileAsset> PhysicsControlProfileAsset;
+	TObjectPtr<UPhysicsControlAsset> PhysicsControlAsset;
 
 	/**
 	 * A map of bone names to "body" names, the latter being used to assign names to controls/modifiers. 
diff --git a/Engine/Plugins/Experimental/PhysicsControl/Source/PhysicsControl/Public/PhysicsControlAsset.h b/Engine/Plugins/Experimental/PhysicsControl/Source/PhysicsControl/Public/PhysicsControlAsset.h
new file mode 100644
--- /dev/null
+++ b/Engine/Plugins/Experimental/PhysicsControl/Source/PhysicsControl/Public/PhysicsControlAsset.h
@@ -0,0 +1,166 @@
+// Copyright Epic Games, Inc. All Rights Reserved.
+
+#pragma once
+
+#include "PhysicsControlData.h"
+#include "PhysicsControlLimbData.h"
+
+#include "Interfaces/Interface_PreviewMeshProvider.h"
+
+#include "PhysicsControlAsset.generated.h"
+
+class USkeletalMesh;
+
+/**
+ * Asset for storing Physics Control Profiles. These will contain data that define:
+ * - Controls and body modifiers to be created on a mesh
+ * - Sets referencing those controls and body modifiers
+ * - Full profiles containing settings for all the controls/modifiers
+ * - Sparse profiles containing partial sets of settings for specific controls/modifiers
+ * 
+ * It will also be desirable to support "inheritance" - so a generic profile can be made, and then 
+ * customized for certain characters or scenarios.
+ */
+UCLASS(BlueprintType)
+class PHYSICSCONTROL_API UPhysicsControlAsset : public UObject, public IInterface_PreviewMeshProvider
+{
+	GENERATED_BODY()
+public:
+	UPhysicsControlAsset();
+
+	// Data that have been compiled from a combination of inherited and "My" data.
+
+	/**
+	 * We can define controls in the form of limbs etc here
+	 */
+	UPROPERTY()
+	FPhysicsControlCharacterSetupData CharacterSetupData;
+
+	/**
+	 * Additional controls and modifiers. If these have the same name as one that's already 
+	 * created, they'll just override it.
+	 */
+	UPROPERTY()
+	FPhysicsControlAndBodyModifierCreationDatas AdditionalControlsAndModifiers;
+
+	/**
+	 * Additional control and body modifier sets
+	 */
+	UPROPERTY()
+	FPhysicsControlSetUpdates AdditionalSets;
+
+	/**
+	 * Initial updates to apply immediately after controls and modifiers are created
+	 */
+	UPROPERTY()
+	TArray<FPhysicsControlControlAndModifierUpdates> InitialControlAndModifierUpdates;
+
+	/**
+	 * The named profiles, which are essentially control and modifier updates
+	 */
+	UPROPERTY()
+	TMap<FName, FPhysicsControlControlAndModifierUpdates> Profiles;
+
+public:
+	// Data that will then be compiled down into the runtime data
+
+#if WITH_EDITORONLY_DATA
+	/** A profile asset to inherit from (can be null). If set, we will just add/modify data in that */
+	UPROPERTY(EditAnywhere, Category = Inheritance)
+	TSoftObjectPtr<UPhysicsControlAsset> ParentAsset;
+
+	/** 
+	 * Additional profile assets from which profiles (not the setup data, extra sets etc) will be added 
+	 * to this asset.
+	 */
+	UPROPERTY(EditAnywhere, Category = Inheritance)
+	TArray<TSoftObjectPtr<UPhysicsControlAsset>> AdditionalProfileAssets;
+
+	/**
+	* The skeletal mesh to use for generating controls and previewing. If it turns out this 
+	* doesn't need to be stored in the asset, it will get moved out of here.
+	*/
+	UPROPERTY(AssetRegistrySearchable, EditAnywhere, Category = PreviewMesh)
+	TSoftObjectPtr<USkeletalMesh> PreviewSkeletalMesh;
+
+public:
+	// "My" runtime data - i.e. the data that will be combined with what has been inherited
+	// We should have custom UI that displays this combined with the inherited data
+
+	/**
+	 * We can define controls in the form of limbs etc here
+	 */
+	UPROPERTY(EditAnywhere, Category = ProfileData, meta=(DisplayName="Character Setup Data"))
+	FPhysicsControlCharacterSetupData MyCharacterSetupData;
+
+	/**
+	 * Additional controls and modifiers. If these have the same name as one that's 
+	 * already created, they'll just override it.
+	 */
+	UPROPERTY(EditAnywhere, Category = ProfileData, meta = (DisplayName = "Additional Controls and Modifiers"))
+	FPhysicsControlAndBodyModifierCreationDatas MyAdditionalControlsAndModifiers;
+
+	/**
+	 * Additional control and body modifier sets
+	 */
+	UPROPERTY(EditAnywhere, Category = ProfileData, meta = (DisplayName = "Additional Sets"))
+	FPhysicsControlSetUpdates MyAdditionalSets;
+
+	/**
+	 * Initial updates to apply immediately after controls and modifiers are created
+	 */
+	UPROPERTY(EditAnywhere, Category = ProfileData, meta = (DisplayName = "Initial Control and Modifier Updates"))
+	TArray<FPhysicsControlControlAndModifierUpdates> MyInitialControlAndModifierUpdates;
+
+	/**
+	 * The named profiles, which are essentially control and modifier updates
+	 */
+	UPROPERTY(EditAnywhere, Category = ProfileData, meta = (DisplayName = "Profiles"))
+	TMap<FName, FPhysicsControlControlAndModifierUpdates> MyProfiles;
+#endif
+
+public:
+	// Buttons/actions in the editor
+
+#if WITH_EDITOR
+	/** Shows all the controls etc that would be made */
+	UFUNCTION(CallInEditor, Category = Development)
+	void ShowCompiledData() const;
+
+	/** 
+	 * Collapses inherited and authored profiles etc to make a profile asset that can be read without 
+	 * need for subsequent processing.
+	 */
+	UFUNCTION(CallInEditor)
+	void Compile();
+
+	/** Combines and returns data from our parent and ourself */
+	FPhysicsControlCharacterSetupData GetCharacterSetupData() const;
+
+	/** Combines and returns data from our parent and ourself */
+	FPhysicsControlAndBodyModifierCreationDatas GetAdditionalControlsAndModifiers() const;
+
+	/** Combines and returns data from our parent and ourself */
+	FPhysicsControlSetUpdates GetAdditionalSets() const;
+
+	/** Combines and returns data from our parent and ourself */
+	TArray<FPhysicsControlControlAndModifierUpdates> GetInitialControlAndModifierUpdates() const;
+
+	/** Combines and returns data from our parent and ourself */
+	TMap<FName, FPhysicsControlControlAndModifierUpdates> GetProfiles() const;
+
+#endif
+
+public:
+
+	/** IInterface_PreviewMeshProvider interface */
+	virtual void SetPreviewMesh(USkeletalMesh* PreviewMesh, bool bMarkAsDirty = true) override;
+	virtual USkeletalMesh* GetPreviewMesh() const override;
+	/** END IInterface_PreviewMeshProvider interface */
+
+#if WITH_EDITOR
+	/* Get name of Preview Mesh property */
+	static const FName GetPreviewMeshPropertyName();
+#endif
+
+};
diff --git a/Engine/Plugins/Experimental/PhysicsControl/Source/PhysicsControl/Public/PhysicsControlComponent.h b/Engine/Plugins/Experimental/PhysicsControl/Source/PhysicsControl/Public/PhysicsControlComponent.h
--- a/Engine/Plugins/Experimental/PhysicsControl/Source/PhysicsControl/Public/PhysicsControlComponent.h
+++ b/Engine/Plugins/Experimental/PhysicsControl/Source/PhysicsControl/Public/PhysicsControlComponent.h
@@ -5,7 +5,7 @@
 #include "PhysicsControlData.h"
 #include "PhysicsControlLimbData.h"
 #include "PhysicsControlNameRecords.h"
-#include "PhysicsControlProfileAsset.h"
+#include "PhysicsControlAsset.h"
 #include "PhysicsControlRecord.h"
 #include "PhysicsControlPoseData.h"
 
@@ -1606,8 +1606,12 @@ public:
 public:
 	// Public property data
 
-	UPROPERTY(EditAnywhere, BlueprintReadWrite, Category = PhysicsControlProfile)
-	TSoftObjectPtr<UPhysicsControlProfileAsset> PhysicsControlProfileAsset;
+	/**
+	 * Optional PhysicsControlAsset to use. If this is set, then controls can be instantiated from the asset. 
+	 * Individual profiles can subsequently be "invoked" referencing the named profiles in the asset.
+	 */
+	UPROPERTY(EditAnywhere, BlueprintReadWrite, Category = PhysicsControlAsset)
+	TSoftObjectPtr<UPhysicsControlAsset> PhysicsControlAsset;
 
 	/**
 	 * If the component moves by more than this distance then it is treated as a teleport,
diff --git a/Engine/Plugins/Experimental/PhysicsControl/Source/PhysicsControl/Public/PhysicsControlProfileAsset.h b/Engine/Plugins/Experimental/PhysicsControl/Source/PhysicsControl/Public/PhysicsControlProfileAsset.h
deleted file mode 100644
diff --git a/Engine/Plugins/Experimental/PhysicsControl/Source/PhysicsControlEditor/Private/PhysicsControlAssetActions.cpp b/Engine/Plugins/Experimental/PhysicsControl/Source/PhysicsControlEditor/Private/PhysicsControlAssetActions.cpp
new file mode 100644
--- /dev/null
+++ b/Engine/Plugins/Experimental/PhysicsControl/Source/PhysicsControlEditor/Private/PhysicsControlAssetActions.cpp
@@ -0,0 +1,47 @@
+// Copyright Epic Games, Inc. All Rights Reserved.
+
+#include "PhysicsControlAssetActions.h"
+#include "PhysicsControlAsset.h"
+#include "PhysicsControlAssetEditorToolkit.h"
+
+//======================================================================================================================
+UClass* FPhysicsControlAssetActions::GetSupportedClass() const
+{
+	return UPhysicsControlAsset::StaticClass();
+}
+
+//======================================================================================================================
+FText FPhysicsControlAssetActions::GetName() const
+{
+	return INVTEXT("Physics Control Profile");
+}
+
+//======================================================================================================================
+FColor FPhysicsControlAssetActions::GetTypeColor() const
+{
+	// Match the "standard" physics color - they tend to be variations around this value
+	return FColor(255, 192, 128);
+}
+
+//======================================================================================================================
+uint32 FPhysicsControlAssetActions::GetCategories()
+{
+	return EAssetTypeCategories::Physics;
+}
+
+//======================================================================================================================
+void FPhysicsControlAssetActions::OpenAssetEditor(
+	const TArray<UObject*>& InObjects, TSharedPtr<IToolkitHost> EditWithinLevelEditor)
+{
+	const EToolkitMode::Type Mode = 
+		EditWithinLevelEditor.IsValid() ? EToolkitMode::WorldCentric : EToolkitMode::Standalone;
+
+	for (UObject* Object : InObjects)
+	{
+		if (UPhysicsControlAsset* Asset = Cast<UPhysicsControlAsset>(Object))
+		{
+			TSharedRef<FPhysicsControlAssetEditorToolkit> NewEditor(new FPhysicsControlAssetEditorToolkit());
+			NewEditor->InitAssetEditor(Mode, EditWithinLevelEditor, Asset);
+		}
+	}
+}
diff --git a/Engine/Plugins/Experimental/PhysicsControl/Source/PhysicsControlEditor/Private/PhysicsControlAssetActions.h b/Engine/Plugins/Experimental/PhysicsControl/Source/PhysicsControlEditor/Private/PhysicsControlAssetActions.h
new file mode 100644
--- /dev/null
+++ b/Engine/Plugins/Experimental/PhysicsControl/Source/PhysicsControlEditor/Private/PhysicsControlAssetActions.h
@@ -0,0 +1,19 @@
+// Copyright Epic Games, Inc. All Rights Reserved.
+
+#pragma once
+
+#include "CoreMinimal.h"
+#include "AssetTypeActions_Base.h"
+
+class FPhysicsControlAssetActions : public FAssetTypeActions_Base
+{
+public:
+	UClass* GetSupportedClass() const override;
+	FText GetName() const override;
+	FColor GetTypeColor() const override;
+	uint32 GetCategories() override;
+
+	void OpenAssetEditor(
+		const TArray<UObject*>& InObjects, 
+		TSharedPtr<class IToolkitHost> EditWithinLevelEditor = TSharedPtr<IToolkitHost>()) override;
+};
diff --git a/Engine/Plugins/Experimental/PhysicsControl/Source/PhysicsControlEditor/Private/PhysicsControlAssetApplicationMode.cpp b/Engine/Plugins/Experimental/PhysicsControl/Source/PhysicsControlEditor/Private/PhysicsControlAssetApplicationMode.cpp
new file mode 100644
--- /dev/null
+++ b/Engine/Plugins/Experimental/PhysicsControl/Source/PhysicsControlEditor/Private/PhysicsControlAssetApplicationMode.cpp
@@ -0,0 +1,89 @@
+// Copyright Epic Games, Inc. All Rights Reserved.
+
+#include "PhysicsControlAssetApplicationMode.h"
+#include "PhysicsControlAssetEditorToolkit.h"
+#include "PersonaModule.h"
+#include "Modules/ModuleManager.h"
+#include "PersonaTabs.h"
+
+#define LOCTEXT_NAMESPACE "PhysicsControlAssetApplicationMode"
+
+FName FPhysicsControlAssetApplicationMode::ModeName("PhysicsControlAssetEditMode");
+
+//======================================================================================================================
+FPhysicsControlAssetApplicationMode::FPhysicsControlAssetApplicationMode(
+	TSharedRef<FWorkflowCentricApplication> InHostingApp,
+	TSharedRef<IPersonaPreviewScene>        InPreviewScene)
+	: 
+	FApplicationMode(PhysicsControlAssetEditorModes::Editor)
+{
+	EditorToolkit = StaticCastSharedRef<FPhysicsControlAssetEditorToolkit>(InHostingApp);
+	TSharedRef<FPhysicsControlAssetEditorToolkit> PhysicsControlAssetEditor = 
+		StaticCastSharedRef<FPhysicsControlAssetEditorToolkit>(InHostingApp);
+
+	FPersonaViewportArgs ViewportArgs(InPreviewScene);
+	ViewportArgs.bAlwaysShowTransformToolbar = true;
+	ViewportArgs.bShowStats = false;
+	ViewportArgs.bShowTimeline = false;
+	ViewportArgs.bShowLODMenu = true;
+	ViewportArgs.bShowPlaySpeedMenu = false;
+	ViewportArgs.bShowPhysicsMenu = false;
+	ViewportArgs.ContextName = TEXT("PhysicsControlAssetEditor.Viewport");
+	ViewportArgs.OnViewportCreated = FOnViewportCreated::CreateSP(
+		PhysicsControlAssetEditor, &FPhysicsControlAssetEditorToolkit::HandleViewportCreated);
+
+	// Register Persona tabs.
+	FPersonaModule& PersonaModule = FModuleManager::LoadModuleChecked<FPersonaModule>("Persona");
+	TabFactories.RegisterFactory(PersonaModule.CreatePersonaViewportTabFactory(InHostingApp, ViewportArgs));
+	TabFactories.RegisterFactory(PersonaModule.CreateDetailsTabFactory(InHostingApp, FOnDetailsCreated::CreateSP(&PhysicsControlAssetEditor.Get(), &FPhysicsControlAssetEditorToolkit::HandleDetailsCreated)));
+
+	// Create tab layout.
+	TabLayout = FTabManager::NewLayout("Standalone_PhysicsControlAssetEditor_Layout_v0.001")
+		->AddArea
+		(
+			FTabManager::NewPrimaryArea()
+			->SetOrientation(Orient_Vertical)
+			->Split
+			(
+				FTabManager::NewSplitter()
+				->SetSizeCoefficient(0.32f)
+				->SetOrientation(Orient_Horizontal)
+				->Split
+				(
+					FTabManager::NewSplitter()
+					->SetSizeCoefficient(1.0f)
+					->SetOrientation(Orient_Vertical)
+					->Split
+					(
+						FTabManager::NewStack()
+						->SetSizeCoefficient(0.7f)
+						->SetHideTabWell(true)
+						->AddTab(FPersonaTabs::PreviewViewportID, ETabState::OpenedTab)
+					)
+				)
+				->Split
+				(
+					FTabManager::NewStack()
+					->SetSizeCoefficient(0.3f)
+					->AddTab(FPersonaTabs::DetailsID, ETabState::OpenedTab)
+					->SetForegroundTab(FPersonaTabs::DetailsID)
+				)
+			)
+		);
+
+	PersonaModule.OnRegisterTabs().Broadcast(TabFactories, InHostingApp);
+	LayoutExtender = MakeShared<FLayoutExtender>();
+	PersonaModule.OnRegisterLayoutExtensions().Broadcast(*LayoutExtender.Get());
+	TabLayout->ProcessExtensions(*LayoutExtender.Get());
+}
+
+//======================================================================================================================
+void FPhysicsControlAssetApplicationMode::RegisterTabFactories(TSharedPtr<FTabManager> InTabManager)
+{
+	TSharedPtr<FPhysicsControlAssetEditorToolkit> Editor = EditorToolkit.Pin();
+	Editor->RegisterTabSpawners(InTabManager.ToSharedRef());
+	Editor->PushTabFactories(TabFactories);
+	FApplicationMode::RegisterTabFactories(InTabManager);
+}
+
+#undef LOCTEXT_NAMESPACE
diff --git a/Engine/Plugins/Experimental/PhysicsControl/Source/PhysicsControlEditor/Private/PhysicsControlAssetEditorData.cpp b/Engine/Plugins/Experimental/PhysicsControl/Source/PhysicsControlEditor/Private/PhysicsControlAssetEditorData.cpp
new file mode 100644
--- /dev/null
+++ b/Engine/Plugins/Experimental/PhysicsControl/Source/PhysicsControlEditor/Private/PhysicsControlAssetEditorData.cpp
@@ -0,0 +1,62 @@
+// Copyright Epic Games, Inc. All Rights Reserved.
+
+#include "PhysicsControlAssetEditorData.h"
+#include "PhysicsControlAsset.h"
+#include "Engine/SkeletalMesh.h"
+#include "Misc/MessageDialog.h"
+
+#define LOCTEXT_NAMESPACE "PhysicsControlAssetEditorData"
+
+//======================================================================================================================
+FPhysicsControlAssetEditorData::FPhysicsControlAssetEditorData()
+{
+}
+
+//======================================================================================================================
+void FPhysicsControlAssetEditorData::Initialize(const TSharedRef<IPersonaPreviewScene>& InPreviewScene)
+{
+	PreviewScene = InPreviewScene;
+
+	EditorSkelComp = nullptr;
+	PhysicsControlComponent = nullptr;
+	FSoftObjectPath PreviewMeshStringRef = PhysicsControlAsset->PreviewSkeletalMesh.ToSoftObjectPath();
+
+	// Support undo/redo
+	PhysicsControlAsset->SetFlags(RF_Transactional);
+}
+
+//======================================================================================================================
+void FPhysicsControlAssetEditorData::CachePreviewMesh()
+{
+	USkeletalMesh* PreviewMesh = PhysicsControlAsset->PreviewSkeletalMesh.LoadSynchronous();
+
+	if (PreviewMesh == nullptr)
+	{
+		// Fall back to the default skeletal mesh in the EngineMeshes package.
+		// This is statically loaded as the package is likely not fully loaded
+		// (otherwise, it would have been found in the above iteration).
+		PreviewMesh = (USkeletalMesh*)StaticLoadObject(USkeletalMesh::StaticClass(), NULL, TEXT("/Engine/EngineMeshes/SkeletalCube.SkeletalCube"), NULL, LOAD_None, NULL);
+		check(PreviewMesh);
+
+		PhysicsControlAsset->PreviewSkeletalMesh = PreviewMesh;
+
+		FMessageDialog::Open(EAppMsgType::Ok, FText::Format(
+			LOCTEXT("Error_PhysicsControlAssetHasNoSkelMesh", "Warning: Physics Control Profile Asset has no skeletal mesh assigned.\nFor now, a simple default skeletal mesh ({0}) will be used.\nYou can fix this by opening the asset and choosing another skeletal mesh from the toolbar."),
+			FText::FromString(PreviewMesh->GetFullName())));
+	}
+	else if (PreviewMesh->GetSkeleton() == nullptr)
+	{
+		// Fall back in the case of a deleted skeleton
+		PreviewMesh = (USkeletalMesh*)StaticLoadObject(USkeletalMesh::StaticClass(), NULL, TEXT("/Engine/EngineMeshes/SkeletalCube.SkeletalCube"), NULL, LOAD_None, NULL);
+		check(PreviewMesh);
+
+		PhysicsControlAsset->PreviewSkeletalMesh = PreviewMesh;
+
+		FMessageDialog::Open(EAppMsgType::Ok, FText::Format(
+			LOCTEXT("Error_PhysicsControlAssetHasNoSkelMeshSkeleton", "Warning: Physics Control Profile Asset has a skeletal mesh with no skeleton assigned.\nFor now, a simple default skeletal mesh ({0}) will be used.\nYou can fix this by opening the asset and choosing another skeletal mesh from the toolbar, or repairing the skeleton."),
+			FText::FromString(PreviewMesh->GetFullName())));
+	}
+}
+
+
+#undef LOCTEXT_NAMESPACE
diff --git a/Engine/Plugins/Experimental/PhysicsControl/Source/PhysicsControlEditor/Private/PhysicsControlAssetEditorData.h b/Engine/Plugins/Experimental/PhysicsControl/Source/PhysicsControlEditor/Private/PhysicsControlAssetEditorData.h
new file mode 100644
--- /dev/null
+++ b/Engine/Plugins/Experimental/PhysicsControl/Source/PhysicsControlEditor/Private/PhysicsControlAssetEditorData.h
@@ -0,0 +1,39 @@
+// Copyright Epic Games, Inc. All Rights Reserved.
+
+#pragma once
+
+#include "CoreMinimal.h"
+#include "UObject/ObjectPtr.h"
+
+class IPersonaPreviewScene;
+class UPhysicsControlAsset;
+class UPhysicsControlAssetEditorSkeletalMeshComponent;
+
+/**
+ * Helper/container for data used by the Physics Control Profile Editor 
+ */
+class FPhysicsControlAssetEditorData
+{
+public:
+	FPhysicsControlAssetEditorData();
+
+	/** Initializes members */
+	void Initialize(const TSharedRef<IPersonaPreviewScene>& InPreviewScene);
+
+	/** Caches a preview mesh. Sets us to a default mesh if none is set yet (or if an older one got deleted) */
+	void CachePreviewMesh();
+
+public:
+	/** The asset being inspected */
+	TObjectPtr<UPhysicsControlAsset> PhysicsControlAsset;
+
+	/** Skeletal mesh component specialized for this asset editor */
+	UPhysicsControlAssetEditorSkeletalMeshComponent* EditorSkelComp;
+
+	/** The physics control component used for testing/simulating on the character */
+	class UPhysicsControlComponent* PhysicsControlComponent;
+
+	/** Preview scene */
+	TWeakPtr<IPersonaPreviewScene> PreviewScene;
+
+};
diff --git a/Engine/Plugins/Experimental/PhysicsControl/Source/PhysicsControlEditor/Private/PhysicsControlAssetEditorMode.cpp b/Engine/Plugins/Experimental/PhysicsControl/Source/PhysicsControlEditor/Private/PhysicsControlAssetEditorMode.cpp
new file mode 100644
--- /dev/null
+++ b/Engine/Plugins/Experimental/PhysicsControl/Source/PhysicsControlEditor/Private/PhysicsControlAssetEditorMode.cpp
@@ -0,0 +1,38 @@
+// Copyright Epic Games, Inc. All Rights Reserved.
+
+#include "PhysicsControlAssetEditorMode.h"
+#include "PhysicsControlEditorModule.h"
+#include "PhysicsControlAssetEditorToolkit.h"
+#include "AssetEditorModeManager.h"
+
+FName FPhysicsControlAssetEditorMode::ModeName("PhysicsControlAssetEditMode");
+
+//======================================================================================================================
+bool FPhysicsControlAssetEditorMode::GetCameraTarget(FSphere& OutTarget) const
+{
+	return false;
+}
+
+//======================================================================================================================
+IPersonaPreviewScene& FPhysicsControlAssetEditorMode::GetAnimPreviewScene() const
+{
+	return *static_cast<IPersonaPreviewScene*>(static_cast<FAssetEditorModeManager*>(Owner)->GetPreviewScene());
+}
+
+//======================================================================================================================
+void FPhysicsControlAssetEditorMode::Render(const FSceneView* View, FViewport* Viewport, FPrimitiveDrawInterface* PDI)
+{
+	FEdMode::Render(View, Viewport, PDI);
+}
+
+//======================================================================================================================
+void FPhysicsControlAssetEditorMode::Tick(FEditorViewportClient* ViewportClient, float DeltaTime)
+{
+	FEdMode::Tick(ViewportClient, DeltaTime);
+}
+
+//======================================================================================================================
+void FPhysicsControlAssetEditorMode::DrawHUD(FEditorViewportClient* ViewportClient, FViewport* Viewport, const FSceneView* View, FCanvas* Canvas)
+{
+	FEdMode::DrawHUD(ViewportClient, Viewport, View, Canvas);
+}
diff --git a/Engine/Plugins/Experimental/PhysicsControl/Source/PhysicsControlEditor/Private/PhysicsControlAssetEditorToolkit.cpp b/Engine/Plugins/Experimental/PhysicsControl/Source/PhysicsControlEditor/Private/PhysicsControlAssetEditorToolkit.cpp
new file mode 100644
--- /dev/null
+++ b/Engine/Plugins/Experimental/PhysicsControl/Source/PhysicsControlEditor/Private/PhysicsControlAssetEditorToolkit.cpp
@@ -0,0 +1,222 @@
+// Copyright Epic Games, Inc. All Rights Reserved.
+
+#include "PhysicsControlAssetEditorToolkit.h"
+
+#include "AnimationEditorPreviewActor.h"
+#include "EditorModeManager.h"
+#include "IDetailsView.h"
+#include "IPersonaToolkit.h"
+#include "Modules/ModuleManager.h"
+#include "PersonaModule.h"
+#include "PhysicsControlAssetApplicationMode.h"
+#include "PhysicsControlAsset.h"
+#include "PhysicsControlAssetEditorMode.h"
+#include "Animation/DebugSkelMeshComponent.h"
+#include "AnimPreviewInstance.h"
+#include "PhysicsControlAssetEditorData.h"
+
+#define LOCTEXT_NAMESPACE "PhysicsControlAssetEditorToolkit"
+
+const FName PhysicsControlAssetEditorModes::Editor("PhysicsControlAssetEditorMode");
+const FName PhysicsControlAssetEditorAppName = FName(TEXT("PhysicsControlAssetEditorApp"));
+
+//======================================================================================================================
+void FPhysicsControlAssetEditorToolkit::InitAssetEditor(
+	const EToolkitMode::Type        Mode,
+	const TSharedPtr<IToolkitHost>& InitToolkitHost,
+	UPhysicsControlAsset*    InPhysicsControlAsset)
+{
+	bIsInitialized = false;
+
+	// Initialise EditorData
+	{
+		EditorData = MakeShared<FPhysicsControlAssetEditorData>();
+		EditorData->PhysicsControlAsset = InPhysicsControlAsset;
+		EditorData->CachePreviewMesh();
+	}
+
+	// Create Persona toolkit
+	{
+		FPersonaToolkitArgs PersonaToolkitArgs;
+		PersonaToolkitArgs.OnPreviewSceneCreated = FOnPreviewSceneCreated::FDelegate::CreateSP(
+			this, &FPhysicsControlAssetEditorToolkit::HandlePreviewSceneCreated);
+		FPersonaModule& PersonaModule = FModuleManager::LoadModuleChecked<FPersonaModule>("Persona");
+		PersonaToolkit = PersonaModule.CreatePersonaToolkit(InPhysicsControlAsset, PersonaToolkitArgs);
+		PersonaModule.RecordAssetOpened(FAssetData(InPhysicsControlAsset));
+	}
+
+	// Note - we might want to make a custom skeleton tree view here, based on showing either the
+	// animation or the physics bones/skeleton. See FPhysicsAssetEditor::InitPhysicsAssetEditor
+
+	GEditor->RegisterForUndo(this);
+
+	// Initialise the asset editor
+	const bool bCreateDefaultStandaloneMenu = true;
+	const bool bCreateDefaultToolbar = true;
+	FAssetEditorToolkit::InitAssetEditor(
+		Mode,
+		InitToolkitHost,
+		PhysicsControlAssetEditorAppName,
+		FTabManager::FLayout::NullLayout,
+		bCreateDefaultStandaloneMenu,
+		bCreateDefaultToolbar,
+		InPhysicsControlAsset);
+
+	// Create and set the application mode.
+	ApplicationMode = new FPhysicsControlAssetApplicationMode(SharedThis(this), PersonaToolkit->GetPreviewScene());
+	AddApplicationMode(PhysicsControlAssetEditorModes::Editor, MakeShareable(ApplicationMode));
+	SetCurrentMode(PhysicsControlAssetEditorModes::Editor);
+
+	// Activate the editor mode.
+	GetEditorModeManager().SetDefaultMode(FPhysicsControlAssetEditorMode::ModeName);
+	GetEditorModeManager().ActivateMode(FPhysicsControlAssetEditorMode::ModeName);
+
+	FPhysicsControlAssetEditorMode* EditorMode = GetEditorModeManager().
+		GetActiveModeTyped<FPhysicsControlAssetEditorMode>(FPhysicsControlAssetEditorMode::ModeName);
+	EditorMode->SetEditorToolkit(this);
+
+	bIsInitialized = true;
+}
+
+
+//======================================================================================================================
+void FPhysicsControlAssetEditorToolkit::RegisterTabSpawners(const TSharedRef<FTabManager>& InTabManager)
+{
+	WorkspaceMenuCategory = InTabManager->AddLocalWorkspaceMenuCategory(LOCTEXT("WorkspaceMenu_MLDeformerEditor", "ML Deformer Editor"));
+	auto WorkspaceMenuCategoryRef = WorkspaceMenuCategory.ToSharedRef();
+
+	FAssetEditorToolkit::RegisterTabSpawners(InTabManager);
+}
+
+//======================================================================================================================
+void FPhysicsControlAssetEditorToolkit::UnregisterTabSpawners(const TSharedRef<FTabManager>& InTabManager)
+{
+	FAssetEditorToolkit::UnregisterTabSpawners(InTabManager);
+}
+
+//======================================================================================================================
+FName FPhysicsControlAssetEditorToolkit::GetToolkitFName() const
+{
+	return FName("PhysicsControlAssetEditor");
+}
+
+//======================================================================================================================
+FText FPhysicsControlAssetEditorToolkit::GetBaseToolkitName() const
+{
+	return LOCTEXT("PhysicsControlAssetEditorAppLabel", "Physics Control Asset Editor");
+}
+
+//======================================================================================================================
+FText FPhysicsControlAssetEditorToolkit::GetToolkitName() const
+{
+	FFormatNamedArguments Args;
+	Args.Add(TEXT("AssetName"), FText::FromString(EditorData->PhysicsControlAsset->GetName()));
+	return FText::Format(LOCTEXT("PhysicsControlAssetEditorToolkitName", "{AssetName}"), Args);
+}
+
+//======================================================================================================================
+FLinearColor FPhysicsControlAssetEditorToolkit::GetWorldCentricTabColorScale() const
+{
+	return FLinearColor::White;
+}
+
+//======================================================================================================================
+FString FPhysicsControlAssetEditorToolkit::GetWorldCentricTabPrefix() const
+{
+	return TEXT("PhysicsControlAssetEditor");
+}
+
+//======================================================================================================================
+void FPhysicsControlAssetEditorToolkit::AddReferencedObjects(FReferenceCollector& Collector)
+{
+	Collector.AddReferencedObject(EditorData->PhysicsControlAsset);
+}
+
+//======================================================================================================================
+TStatId FPhysicsControlAssetEditorToolkit::GetStatId() const
+{
+	RETURN_QUICK_DECLARE_CYCLE_STAT(FPhysicsControlAssetEditorToolkit, STATGROUP_Tickables);
+}
+
+//======================================================================================================================
+TSharedRef<IPersonaToolkit> FPhysicsControlAssetEditorToolkit::GetPersonaToolkit() const
+{
+	return PersonaToolkit.ToSharedRef();
+}
+
+//======================================================================================================================
+// For inspiration here, see:
+// - PhysicsAssetEditor
+// - MLDeformerEditorToolkit
+// - IKRigToolkit
+// though note that they all do things differently!
+void FPhysicsControlAssetEditorToolkit::HandlePreviewSceneCreated(
+	const TSharedRef<IPersonaPreviewScene>& InPersonaPreviewScene)
+{
+	EditorData->PreviewScene = InPersonaPreviewScene;
+
+	AAnimationEditorPreviewActor* Actor = InPersonaPreviewScene->GetWorld()->SpawnActor<AAnimationEditorPreviewActor>(
+		AAnimationEditorPreviewActor::StaticClass(), FTransform::Identity);
+	Actor->SetFlags(RF_Transient);
+	InPersonaPreviewScene->SetActor(Actor);
+
+	// Create the preview skeletal mesh component
+	ViewportSkeletalMeshComponent = NewObject<UDebugSkelMeshComponent>(Actor);
+	// turn off default bone rendering
+	//ViewportSkeletalMeshComponent->SkeletonDrawMode = ESkeletonDrawMode::Hidden;
+
+	// Setup and apply an anim instance to the skeletal mesh component
+	ViewportAnimInstance = NewObject<UAnimPreviewInstance>(
+		ViewportSkeletalMeshComponent, TEXT("PhysicsControlAssetEditorAnimInstance"));
+	ViewportSkeletalMeshComponent->PreviewInstance = ViewportAnimInstance;
+	//ViewportAnimInstance->InitializeAnimation();
+
+	// Set the skeletal mesh on the component, using the asset. Note that this will change if/when
+	// the asset doesn't hold a mesh.
+	USkeletalMesh* Mesh = EditorData->PhysicsControlAsset->PreviewSkeletalMesh.Get();
+	ViewportSkeletalMeshComponent->SetSkeletalMesh(Mesh);
+
+	// apply mesh to the preview scene
+	InPersonaPreviewScene->SetPreviewMeshComponent(ViewportSkeletalMeshComponent);
+	InPersonaPreviewScene->SetAllowMeshHitProxies(false);
+	InPersonaPreviewScene->SetAdditionalMeshesSelectable(false);
+	ViewportSkeletalMeshComponent->bSelectable = false;
+	InPersonaPreviewScene->SetPreviewMesh(Mesh);
+	InPersonaPreviewScene->AddComponent(ViewportSkeletalMeshComponent, FTransform::Identity);
+}
+
+//======================================================================================================================
+void FPhysicsControlAssetEditorToolkit::HandleViewportCreated(const TSharedRef<IPersonaViewport>& InPersonaViewport)
+{
+	PersonaViewport = InPersonaViewport;
+}
+
+//======================================================================================================================
+void FPhysicsControlAssetEditorToolkit::ShowEmptyDetails() const
+{
+	DetailsView->SetObject(EditorData->PhysicsControlAsset);
+}
+
+//======================================================================================================================
+void FPhysicsControlAssetEditorToolkit::HandleDetailsCreated(const TSharedRef<class IDetailsView>& InDetailsView)
+{
+	DetailsView = InDetailsView;
+	DetailsView->OnFinishedChangingProperties().AddSP(
+		this, &FPhysicsControlAssetEditorToolkit::OnFinishedChangingDetails);
+	ShowEmptyDetails();
+}
+
+//======================================================================================================================
+void FPhysicsControlAssetEditorToolkit::OnFinishedChangingDetails(const FPropertyChangedEvent& PropertyChangedEvent)
+{
+	const bool bPreviewMeshChanged = PropertyChangedEvent.GetPropertyName() == UPhysicsControlAsset::GetPreviewMeshPropertyName();
+	if (bPreviewMeshChanged)
+	{
+		USkeletalMesh* Mesh = EditorData->PhysicsControlAsset->PreviewSkeletalMesh.LoadSynchronous();
+		ViewportSkeletalMeshComponent->SetSkeletalMesh(Mesh);
+		EditorData->CachePreviewMesh();
+	}
+}
+
+#undef LOCTEXT_NAMESPACE
+
diff --git a/Engine/Plugins/Experimental/PhysicsControl/Source/PhysicsControlEditor/Private/PhysicsControlAssetFactory.cpp b/Engine/Plugins/Experimental/PhysicsControl/Source/PhysicsControlEditor/Private/PhysicsControlAssetFactory.cpp
new file mode 100644
--- /dev/null
+++ b/Engine/Plugins/Experimental/PhysicsControl/Source/PhysicsControlEditor/Private/PhysicsControlAssetFactory.cpp
@@ -0,0 +1,20 @@
+// Copyright Epic Games, Inc. All Rights Reserved.
+
+#include "PhysicsControlAssetFactory.h"
+#include "PhysicsControlAsset.h"
+
+//======================================================================================================================
+UPhysicsControlAssetFactory::UPhysicsControlAssetFactory(const FObjectInitializer& ObjectInitializer)
+	: Super(ObjectInitializer)
+{
+	SupportedClass = UPhysicsControlAsset::StaticClass();
+	bCreateNew = true;
+	bEditAfterNew = false;
+}
+
+//======================================================================================================================
+UObject* UPhysicsControlAssetFactory::FactoryCreateNew(
+	UClass* Class, UObject* InParent, FName Name, EObjectFlags Flags, UObject* Context, FFeedbackContext* Warn)
+{
+	return NewObject<UPhysicsControlAsset>(InParent, Class, Name, Flags, Context);
+}
diff --git a/Engine/Plugins/Experimental/PhysicsControl/Source/PhysicsControlEditor/Private/PhysicsControlAssetFactory.h b/Engine/Plugins/Experimental/PhysicsControl/Source/PhysicsControlEditor/Private/PhysicsControlAssetFactory.h
new file mode 100644
--- /dev/null
+++ b/Engine/Plugins/Experimental/PhysicsControl/Source/PhysicsControlEditor/Private/PhysicsControlAssetFactory.h
@@ -0,0 +1,17 @@
+// Copyright Epic Games, Inc. All Rights Reserved.
+
+#pragma once
+
+#include "CoreMinimal.h"
+#include "Factories/Factory.h"
+#include "PhysicsControlAssetFactory.generated.h"
+
+UCLASS()
+class UPhysicsControlAssetFactory : public UFactory
+{
+	GENERATED_BODY()
+public:
+	UPhysicsControlAssetFactory(const FObjectInitializer& ObjectInitializer);
+	UObject* FactoryCreateNew(
+		UClass* Class, UObject* InParent, FName Name, EObjectFlags Flags, UObject* Context, FFeedbackContext* Warn);
+};
diff --git a/Engine/Plugins/Experimental/PhysicsControl/Source/PhysicsControlEditor/Private/PhysicsControlEditorModule.cpp b/Engine/Plugins/Experimental/PhysicsControl/Source/PhysicsControlEditor/Private/PhysicsControlEditorModule.cpp
--- a/Engine/Plugins/Experimental/PhysicsControl/Source/PhysicsControlEditor/Private/PhysicsControlEditorModule.cpp
+++ b/Engine/Plugins/Experimental/PhysicsControl/Source/PhysicsControlEditor/Private/PhysicsControlEditorModule.cpp
@@ -4,8 +4,8 @@
 #include "PhysicsControlComponent.h"
 #include "PhysicsControlComponentVisualizer.h"
 #include "OperatorEditor/OperatorEditor.h"
-#include "PhysicsControlProfileAssetActions.h"
-#include "PhysicsControlProfileEditorMode.h"
+#include "PhysicsControlAssetActions.h"
+#include "PhysicsControlAssetEditorMode.h"
 #include "UnrealEdGlobals.h"
 #include "Editor/UnrealEdEngine.h"
 #include "Features/IModularFeatures.h"
@@ -18,12 +18,12 @@ static const FName PhysicsControlEditorModule_PhysicsControlEditorInterface("Phy
 void FPhysicsControlEditorModule::StartupModule()
 {
 	IAssetTools& AssetTools = FModuleManager::LoadModuleChecked<FAssetToolsModule>("AssetTools").Get();
-	PhysicsControlProfileAssetActions = MakeShared<FPhysicsControlProfileAssetActions>();
-	AssetTools.RegisterAssetTypeActions(PhysicsControlProfileAssetActions.ToSharedRef());
+	PhysicsControlAssetActions = MakeShared<FPhysicsControlAssetActions>();
+	AssetTools.RegisterAssetTypeActions(PhysicsControlAssetActions.ToSharedRef());
 
-	FEditorModeRegistry::Get().RegisterMode<FPhysicsControlProfileEditorMode>(
-		FPhysicsControlProfileEditorMode::ModeName, 
-		LOCTEXT("PhysicsControlProfileEditorMode", "PhysicsControlProfile"), 
+	FEditorModeRegistry::Get().RegisterMode<FPhysicsControlAssetEditorMode>(
+		FPhysicsControlAssetEditorMode::ModeName, 
+		LOCTEXT("PhysicsControlAssetEditorMode", "PhysicsControlAsset"), 
 		FSlateIcon(), false);
 
 	if (GUnrealEd)
@@ -60,12 +60,12 @@ void FPhysicsControlEditorModule::ShutdownModule()
 	}
 	
 	// Physics Control Profile editor/asset is disabled for now
-	FEditorModeRegistry::Get().UnregisterMode(FPhysicsControlProfileEditorMode::ModeName);
+	FEditorModeRegistry::Get().UnregisterMode(FPhysicsControlAssetEditorMode::ModeName);
 
 	if (FModuleManager::Get().IsModuleLoaded("AssetTools"))
 	{
 		FAssetToolsModule::GetModule().Get().UnregisterAssetTypeActions(
-			PhysicsControlProfileAssetActions.ToSharedRef());
+			PhysicsControlAssetActions.ToSharedRef());
 	}
 
 	if (GUnrealEd)
diff --git a/Engine/Plugins/Experimental/PhysicsControl/Source/PhysicsControlEditor/Private/PhysicsControlProfileApplicationMode.cpp b/Engine/Plugins/Experimental/PhysicsControl/Source/PhysicsControlEditor/Private/PhysicsControlProfileApplicationMode.cpp
deleted file mode 100644
diff --git a/Engine/Plugins/Experimental/PhysicsControl/Source/PhysicsControlEditor/Private/PhysicsControlProfileAssetActions.cpp b/Engine/Plugins/Experimental/PhysicsControl/Source/PhysicsControlEditor/Private/PhysicsControlProfileAssetActions.cpp
deleted file mode 100644
diff --git a/Engine/Plugins/Experimental/PhysicsControl/Source/PhysicsControlEditor/Private/PhysicsControlProfileAssetActions.h b/Engine/Plugins/Experimental/PhysicsControl/Source/PhysicsControlEditor/Private/PhysicsControlProfileAssetActions.h
deleted file mode 100644
diff --git a/Engine/Plugins/Experimental/PhysicsControl/Source/PhysicsControlEditor/Private/PhysicsControlProfileAssetFactory.cpp b/Engine/Plugins/Experimental/PhysicsControl/Source/PhysicsControlEditor/Private/PhysicsControlProfileAssetFactory.cpp
deleted file mode 100644
diff --git a/Engine/Plugins/Experimental/PhysicsControl/Source/PhysicsControlEditor/Private/PhysicsControlProfileAssetFactory.h b/Engine/Plugins/Experimental/PhysicsControl/Source/PhysicsControlEditor/Private/PhysicsControlProfileAssetFactory.h
deleted file mode 100644
diff --git a/Engine/Plugins/Experimental/PhysicsControl/Source/PhysicsControlEditor/Private/PhysicsControlProfileEditorData.cpp b/Engine/Plugins/Experimental/PhysicsControl/Source/PhysicsControlEditor/Private/PhysicsControlProfileEditorData.cpp
deleted file mode 100644
diff --git a/Engine/Plugins/Experimental/PhysicsControl/Source/PhysicsControlEditor/Private/PhysicsControlProfileEditorData.h b/Engine/Plugins/Experimental/PhysicsControl/Source/PhysicsControlEditor/Private/PhysicsControlProfileEditorData.h
deleted file mode 100644
diff --git a/Engine/Plugins/Experimental/PhysicsControl/Source/PhysicsControlEditor/Private/PhysicsControlProfileEditorMode.cpp b/Engine/Plugins/Experimental/PhysicsControl/Source/PhysicsControlEditor/Private/PhysicsControlProfileEditorMode.cpp
deleted file mode 100644
diff --git a/Engine/Plugins/Experimental/PhysicsControl/Source/PhysicsControlEditor/Private/PhysicsControlProfileEditorToolkit.cpp b/Engine/Plugins/Experimental/PhysicsControl/Source/PhysicsControlEditor/Private/PhysicsControlProfileEditorToolkit.cpp
deleted file mode 100644
diff --git a/Engine/Plugins/Experimental/PhysicsControl/Source/PhysicsControlEditor/Public/PhysicsControlAssetApplicationMode.h b/Engine/Plugins/Experimental/PhysicsControl/Source/PhysicsControlEditor/Public/PhysicsControlAssetApplicationMode.h
new file mode 100644
--- /dev/null
+++ b/Engine/Plugins/Experimental/PhysicsControl/Source/PhysicsControlEditor/Public/PhysicsControlAssetApplicationMode.h
@@ -0,0 +1,37 @@
+// Copyright Epic Games, Inc. All Rights Reserved.
+
+#pragma once
+
+#include "WorkflowOrientedApp/ApplicationMode.h"
+#include "WorkflowOrientedApp/WorkflowTabManager.h"
+
+class FWorkflowCentricApplication;
+class IPersonaPreviewScene;
+
+class FPhysicsControlAssetEditorToolkit;
+
+/**
+ * The application mode for the Physics Control Profile Editor
+ * This defines the layout of the UI. It basically spawns all the tabs, such as the viewport, details panels, etc.
+ */
+class PHYSICSCONTROLEDITOR_API FPhysicsControlAssetApplicationMode : public FApplicationMode
+{
+public:
+	/** The name of this mode. */
+	static FName ModeName;
+
+	FPhysicsControlAssetApplicationMode(
+		TSharedRef<FWorkflowCentricApplication> InHostingApp, 
+		TSharedRef<IPersonaPreviewScene>        InPreviewScene);
+
+	// FApplicationMode overrides.
+	virtual void RegisterTabFactories(TSharedPtr<FTabManager> InTabManager) override;
+	// ~END FApplicationMode overrides.
+
+protected:
+	/** The hosting app. */
+	TWeakPtr<FPhysicsControlAssetEditorToolkit> EditorToolkit = nullptr;
+
+	/** The tab factories we support. */
+	FWorkflowAllowedTabSet TabFactories;
+};
diff --git a/Engine/Plugins/Experimental/PhysicsControl/Source/PhysicsControlEditor/Public/PhysicsControlAssetEditorMode.h b/Engine/Plugins/Experimental/PhysicsControl/Source/PhysicsControlEditor/Public/PhysicsControlAssetEditorMode.h
new file mode 100644
--- /dev/null
+++ b/Engine/Plugins/Experimental/PhysicsControl/Source/PhysicsControlEditor/Public/PhysicsControlAssetEditorMode.h
@@ -0,0 +1,40 @@
+// Copyright Epic Games, Inc. All Rights Reserved.
+
+#pragma once
+
+#include "CoreMinimal.h"
+#include "IPersonaEditMode.h"
+
+class FPhysicsControlAssetEditorToolkit;
+
+/**
+ * The Physics Control Profile Persona editor mode.
+ */
+class PHYSICSCONTROLEDITOR_API FPhysicsControlAssetEditorMode	: public IPersonaEditMode
+{
+public:
+	/** The name of the mode. */
+	static FName ModeName;
+
+	void SetEditorToolkit(FPhysicsControlAssetEditorToolkit* InToolkit) { PhysicsControlAssetEditorToolkit = InToolkit; }
+
+	// IPersonaEditMode overrides.
+	virtual bool GetCameraTarget(FSphere& OutTarget) const override;
+	virtual class IPersonaPreviewScene& GetAnimPreviewScene() const override;
+	virtual void GetOnScreenDebugInfo(TArray<FText>& OutDebugInfo) const override {}
+	// ~END IPersonaEditMode overrides.
+
+	// FEdMode overrides.
+	virtual void Tick(FEditorViewportClient* ViewportClient, float DeltaTime) override;
+	virtual void Render(const FSceneView* View, FViewport* Viewport, FPrimitiveDrawInterface* PDI) override;
+	virtual void DrawHUD(FEditorViewportClient* ViewportClient, FViewport* Viewport, const FSceneView* View, FCanvas* Canvas) override;
+	virtual bool IsCompatibleWith(FEditorModeID OtherModeID) const override { return true; }
+	virtual bool AllowWidgetMove() override { return false; }
+	virtual bool ShouldDrawWidget() const override { return false; }
+	virtual bool UsesTransformWidget() const override { return false; }
+	virtual bool UsesTransformWidget(UE::Widget::EWidgetMode CheckMode) const override { return false; }
+	// ~END FEdMode overrides.
+
+protected:
+	FPhysicsControlAssetEditorToolkit* PhysicsControlAssetEditorToolkit = nullptr;
+};
diff --git a/Engine/Plugins/Experimental/PhysicsControl/Source/PhysicsControlEditor/Public/PhysicsControlAssetEditorToolkit.h b/Engine/Plugins/Experimental/PhysicsControl/Source/PhysicsControlEditor/Public/PhysicsControlAssetEditorToolkit.h
new file mode 100644
--- /dev/null
+++ b/Engine/Plugins/Experimental/PhysicsControl/Source/PhysicsControlEditor/Public/PhysicsControlAssetEditorToolkit.h
@@ -0,0 +1,102 @@
+// Copyright Epic Games, Inc. All Rights Reserved.
+
+#pragma once
+
+#include "CoreMinimal.h"
+#include "EditorUndoClient.h"
+#include "IHasPersonaToolkit.h"
+#include "IPersonaPreviewScene.h"
+#include "IPersonaViewport.h"
+#include "PersonaAssetEditorToolkit.h"
+
+class FPhysicsControlAssetApplicationMode;
+class FPhysicsControlAssetEditorData;
+class UAnimPreviewInstance;
+class UPhysicsControlAsset;
+
+namespace PhysicsControlAssetEditorModes
+{
+	extern const FName Editor;
+}
+
+/**
+ * The main toolkit/editor for working with Physics Control Profile assets
+ */
+class PHYSICSCONTROLEDITOR_API FPhysicsControlAssetEditorToolkit :
+	public FPersonaAssetEditorToolkit,
+	public IHasPersonaToolkit,
+	public FGCObject,
+	public FEditorUndoClient,
+	public FTickableEditorObject
+{
+public:
+	friend class FPhysicsControlAssetApplicationMode;
+
+public:
+
+	/** Initialize the asset editor. This will register the application mode, init the preview scene, etc. */
+	void InitAssetEditor(
+		const EToolkitMode::Type        Mode,
+		const TSharedPtr<IToolkitHost>& InitToolkitHost,
+		UPhysicsControlAsset*    InPhysicsControlAsset);
+
+	// FAssetEditorToolkit overrides.
+	virtual void RegisterTabSpawners(const TSharedRef<FTabManager>& InTabManager) override;
+	virtual void UnregisterTabSpawners(const TSharedRef<FTabManager>& InTabManager) override;
+	virtual FName GetToolkitFName() const override;
+	virtual FText GetBaseToolkitName() const override;
+	virtual FText GetToolkitName() const override;
+	virtual FLinearColor GetWorldCentricTabColorScale() const override;
+	virtual FString GetWorldCentricTabPrefix() const override;
+	// ~END FAssetEditorToolkit overrides.
+
+	// FGCObject overrides.
+	virtual FString GetReferencerName() const override { return TEXT("FPhysicsControlAssetEditorToolkit"); }
+	virtual void AddReferencedObjects(FReferenceCollector& Collector) override;
+	// ~END FGCObject overrides.
+
+	// FTickableEditorObject overrides.
+	virtual void Tick(float DeltaTime) override {};
+	virtual ETickableTickType GetTickableTickType() const override { return ETickableTickType::Always; }
+	virtual TStatId GetStatId() const override;
+	// ~END FTickableEditorObject overrides.
+
+	// IHasPersonaToolkit overrides.
+	virtual TSharedRef<IPersonaToolkit> GetPersonaToolkit() const override;
+	// ~END IHasPersonaToolkit overrides.
+
+	IPersonaToolkit* GetPersonaToolkitPointer() const { return PersonaToolkit.Get(); }
+
+private:
+	/** Preview scene setup. */
+	void HandlePreviewSceneCreated(const TSharedRef<IPersonaPreviewScene>& InPersonaPreviewScene);
+	void HandleViewportCreated(const TSharedRef<IPersonaViewport>& InPersonaViewport);
+	void HandleDetailsCreated(const TSharedRef<class IDetailsView>& InDetailsView);
+	void OnFinishedChangingDetails(const FPropertyChangedEvent& PropertyChangedEvent);
+	void ShowEmptyDetails() const;
+
+private:
+	/** The persona toolkit. */
+	TSharedPtr<IPersonaToolkit> PersonaToolkit = nullptr;
+
+	// Persona viewport.
+	TSharedPtr<IPersonaViewport> PersonaViewport = nullptr;
+
+	/** Data and methods shared across multiple classes */
+	TSharedPtr<FPhysicsControlAssetEditorData> EditorData;
+
+	// Asset properties tab 
+	TSharedPtr<IDetailsView> DetailsView;
+
+	FPhysicsControlAssetApplicationMode* ApplicationMode = nullptr;
+
+	// viewport anim instance 
+	UPROPERTY(transient, NonTransactional)
+	TObjectPtr<UAnimPreviewInstance> ViewportAnimInstance;
+
+	// viewport skeletal mesh 
+	UDebugSkelMeshComponent* ViewportSkeletalMeshComponent;
+
+	/** Has the asset editor been initialized? */
+	bool bIsInitialized = false;
+};
diff --git a/Engine/Plugins/Experimental/PhysicsControl/Source/PhysicsControlEditor/Public/PhysicsControlEditorModule.h b/Engine/Plugins/Experimental/PhysicsControl/Source/PhysicsControlEditor/Public/PhysicsControlEditorModule.h
--- a/Engine/Plugins/Experimental/PhysicsControl/Source/PhysicsControlEditor/Public/PhysicsControlEditorModule.h
+++ b/Engine/Plugins/Experimental/PhysicsControl/Source/PhysicsControlEditor/Public/PhysicsControlEditorModule.h
@@ -8,7 +8,7 @@
 
 
 class FPhysicsControlOperatorEditor;
-class FPhysicsControlProfileAssetActions;
+class FPhysicsControlAssetActions;
 
 /**
  * Module handles the Physics Control Profile editor, and also the Physics Control Visualizer
@@ -24,7 +24,7 @@ public:
 private:
 
 	TArray<FName> VisualizersToUnregisterOnShutdown;
-	TSharedPtr<FPhysicsControlProfileAssetActions> PhysicsControlProfileAssetActions;
+	TSharedPtr<FPhysicsControlAssetActions> PhysicsControlAssetActions;
 
 	FPhysicsControlOperatorEditor* EditorInterface = nullptr;
 };
diff --git a/Engine/Plugins/Experimental/PhysicsControl/Source/PhysicsControlEditor/Public/PhysicsControlProfileApplicationMode.h b/Engine/Plugins/Experimental/PhysicsControl/Source/PhysicsControlEditor/Public/PhysicsControlProfileApplicationMode.h
deleted file mode 100644
diff --git a/Engine/Plugins/Experimental/PhysicsControl/Source/PhysicsControlEditor/Public/PhysicsControlProfileEditorMode.h b/Engine/Plugins/Experimental/PhysicsControl/Source/PhysicsControlEditor/Public/PhysicsControlProfileEditorMode.h
deleted file mode 100644
diff --git a/Engine/Plugins/Experimental/PhysicsControl/Source/PhysicsControlEditor/Public/PhysicsControlProfileEditorToolkit.h b/Engine/Plugins/Experimental/PhysicsControl/Source/PhysicsControlEditor/Public/PhysicsControlProfileEditorToolkit.h
deleted file mode 100644
diff --git a/Engine/Plugins/Experimental/PhysicsControl/Source/PhysicsControlUncookedOnly/Private/AnimGraphNode_RigidBodyWithControl.cpp b/Engine/Plugins/Experimental/PhysicsControl/Source/PhysicsControlUncookedOnly/Private/AnimGraphNode_RigidBodyWithControl.cpp
--- a/Engine/Plugins/Experimental/PhysicsControl/Source/PhysicsControlUncookedOnly/Private/AnimGraphNode_RigidBodyWithControl.cpp
+++ b/Engine/Plugins/Experimental/PhysicsControl/Source/PhysicsControlUncookedOnly/Private/AnimGraphNode_RigidBodyWithControl.cpp
@@ -10,7 +10,7 @@
 #include "IPhysicsControlOperatorEditorInterface.h"
 #include "Kismet2/BlueprintEditorUtils.h"
 #include "PhysicsControlOperatorNameGeneration.h"
-#include "PhysicsControlProfileAsset.h"
+#include "PhysicsControlAsset.h"
 
 // Details includes
 #include "PropertyHandle.h"
@@ -342,9 +342,9 @@ TArray<TPair<FName, TArray<FName>>> UAnimGraphNode_RigidBodyWithControl::Generat
 
 		// Note that controls can come from the setup data in the node and/or from a profile asset
 		FPhysicsControlCharacterSetupData SetupData;
-		if (IsValid(Node.PhysicsControlProfileAsset))
+		if (IsValid(Node.PhysicsControlAsset))
 		{
-			SetupData = Node.PhysicsControlProfileAsset->CharacterSetupData;
+			SetupData = Node.PhysicsControlAsset->CharacterSetupData;
 		}
 		if (Node.bEnableCharacterSetupData)
 		{
@@ -352,9 +352,9 @@ TArray<TPair<FName, TArray<FName>>> UAnimGraphNode_RigidBodyWithControl::Generat
 		}
 
 		FPhysicsControlAndBodyModifierCreationDatas AdditionalControlAndBodyModifierCreationDatas;
-		if (IsValid(Node.PhysicsControlProfileAsset))
+		if (IsValid(Node.PhysicsControlAsset))
 		{
-			AdditionalControlAndBodyModifierCreationDatas = Node.PhysicsControlProfileAsset->AdditionalControlsAndModifiers;
+			AdditionalControlAndBodyModifierCreationDatas = Node.PhysicsControlAsset->AdditionalControlsAndModifiers;
 		}
 		AdditionalControlAndBodyModifierCreationDatas += Node.AdditionalControlsAndBodyModifiers;
 
@@ -364,10 +364,10 @@ TArray<TPair<FName, TArray<FName>>> UAnimGraphNode_RigidBodyWithControl::Generat
 			AllLimbBones, RefSkeleton, Node.OverridePhysicsAsset.Get(), BodyModifierNames, ControlNames, NameRecords);
 
 		// Create any additional sets that have been requested
-		if (IsValid(Node.PhysicsControlProfileAsset))
+		if (IsValid(Node.PhysicsControlAsset))
 		{
 			UE::PhysicsControl::CreateAdditionalSets(
-				Node.PhysicsControlProfileAsset->AdditionalSets, BodyModifierNames, ControlNames, NameRecords);
+				Node.PhysicsControlAsset->AdditionalSets, BodyModifierNames, ControlNames, NameRecords);
 		}
 		UE::PhysicsControl::CreateAdditionalSets(Node.AdditionalSets, BodyModifierNames, ControlNames, NameRecords);
 
