Added guard when accessing FRigidBodyPoseData to prevent OOB and memory being stomped

#rnx

diff --git a/Engine/Plugins/Experimental/PhysicsControl/Source/PhysicsControl/Private/AnimNode_RigidBodyWithControl.cpp b/Engine/Plugins/Experimental/PhysicsControl/Source/PhysicsControl/Private/AnimNode_RigidBodyWithControl.cpp
--- a/Engine/Plugins/Experimental/PhysicsControl/Source/PhysicsControl/Private/AnimNode_RigidBodyWithControl.cpp
+++ b/Engine/Plugins/Experimental/PhysicsControl/Source/PhysicsControl/Private/AnimNode_RigidBodyWithControl.cpp
@@ -795,7 +795,10 @@ void FAnimNode_RigidBodyWithControl::EvaluateSkeletalControl_AnyThread(FComponen
 				if (!BodyAnimData[BodyIndex].bIsSimulated)
 				{
 					// Custom kinematic targets are set subsequently, overriding these.
-					Bodies[BodyIndex]->SetKinematicTarget(PoseData.GetTM(BodyIndex).ToTransform());
+					if (ensure(PoseData.IsValidIndex(BodyIndex)))
+					{
+						Bodies[BodyIndex]->SetKinematicTarget(PoseData.GetTM(BodyIndex).ToTransform());
+					}
 				}
 			}
 
diff --git a/Engine/Plugins/Experimental/PhysicsControl/Source/PhysicsControl/Private/RigidBodyPoseData.cpp b/Engine/Plugins/Experimental/PhysicsControl/Source/PhysicsControl/Private/RigidBodyPoseData.cpp
--- a/Engine/Plugins/Experimental/PhysicsControl/Source/PhysicsControl/Private/RigidBodyPoseData.cpp
+++ b/Engine/Plugins/Experimental/PhysicsControl/Source/PhysicsControl/Private/RigidBodyPoseData.cpp
@@ -31,9 +31,12 @@ void FRigidBodyPoseData::Update(
 		for (const FOutputBoneData& OutputData : OutputBoneData)
 		{
 			const int32 BodyIndex = OutputData.BodyIndex;
-			const FTransform& ComponentSpaceTM = ComponentSpacePoseContext.Pose.GetComponentSpaceTransform(OutputData.CompactPoseBoneIndex);
-			const FTransform BodyTM = ConvertCSTransformToSimSpace(SimulationSpace, ComponentSpaceTM, CompWorldSpaceTM, BaseBoneTM);
-			BoneTMs[BodyIndex] = BodyTM;
+			if (ensure(BoneTMs.IsValidIndex(BodyIndex)))
+			{
+				const FTransform& ComponentSpaceTM = ComponentSpacePoseContext.Pose.GetComponentSpaceTransform(OutputData.CompactPoseBoneIndex);
+				const FTransform BodyTM = ConvertCSTransformToSimSpace(SimulationSpace, ComponentSpaceTM, CompWorldSpaceTM, BaseBoneTM);
+				BoneTMs[BodyIndex] = BodyTM;
+			}
 		}
 	}
 	else
diff --git a/Engine/Plugins/Experimental/PhysicsControl/Source/PhysicsControl/Private/RigidBodyWithControl.cpp b/Engine/Plugins/Experimental/PhysicsControl/Source/PhysicsControl/Private/RigidBodyWithControl.cpp
--- a/Engine/Plugins/Experimental/PhysicsControl/Source/PhysicsControl/Private/RigidBodyWithControl.cpp
+++ b/Engine/Plugins/Experimental/PhysicsControl/Source/PhysicsControl/Private/RigidBodyWithControl.cpp
@@ -568,14 +568,25 @@ static UE::PhysicsControl::FPosQuat CalculateTargetTM(
 	const int32                                   ParentBodyIndex, 
 	const int32                                   ChildBodyIndex)
 {
+	if (!ensure(PoseData.IsValidIndex(ChildBodyIndex)))
+	{
+		return UE::PhysicsControl::FPosQuat(JointSettings.ConnectorTransforms[ConstraintChildIndex]);
+	}
+	
 	const UE::PhysicsControl::FPosQuat ChildTargetTM =
-		PoseData.GetTM(ChildBodyIndex) * 
+		PoseData.GetTM(ChildBodyIndex) *
 		UE::PhysicsControl::FPosQuat(JointSettings.ConnectorTransforms[ConstraintChildIndex]);
+	
 	if (ParentBodyIndex >= 0)
 	{
+		if (!ensure(PoseData.IsValidIndex(ParentBodyIndex)))
+		{
+			return ChildTargetTM;
+		}
+
 		const UE::PhysicsControl::FPosQuat ParentTargetTM =
-			PoseData.GetTM(ParentBodyIndex) * 
-			UE::PhysicsControl::FPosQuat(JointSettings.ConnectorTransforms[ConstraintParentIndex]);
+		PoseData.GetTM(ParentBodyIndex) *
+		UE::PhysicsControl::FPosQuat(JointSettings.ConnectorTransforms[ConstraintParentIndex]);
 		return ParentTargetTM.Inverse() * ChildTargetTM;
 	}
 	return ChildTargetTM;
@@ -886,7 +897,7 @@ void FAnimNode_RigidBodyWithControl::ApplyKinematicTargets()
 				if (ActorHandle->GetIsKinematic() && BodyIndex != INDEX_NONE)
 				{
 					UE::PhysicsControl::FPosQuat TM(Target.TargetOrientation, Target.TargetPosition);
-					if (Target.bUseSkeletalAnimation)
+					if (Target.bUseSkeletalAnimation && ensure(PoseData.IsValidIndex(BodyIndex)))
 					{
 						TM = PoseData.GetTM(BodyIndex) * TM;
 					}
