Fixes for RigidBodyWithControl looking up the wrong body TMs

diff --git a/Engine/Plugins/Experimental/PhysicsControl/Source/PhysicsControl/Private/AnimNode_RigidBodyWithControl.cpp b/Engine/Plugins/Experimental/PhysicsControl/Source/PhysicsControl/Private/AnimNode_RigidBodyWithControl.cpp
--- a/Engine/Plugins/Experimental/PhysicsControl/Source/PhysicsControl/Private/AnimNode_RigidBodyWithControl.cpp
+++ b/Engine/Plugins/Experimental/PhysicsControl/Source/PhysicsControl/Private/AnimNode_RigidBodyWithControl.cpp
@@ -36,8 +36,6 @@ LLM_DEFINE_TAG(Animation_RigidBodyWithControl);
 /////////////////////////////////////////////////////
 // FAnimNode_RigidBodyWithControl
 
-//UE_DISABLE_OPTIMIZATION
-
 #define LOCTEXT_NAMESPACE "ImmediatePhysicsWithControl"
 
 DECLARE_CYCLE_STAT(TEXT("RigidBodyNodeWithControl_InitPhysics"), STAT_RigidBodyWithControlInitPhysicsTime, STATGROUP_Anim);
@@ -857,16 +855,19 @@ void FAnimNode_RigidBodyWithControl::EvaluateSkeletalControl_AnyThread(FComponen
 			const int32 BodyIndex = OutputData.BodyIndex;
 			// Note that we always read back, whether kinematic or simulated
 			FTransform BodyTM = Bodies[BodyIndex]->GetWorldTransform();
-			FTransform ComponentSpaceTM;
-			switch(SimulationSpace)
+			if (ensure(!BodyTM.ContainsNaN()))
 			{
-				case ESimulationSpace::ComponentSpace: ComponentSpaceTM = BodyTM; break;
-				case ESimulationSpace::WorldSpace: ComponentSpaceTM = BodyTM.GetRelativeTransform(SimulationWorldSpaceTM); break;
-				case ESimulationSpace::BaseBoneSpace: ComponentSpaceTM = BodyTM * BaseBoneTM; break;
-				default: ensureMsgf(false, TEXT("Unsupported Simulation Space")); ComponentSpaceTM = BodyTM;
-			}
+				FTransform ComponentSpaceTM;
+				switch(SimulationSpace)
+				{
+					case ESimulationSpace::ComponentSpace: ComponentSpaceTM = BodyTM; break;
+					case ESimulationSpace::WorldSpace: ComponentSpaceTM = BodyTM.GetRelativeTransform(SimulationWorldSpaceTM); break;
+					case ESimulationSpace::BaseBoneSpace: ComponentSpaceTM = BodyTM * BaseBoneTM; break;
+					default: ensureMsgf(false, TEXT("Unsupported Simulation Space")); ComponentSpaceTM = BodyTM;
+				}
 					
-			OutBoneTransforms.Add(FBoneTransform(OutputData.CompactPoseBoneIndex, ComponentSpaceTM));
+				OutBoneTransforms.Add(FBoneTransform(OutputData.CompactPoseBoneIndex, ComponentSpaceTM));
+			}
 		}
 
 		// Deferred task must be started after we read actor poses to avoid a race
@@ -1845,6 +1846,7 @@ void FAnimNode_RigidBodyWithControl::InitializeBoneReferences(const FBoneContain
 	const FReferenceSkeleton& RefSkeleton = RequiredBones.GetReferenceSkeleton();
 
 	OutputBoneData.Empty(NumBodies);
+	int32 OutputBoneDataMaxBodyIndex = -1;
 
 	int32 NumSimulatedBodies = 0;
 	TArray<int32> SimulatedBodyIndices;
@@ -1885,10 +1887,10 @@ void FAnimNode_RigidBodyWithControl::InitializeBoneReferences(const FBoneContain
 		if (BodyIndex != INDEX_NONE)
 		{
 			//If we have a body we need to save it for later
-			UE::PhysicsControl::FOutputBoneData* OutputData = 
-				new (OutputBoneData) UE::PhysicsControl::FOutputBoneData();
-			OutputData->BodyIndex = BodyIndex;
-			OutputData->CompactPoseBoneIndex = CompactPoseBoneIndex;
+			UE::PhysicsControl::FOutputBoneData& OutputData = OutputBoneData.AddDefaulted_GetRef();
+			OutputData.BodyIndex = BodyIndex;
+			OutputBoneDataMaxBodyIndex = FMath::Max(OutputBoneDataMaxBodyIndex, BodyIndex);
+			OutputData.CompactPoseBoneIndex = CompactPoseBoneIndex;
 
 			if (BodyAnimData[BodyIndex].bIsSimulated)
 			{
@@ -1896,12 +1898,12 @@ void FAnimNode_RigidBodyWithControl::InitializeBoneReferences(const FBoneContain
 				SimulatedBodyIndices.AddUnique(BodyIndex);
 			}
 
-			OutputData->BoneIndicesToParentBody.Add(CompactPoseBoneIndex);
+			OutputData.BoneIndicesToParentBody.Add(CompactPoseBoneIndex);
 
 			// Walk up parent chain until we find parent body.
-			OutputData->ParentBodyIndex = INDEX_NONE;
+			OutputData.ParentBodyIndex = INDEX_NONE;
 			FCompactPoseBoneIndex CompactParentIndex = RequiredBones.GetParentBoneIndex(CompactPoseBoneIndex);
-			OutputData->CompactPoseParentBoneIndex = CompactParentIndex;
+			OutputData.CompactPoseParentBoneIndex = CompactParentIndex;
 			while (CompactParentIndex != INDEX_NONE)
 			{
 				const FBoneIndexType SkeletonParentBoneIndex = (FBoneIndexType ) RequiredBones.GetSkeletonIndex(CompactParentIndex);
@@ -1913,14 +1915,14 @@ void FAnimNode_RigidBodyWithControl::InitializeBoneReferences(const FBoneContain
 					break;
 				}
 
-				OutputData->ParentBodyIndex = SkeletonBoneIndexToBodyIndex[SkeletonParentBoneIndex];
-				OutputData->CompactPoseParentBoneIndex = CompactParentIndex;
-				if (OutputData->ParentBodyIndex != INDEX_NONE)
+				OutputData.ParentBodyIndex = SkeletonBoneIndexToBodyIndex[SkeletonParentBoneIndex];
+				OutputData.CompactPoseParentBoneIndex = CompactParentIndex;
+				if (OutputData.ParentBodyIndex != INDEX_NONE)
 				{
 					break;
 				}
 
-				OutputData->BoneIndicesToParentBody.Add(CompactParentIndex);
+				OutputData.BoneIndicesToParentBody.Add(CompactParentIndex);
 				CompactParentIndex = RequiredBones.GetParentBoneIndex(CompactParentIndex);
 			}
 
@@ -1940,6 +1942,9 @@ void FAnimNode_RigidBodyWithControl::InitializeBoneReferences(const FBoneContain
 	}
 	else
 	{
+		// The Pose data cache will only be indexed using the body indices from OutputBoneData
+		PoseData.SetSize(OutputBoneDataMaxBodyIndex + 1);
+
 		// New bodies potentially introduced with new LOD
 		// We'll have to initialize their transform.
 		bCheckForBodyTransformInit = true;
diff --git a/Engine/Plugins/Experimental/PhysicsControl/Source/PhysicsControl/Private/PhysicsControlComponent.cpp b/Engine/Plugins/Experimental/PhysicsControl/Source/PhysicsControl/Private/PhysicsControlComponent.cpp
--- a/Engine/Plugins/Experimental/PhysicsControl/Source/PhysicsControl/Private/PhysicsControlComponent.cpp
+++ b/Engine/Plugins/Experimental/PhysicsControl/Source/PhysicsControl/Private/PhysicsControlComponent.cpp
@@ -25,8 +25,6 @@
 DECLARE_CYCLE_STAT(TEXT("PhysicsControl UpdateTargetCaches"), STAT_PhysicsControl_UpdateTargetCaches, STATGROUP_Anim);
 DECLARE_CYCLE_STAT(TEXT("PhysicsControl UpdateControls"), STAT_PhysicsControl_UpdateControls, STATGROUP_Anim);
 
-//UE_DISABLE_OPTIMIZATION;
-
 //======================================================================================================================
 // This file contains the public member functions of UPhysicsControlComponent
 //======================================================================================================================
diff --git a/Engine/Plugins/Experimental/PhysicsControl/Source/PhysicsControl/Private/PhysicsControlComponentImpl.cpp b/Engine/Plugins/Experimental/PhysicsControl/Source/PhysicsControl/Private/PhysicsControlComponentImpl.cpp
--- a/Engine/Plugins/Experimental/PhysicsControl/Source/PhysicsControl/Private/PhysicsControlComponentImpl.cpp
+++ b/Engine/Plugins/Experimental/PhysicsControl/Source/PhysicsControl/Private/PhysicsControlComponentImpl.cpp
@@ -15,8 +15,6 @@
 #include "Engine/SkeletalMesh.h"
 #include "Engine/World.h"
 
-//UE_DISABLE_OPTIMIZATION;
-
 //======================================================================================================================
 // This file contains the non-public member functions of UPhysicsControlComponent
 //======================================================================================================================
diff --git a/Engine/Plugins/Experimental/PhysicsControl/Source/PhysicsControl/Private/RigidBodyPoseData.cpp b/Engine/Plugins/Experimental/PhysicsControl/Source/PhysicsControl/Private/RigidBodyPoseData.cpp
--- a/Engine/Plugins/Experimental/PhysicsControl/Source/PhysicsControl/Private/RigidBodyPoseData.cpp
+++ b/Engine/Plugins/Experimental/PhysicsControl/Source/PhysicsControl/Private/RigidBodyPoseData.cpp
@@ -26,28 +26,14 @@ void FRigidBodyPoseData::Update(
 	const FTransform BaseBoneTM = ComponentSpacePoseContext.Pose.GetComponentSpaceTransform(
 		BaseBoneRef.GetCompactPoseIndex(BoneContainer));
 
-	if (BoneTMs.Num() == OutputBoneData.Num())
+	for (const FOutputBoneData& OutputData : OutputBoneData)
 	{
-		for (const FOutputBoneData& OutputData : OutputBoneData)
+		const int32 BodyIndex = OutputData.BodyIndex;
+		if (ensure(BoneTMs.IsValidIndex(BodyIndex)))
 		{
-			const int32 BodyIndex = OutputData.BodyIndex;
-			if (ensure(BoneTMs.IsValidIndex(BodyIndex)))
-			{
-				const FTransform& ComponentSpaceTM = ComponentSpacePoseContext.Pose.GetComponentSpaceTransform(OutputData.CompactPoseBoneIndex);
-				const FTransform BodyTM = ConvertCSTransformToSimSpace(SimulationSpace, ComponentSpaceTM, CompWorldSpaceTM, BaseBoneTM);
-				BoneTMs[BodyIndex] = BodyTM;
-			}
-		}
-	}
-	else
-	{
-		BoneTMs.Empty(OutputBoneData.Num());
-		for (const FOutputBoneData& OutputData : OutputBoneData)
-		{
-			const int32 BodyIndex = OutputData.BodyIndex;
 			const FTransform& ComponentSpaceTM = ComponentSpacePoseContext.Pose.GetComponentSpaceTransform(OutputData.CompactPoseBoneIndex);
 			const FTransform BodyTM = ConvertCSTransformToSimSpace(SimulationSpace, ComponentSpaceTM, CompWorldSpaceTM, BaseBoneTM);
-			BoneTMs.Emplace(BodyTM);
+			BoneTMs[BodyIndex] = BodyTM;
 		}
 	}
 }
diff --git a/Engine/Plugins/Experimental/PhysicsControl/Source/PhysicsControl/Private/RigidBodyWithControl.cpp b/Engine/Plugins/Experimental/PhysicsControl/Source/PhysicsControl/Private/RigidBodyWithControl.cpp
--- a/Engine/Plugins/Experimental/PhysicsControl/Source/PhysicsControl/Private/RigidBodyWithControl.cpp
+++ b/Engine/Plugins/Experimental/PhysicsControl/Source/PhysicsControl/Private/RigidBodyWithControl.cpp
@@ -21,8 +21,6 @@
 #include "Chaos/ChaosConstraintSettings.h"
 #include "Logging/StructuredLog.h"
 
-//UE_DISABLE_OPTIMIZATION
-
 DECLARE_CYCLE_STAT(TEXT("RigidBodyNodeWithControl_InitControlsAndBodyModifiers"), STAT_RigidBodyNodeWithControl_InitControlsAndBodyModifiers, STATGROUP_Anim);
 DECLARE_CYCLE_STAT(TEXT("RigidBodyNodeWithControl_LogControlsModifiersAndSets"), STAT_RigidBodyNodeWithControl_LogControlsModifiersAndSets, STATGROUP_Anim);
 DECLARE_CYCLE_STAT(TEXT("RigidBodyNodeWithControl_ApplyControlAndModifierUpdatesAndParametersToRecords"), STAT_RigidBodyNodeWithControl_ApplyControlAndModifierUpdatesAndParametersToRecords, STATGROUP_Anim);
diff --git a/Engine/Plugins/Experimental/PhysicsControl/Source/PhysicsControl/Public/PhysicsControlPoseData.h b/Engine/Plugins/Experimental/PhysicsControl/Source/PhysicsControl/Public/PhysicsControlPoseData.h
--- a/Engine/Plugins/Experimental/PhysicsControl/Source/PhysicsControl/Public/PhysicsControlPoseData.h
+++ b/Engine/Plugins/Experimental/PhysicsControl/Source/PhysicsControl/Public/PhysicsControlPoseData.h
@@ -60,6 +60,11 @@ struct FPosQuat
 		return FPosQuat(OutRotation * -Translation, OutRotation);
 	}
 
+	FORCEINLINE bool ContainsNaN() const
+	{
+		return Translation.ContainsNaN() || Rotation.ContainsNaN();
+	}
+
 	FVector Translation;
 	FQuat Rotation;
 };
diff --git a/Engine/Plugins/Experimental/PhysicsControl/Source/PhysicsControl/Public/RigidBodyPoseData.h b/Engine/Plugins/Experimental/PhysicsControl/Source/PhysicsControl/Public/RigidBodyPoseData.h
--- a/Engine/Plugins/Experimental/PhysicsControl/Source/PhysicsControl/Public/RigidBodyPoseData.h
+++ b/Engine/Plugins/Experimental/PhysicsControl/Source/PhysicsControl/Public/RigidBodyPoseData.h
@@ -44,10 +44,13 @@ public:
 		const FBoneReference&          BaseBoneRef,
 		const FGraphTraversalCounter&  InUpdateCounter);
 
-	UE::PhysicsControl::FPosQuat GetTM(int32 Index) const { return BoneTMs[Index]; }
+	UE::PhysicsControl::FPosQuat GetTM(int32 Index) const { 
+		check(IsValidIndex(Index)); check(!BoneTMs[Index].ContainsNaN()); return BoneTMs[Index]; }
 	bool IsValidIndex(const int32 Index) const { return BoneTMs.IsValidIndex(Index); }
 	bool IsEmpty() const { return BoneTMs.IsEmpty(); }
 
+	void SetSize(const int32 NumBones) { BoneTMs.SetNum(NumBones); }
+
 	/**
 	 * The cached skeletal data, updated at the start of each tick
 	 */
