Progress on making a functional PhysicsControlProfile editor

diff --git a/Engine/Plugins/Experimental/PhysicsControl/Source/PhysicsControl/Private/PhysicsControlAsset.cpp b/Engine/Plugins/Experimental/PhysicsControl/Source/PhysicsControl/Private/PhysicsControlAsset.cpp
--- a/Engine/Plugins/Experimental/PhysicsControl/Source/PhysicsControl/Private/PhysicsControlAsset.cpp
+++ b/Engine/Plugins/Experimental/PhysicsControl/Source/PhysicsControl/Private/PhysicsControlAsset.cpp
@@ -1,7 +1,9 @@
 // Copyright Epic Games, Inc. All Rights Reserved.
 
 #include "PhysicsControlAsset.h"
+
 #include "Engine/SkeletalMesh.h"
+#include "PhysicsEngine/PhysicsAsset.h"
 
 //#ifdef WITH_EDITOR
 //#include "Editor.h"
@@ -66,6 +68,32 @@ void UPhysicsControlAsset::Compile()
 	Modify();
 }
 
+//======================================================================================================================
+bool UPhysicsControlAsset::IsCompilationNeeded() const
+{
+	if (CharacterSetupData != GetCharacterSetupData())
+	{
+		return true;
+	}
+	if (AdditionalControlsAndModifiers != GetAdditionalControlsAndModifiers())
+	{
+		return true;
+	}
+	if (AdditionalSets != GetAdditionalSets())
+	{
+		return true;
+	}
+	if (InitialControlAndModifierUpdates != GetInitialControlAndModifierUpdates())
+	{
+		return true;
+	}
+	if (!Profiles.OrderIndependentCompareEqual(GetProfiles()))
+	{
+		return true;
+	}
+	return false;
+}
+
 //======================================================================================================================
 FPhysicsControlCharacterSetupData UPhysicsControlAsset::GetCharacterSetupData() const
 {
@@ -137,18 +165,34 @@ TMap<FName, FPhysicsControlControlAndModifierUpdates> UPhysicsControlAsset::GetP
 #endif
 
 #if WITH_EDITOR
+//======================================================================================================================
+UPhysicsAsset* UPhysicsControlAsset::GetPhysicsAsset() const
+{
+	return PhysicsAsset.LoadSynchronous();
+}
+
+//======================================================================================================================
+void UPhysicsControlAsset::SetPhysicsAsset(UPhysicsAsset* InPhysicsAsset)
+{
+	PhysicsAsset = InPhysicsAsset;
+}
+
 //======================================================================================================================
 const FName UPhysicsControlAsset::GetPreviewMeshPropertyName()
 {
-	return GET_MEMBER_NAME_STRING_CHECKED(UPhysicsControlAsset, PreviewSkeletalMesh);
+	return GET_MEMBER_NAME_STRING_CHECKED(UPhysicsAsset, PreviewSkeletalMesh);
 };
+
 #endif
 
 //======================================================================================================================
 void UPhysicsControlAsset::SetPreviewMesh(USkeletalMesh* PreviewMesh, bool bMarkAsDirty)
 {
 #if WITH_EDITOR
-	PreviewSkeletalMesh = PreviewMesh;
+	if (UPhysicsAsset* PA = GetPhysicsAsset())
+	{
+		PA->SetPreviewMesh(PreviewMesh, bMarkAsDirty);
+	}
 #endif
 }
 
@@ -156,8 +200,10 @@ void UPhysicsControlAsset::SetPreviewMesh(USkeletalMesh* PreviewMesh, bool bMark
 USkeletalMesh* UPhysicsControlAsset::GetPreviewMesh() const
 {
 #if WITH_EDITOR
-	return PreviewSkeletalMesh.LoadSynchronous();
-#else
-	return nullptr;
+	if (UPhysicsAsset* PA = GetPhysicsAsset())
+	{
+		return PA->GetPreviewMesh();
+	}
 #endif
+	return nullptr;
 }
diff --git a/Engine/Plugins/Experimental/PhysicsControl/Source/PhysicsControl/Private/PhysicsControlComponent.cpp b/Engine/Plugins/Experimental/PhysicsControl/Source/PhysicsControl/Private/PhysicsControlComponent.cpp
--- a/Engine/Plugins/Experimental/PhysicsControl/Source/PhysicsControl/Private/PhysicsControlComponent.cpp
+++ b/Engine/Plugins/Experimental/PhysicsControl/Source/PhysicsControl/Private/PhysicsControlComponent.cpp
@@ -70,18 +70,7 @@ void UPhysicsControlComponent::InitializeComponent()
 //======================================================================================================================
 void UPhysicsControlComponent::BeginDestroy()
 {
-	for (TPair<FName, FPhysicsControlRecord>& ControlRecordPair : ControlRecords)
-	{
-		DestroyControl(ControlRecordPair.Key, EDestroyBehavior::KeepRecord);
-	}
-	ControlRecords.Empty();
-
-	for (TPair<FName, FPhysicsBodyModifierRecord>& BodyModifierPair : BodyModifierRecords)
-	{
-		DestroyBodyModifier(BodyModifierPair.Key, EDestroyBehavior::KeepRecord);
-	}
-	BodyModifierRecords.Empty();
-
+	DestroyPhysicsState();
 	Super::BeginDestroy();
 }
 
@@ -2386,6 +2375,11 @@ TArray<FTransform> UPhysicsControlComponent::GetCachedBoneTransforms(
 		}
 		else
 		{
+			if (bWarnAboutInvalidNames)
+			{
+				UE_LOG(LogPhysicsControl, Warning,
+					TEXT("GetCachedBoneTransforms - unable to get bone data for %s"), *BoneName.ToString());
+			}
 			Result.Add(FTransform::Identity);
 		}
 	}
@@ -2410,6 +2404,11 @@ TArray<FVector> UPhysicsControlComponent::GetCachedBonePositions(
 		}
 		else
 		{
+			if (bWarnAboutInvalidNames)
+			{
+				UE_LOG(LogPhysicsControl, Warning,
+					TEXT("GetCachedBonePositions - unable to get bone data for %s"), *BoneName.ToString());
+			}
 			Result.Add(FVector::ZeroVector);
 		}
 	}
@@ -2434,6 +2433,11 @@ TArray<FRotator> UPhysicsControlComponent::GetCachedBoneOrientations(
 		}
 		else
 		{
+			if (bWarnAboutInvalidNames)
+			{
+				UE_LOG(LogPhysicsControl, Warning,
+					TEXT("GetCachedBoneOrientations - unable to get bone data for %s"), *BoneName.ToString());
+			}
 			Result.Add(FRotator::ZeroRotator);
 		}
 	}
@@ -2443,18 +2447,18 @@ TArray<FRotator> UPhysicsControlComponent::GetCachedBoneOrientations(
 //======================================================================================================================
 FTransform UPhysicsControlComponent::GetCachedBoneTransform(
 	const USkeletalMeshComponent* SkeletalMeshComponent,
-	const FName                   Name)
+	const FName                   BoneName)
 {
 	UE::PhysicsControl::FBoneData BoneData;
 	const UE::PhysicsControl::FPhysicsControlPoseData* PoseData;
-	if (GetBoneData(BoneData, PoseData, SkeletalMeshComponent, Name))
+	if (GetBoneData(BoneData, PoseData, SkeletalMeshComponent, BoneName))
 	{
 		return FTransform(BoneData.CurrentTM.GetRotation(), BoneData.CurrentTM.GetTranslation());
 	}
 	if (bWarnAboutInvalidNames)
 	{
 		UE_LOG(LogPhysicsControl, Warning,
-			TEXT("GetCachedBoneTransform - invalid name %s"), *Name.ToString());
+			TEXT("GetCachedBoneTransform - invalid bone name %s"), *BoneName.ToString());
 	}
 	return FTransform();
 }
@@ -2462,18 +2466,18 @@ FTransform UPhysicsControlComponent::GetCachedBoneTransform(
 //======================================================================================================================
 FVector UPhysicsControlComponent::GetCachedBonePosition(
 	const USkeletalMeshComponent* SkeletalMeshComponent,
-	const FName                   Name)
+	const FName                   BoneName)
 {
 	UE::PhysicsControl::FBoneData BoneData;
 	const UE::PhysicsControl::FPhysicsControlPoseData* PoseData;
-	if (GetBoneData(BoneData, PoseData, SkeletalMeshComponent, Name))
+	if (GetBoneData(BoneData, PoseData, SkeletalMeshComponent, BoneName))
 	{
 		return BoneData.CurrentTM.GetTranslation();
 	}
 	if (bWarnAboutInvalidNames)
 	{
 		UE_LOG(LogPhysicsControl, Warning,
-			TEXT("GetCachedBonePosition - invalid name %s"), *Name.ToString());
+			TEXT("GetCachedBonePosition - invalid bone name %s"), *BoneName.ToString());
 	}
 	return FVector::ZeroVector;
 }
@@ -2481,18 +2485,18 @@ FVector UPhysicsControlComponent::GetCachedBonePosition(
 //======================================================================================================================
 FRotator UPhysicsControlComponent::GetCachedBoneOrientation(
 	const USkeletalMeshComponent* SkeletalMeshComponent,
-	const FName                   Name)
+	const FName                   BoneName)
 {
 	UE::PhysicsControl::FBoneData BoneData;
 	const UE::PhysicsControl::FPhysicsControlPoseData* PoseData;
-	if (GetBoneData(BoneData, PoseData, SkeletalMeshComponent, Name))
+	if (GetBoneData(BoneData, PoseData, SkeletalMeshComponent, BoneName))
 	{
 		return BoneData.CurrentTM.GetRotation().Rotator();
 	}
 	if (bWarnAboutInvalidNames)
 	{
 		UE_LOG(LogPhysicsControl, Warning,
-			TEXT("GetCachedBoneOrientation - invalid name %s"), *Name.ToString());
+			TEXT("GetCachedBoneOrientation - invalid bone name %s"), *BoneName.ToString());
 	}
 	return FRotator::ZeroRotator;
 }
@@ -2500,11 +2504,11 @@ FRotator UPhysicsControlComponent::GetCachedBoneOrientation(
 //======================================================================================================================
 bool UPhysicsControlComponent::SetCachedBoneData(
 	const USkeletalMeshComponent* SkeletalMeshComponent,
-	const FName                   Name,
+	const FName                   BoneName,
 	const FTransform&             TM)
 {
 	UE::PhysicsControl::FBoneData* BoneData;
-	if (GetModifiableBoneData(BoneData, SkeletalMeshComponent, Name))
+	if (GetModifiableBoneData(BoneData, SkeletalMeshComponent, BoneName))
 	{
 		BoneData->CurrentTM = TM;
 		return true;
@@ -2512,7 +2516,7 @@ bool UPhysicsControlComponent::SetCachedBoneData(
 	if (bWarnAboutInvalidNames)
 	{
 		UE_LOG(LogPhysicsControl, Warning,
-			TEXT("SetCachedBoneData - invalid name %s"), *Name.ToString());
+			TEXT("SetCachedBoneData - invalid bone name %s"), *BoneName.ToString());
 	}
 	return false;
 }
@@ -2585,12 +2589,47 @@ bool UPhysicsControlComponent::GetBodyModifierExists(const FName Name) const
 	return FindBodyModifierRecord(Name) != nullptr;
 }
 
-#if WITH_EDITOR
+//======================================================================================================================
+bool UPhysicsControlComponent::ShouldCreatePhysicsState() const
+{
+	// This is needed to ensure we get the destroy call
+	return true;
+}
+
+//======================================================================================================================
+void UPhysicsControlComponent::OnCreatePhysicsState()
+{
+	Super::OnCreatePhysicsState();
+}
+
+//======================================================================================================================
+void UPhysicsControlComponent::DestroyPhysicsState()
+{
+	for (TPair<FName, FPhysicsControlRecord>& ControlRecordPair : ControlRecords)
+	{
+		DestroyControl(ControlRecordPair.Key, EDestroyBehavior::KeepRecord);
+	}
+	ControlRecords.Empty();
+
+	for (TPair<FName, FPhysicsBodyModifierRecord>& BodyModifierPair : BodyModifierRecords)
+	{
+		DestroyBodyModifier(BodyModifierPair.Key, EDestroyBehavior::KeepRecord);
+	}
+	BodyModifierRecords.Empty();
+}
+
+//======================================================================================================================
+void UPhysicsControlComponent::OnDestroyPhysicsState()
+{
+	DestroyPhysicsState();
+	Super::OnDestroyPhysicsState();
+}
 
 //======================================================================================================================
 void UPhysicsControlComponent::OnRegister()
 {
 	Super::OnRegister();
+#if WITH_EDITOR
 
 	if (SpriteComponent)
 	{
@@ -2598,8 +2637,10 @@ void UPhysicsControlComponent::OnRegister()
 		SpriteComponent->SpriteInfo.Category = TEXT("Physics");
 		SpriteComponent->SpriteInfo.DisplayName = NSLOCTEXT("SpriteCategory", "Physics", "Physics");
 	}
+#endif
 }
 
+#if WITH_EDITOR
 //======================================================================================================================
 void UPhysicsControlComponent::DebugDraw(FPrimitiveDrawInterface* PDI) const
 {
diff --git a/Engine/Plugins/Experimental/PhysicsControl/Source/PhysicsControl/Private/PhysicsControlData.cpp b/Engine/Plugins/Experimental/PhysicsControl/Source/PhysicsControl/Private/PhysicsControlData.cpp
--- a/Engine/Plugins/Experimental/PhysicsControl/Source/PhysicsControl/Private/PhysicsControlData.cpp
+++ b/Engine/Plugins/Experimental/PhysicsControl/Source/PhysicsControl/Private/PhysicsControlData.cpp
@@ -25,6 +25,18 @@ FPhysicsControlAndBodyModifierCreationDatas& FPhysicsControlAndBodyModifierCreat
 	return *this;
 }
 
+//======================================================================================================================
+bool FPhysicsControlAndBodyModifierCreationDatas::operator==(const FPhysicsControlAndBodyModifierCreationDatas& other) const
+{
+	if (Controls.OrderIndependentCompareEqual(other.Controls) &&
+		Modifiers.OrderIndependentCompareEqual(other.Modifiers))
+	{
+		return true;
+	}
+	return false;
+}
+
+
 //======================================================================================================================
 FPhysicsControlData Interpolate(
 	const FPhysicsControlData& A, const FPhysicsControlData& B, const float Weight)
diff --git a/Engine/Plugins/Experimental/PhysicsControl/Source/PhysicsControl/Private/PhysicsControlRecord.cpp b/Engine/Plugins/Experimental/PhysicsControl/Source/PhysicsControl/Private/PhysicsControlRecord.cpp
--- a/Engine/Plugins/Experimental/PhysicsControl/Source/PhysicsControl/Private/PhysicsControlRecord.cpp
+++ b/Engine/Plugins/Experimental/PhysicsControl/Source/PhysicsControl/Private/PhysicsControlRecord.cpp
@@ -10,7 +10,7 @@
 //======================================================================================================================
 void FPhysicsControlRecord::ResetConstraint()
 {
-	if (ConstraintInstance)
+	if (ConstraintInstance.IsValid())
 	{
 		ConstraintInstance->TermConstraint();
 	}
@@ -34,11 +34,11 @@ FVector FPhysicsControlRecord::GetControlPoint() const
 //======================================================================================================================
 bool FPhysicsControlRecord::InitConstraint(UObject* ConstraintDebugOwner, FName ControlName)
 {
-	if (!ConstraintInstance)
+	if (!ConstraintInstance.IsValid())
 	{
 		ConstraintInstance = MakeShared<FConstraintInstance>();
 	}
-	check(ConstraintInstance);
+	check(ConstraintInstance.IsValid());
 
 	FBodyInstance* ParentBody = UE::PhysicsControl::GetBodyInstance(
 		ParentComponent.Get(), PhysicsControl.ParentBoneName);
@@ -89,7 +89,7 @@ bool FPhysicsControlRecord::InitConstraint(UObject* ConstraintDebugOwner, FName
 // corresponds to the child frame. Frame2 will always be identity, because we never change it.
 void FPhysicsControlRecord::UpdateConstraintControlPoint()
 {
-	if (ConstraintInstance)
+	if (ConstraintInstance.IsValid())
 	{
 		// Constraints are child then parent
 		FTransform Frame1 = ConstraintInstance->GetRefFrame(EConstraintFrame::Frame1);
diff --git a/Engine/Plugins/Experimental/PhysicsControl/Source/PhysicsControl/Public/PhysicsControlAsset.h b/Engine/Plugins/Experimental/PhysicsControl/Source/PhysicsControl/Public/PhysicsControlAsset.h
--- a/Engine/Plugins/Experimental/PhysicsControl/Source/PhysicsControl/Public/PhysicsControlAsset.h
+++ b/Engine/Plugins/Experimental/PhysicsControl/Source/PhysicsControl/Public/PhysicsControlAsset.h
@@ -12,7 +12,7 @@
 class USkeletalMesh;
 
 /**
- * Asset for storing Physics Control Profiles. These will contain data that define:
+ * Asset for storing Physics Control Profiles. The asset will contain data that define:
  * - Controls and body modifiers to be created on a mesh
  * - Sets referencing those controls and body modifiers
  * - Full profiles containing settings for all the controls/modifiers
@@ -76,12 +76,9 @@ public:
 	UPROPERTY(EditAnywhere, Category = Inheritance)
 	TArray<TSoftObjectPtr<UPhysicsControlAsset>> AdditionalProfileAssets;
 
-	/**
-	* The skeletal mesh to use for generating controls and previewing. If it turns out this 
-	* doesn't need to be stored in the asset, it will get moved out of here.
-	*/
+	/** The PhysicsAsset that this control asset is targeting. Can also get the preview mesh from this */
 	UPROPERTY(AssetRegistrySearchable, EditAnywhere, Category = PreviewMesh)
-	TSoftObjectPtr<USkeletalMesh> PreviewSkeletalMesh;
+	TSoftObjectPtr<UPhysicsAsset> PhysicsAsset;
 
 public:
 	// "My" runtime data - i.e. the data that will be combined with what has been inherited
@@ -90,32 +87,32 @@ public:
 	/**
 	 * We can define controls in the form of limbs etc here
 	 */
-	UPROPERTY(EditAnywhere, Category = ProfileData, meta=(DisplayName="Character Setup Data"))
+	UPROPERTY(EditAnywhere, Category = Setup, meta=(DisplayName="Character Setup Data"))
 	FPhysicsControlCharacterSetupData MyCharacterSetupData;
 
 	/**
 	 * Additional controls and modifiers. If these have the same name as one that's 
 	 * already created, they'll just override it.
 	 */
-	UPROPERTY(EditAnywhere, Category = ProfileData, meta = (DisplayName = "Additional Controls and Modifiers"))
+	UPROPERTY(EditAnywhere, Category = Setup, meta = (DisplayName = "Additional Controls and Modifiers"))
 	FPhysicsControlAndBodyModifierCreationDatas MyAdditionalControlsAndModifiers;
 
 	/**
 	 * Additional control and body modifier sets
 	 */
-	UPROPERTY(EditAnywhere, Category = ProfileData, meta = (DisplayName = "Additional Sets"))
+	UPROPERTY(EditAnywhere, Category = Setup, meta = (DisplayName = "Additional Sets"))
 	FPhysicsControlSetUpdates MyAdditionalSets;
 
 	/**
 	 * Initial updates to apply immediately after controls and modifiers are created
 	 */
-	UPROPERTY(EditAnywhere, Category = ProfileData, meta = (DisplayName = "Initial Control and Modifier Updates"))
+	UPROPERTY(EditAnywhere, Category = Setup, meta = (DisplayName = "Initial Control and Modifier Updates"))
 	TArray<FPhysicsControlControlAndModifierUpdates> MyInitialControlAndModifierUpdates;
 
 	/**
 	 * The named profiles, which are essentially control and modifier updates
 	 */
-	UPROPERTY(EditAnywhere, Category = ProfileData, meta = (DisplayName = "Profiles"))
+	UPROPERTY(EditAnywhere, Category = Profiles, meta = (DisplayName = "Profiles"))
 	TMap<FName, FPhysicsControlControlAndModifierUpdates> MyProfiles;
 #endif
 
@@ -124,16 +121,23 @@ public:
 
 #if WITH_EDITOR
 	/** Shows all the controls etc that would be made */
-	UFUNCTION(CallInEditor, Category = Development)
+	UFUNCTION(CallInEditor, Category = Actions)
 	void ShowCompiledData() const;
 
 	/** 
 	 * Collapses inherited and authored profiles etc to make a profile asset that can be read without 
 	 * need for subsequent processing.
 	 */
-	UFUNCTION(CallInEditor)
+	UFUNCTION(CallInEditor, Category = Actions)
 	void Compile();
 
+	/** 
+	 * Returns true if compilation would change any of our compiled data. Note that this is potentially slow as
+	 * it simply compiles and compares the result with the data we already have.
+	 */
+	UFUNCTION(CallInEditor, Category = Actions)
+	bool IsCompilationNeeded() const;
+
 	/** Combines and returns data from our parent and ourself */
 	FPhysicsControlCharacterSetupData GetCharacterSetupData() const;
 
@@ -159,8 +163,14 @@ public:
 	/** END IInterface_PreviewMeshProvider interface */
 
 #if WITH_EDITOR
-	/* Get name of Preview Mesh property */
+	/** This loads the asset if necessary */
+	UPhysicsAsset* GetPhysicsAsset() const;
+	void SetPhysicsAsset(UPhysicsAsset* InPhysicsAsset);
+
 	static const FName GetPreviewMeshPropertyName();
 #endif
 
+
+
+
 };
diff --git a/Engine/Plugins/Experimental/PhysicsControl/Source/PhysicsControl/Public/PhysicsControlComponent.h b/Engine/Plugins/Experimental/PhysicsControl/Source/PhysicsControl/Public/PhysicsControlComponent.h
--- a/Engine/Plugins/Experimental/PhysicsControl/Source/PhysicsControl/Public/PhysicsControlComponent.h
+++ b/Engine/Plugins/Experimental/PhysicsControl/Source/PhysicsControl/Public/PhysicsControlComponent.h
@@ -1679,11 +1679,14 @@ public:
 	bool bWarnAboutInvalidNames = true;
 
 public:
-#if WITH_EDITOR
 	//Begin ActorComponent interface
 	virtual void OnRegister() override;
+	virtual bool ShouldCreatePhysicsState() const override;
+	virtual void OnCreatePhysicsState() override;
+	virtual void OnDestroyPhysicsState() override;
 	//End ActorComponent interface
 
+#if WITH_EDITOR
 	// Used by the component visualizer
 	virtual void DebugDraw(FPrimitiveDrawInterface* PDI) const;
 	virtual void DebugDrawControl(FPrimitiveDrawInterface* PDI, const FPhysicsControlRecord& Record, const FName ControlName) const;
@@ -1697,6 +1700,8 @@ protected:
 	virtual void TickComponent(float DeltaTime, enum ELevelTick TickType, FActorComponentTickFunction* ThisTickFunction) override;
 	// End UActorComponent Interface
 
+	// Terminates everything we have made in physics
+	void DestroyPhysicsState();
 
 	/**
 	 * Retrieves the bone data for the specified bone given the skeletal mesh component.
diff --git a/Engine/Plugins/Experimental/PhysicsControl/Source/PhysicsControl/Public/PhysicsControlData.h b/Engine/Plugins/Experimental/PhysicsControl/Source/PhysicsControl/Public/PhysicsControlData.h
--- a/Engine/Plugins/Experimental/PhysicsControl/Source/PhysicsControl/Public/PhysicsControlData.h
+++ b/Engine/Plugins/Experimental/PhysicsControl/Source/PhysicsControl/Public/PhysicsControlData.h
@@ -82,6 +82,8 @@ struct PHYSICSCONTROL_API FPhysicsControlSetUpdate
 {
 	GENERATED_BODY();
 
+	bool operator==(const FPhysicsControlSetUpdate& other) const = default;
+
 	UPROPERTY(EditAnywhere, BlueprintReadWrite, Category = PhysicsControl)
 	FName SetName;
 
@@ -100,6 +102,8 @@ struct PHYSICSCONTROL_API FPhysicsControlSetUpdates
 
 	FPhysicsControlSetUpdates& operator+=(const FPhysicsControlSetUpdates& Other);
 
+	bool operator==(const FPhysicsControlSetUpdates& other) const = default;
+
 	UPROPERTY(EditAnywhere, BlueprintReadWrite, Category = PhysicsControl)
 	TArray<FPhysicsControlSetUpdate> ControlSetUpdates;
 
@@ -139,6 +143,8 @@ struct PHYSICSCONTROL_API FPhysicsControlModifierSparseData
 	{
 	}
 
+	bool operator==(const FPhysicsControlModifierSparseData& other) const = default;
+
 	/**
 	 * How the body should move.
 	 */
@@ -221,6 +227,8 @@ struct PHYSICSCONTROL_API FPhysicsControlModifierData
 	{
 	}
 
+	bool operator==(const FPhysicsControlModifierData& other) const = default;
+
 	void UpdateFromSparseData(const FPhysicsControlModifierSparseData& SparseData);
 
 	/**
@@ -281,6 +289,8 @@ struct PHYSICSCONTROL_API FPhysicsBodyModifier
 	FPhysicsBodyModifier(FName InBoneName, const FPhysicsControlModifierData& InModifierData)
 		: BoneName(InBoneName), ModifierData(InModifierData) {}
 
+	bool operator==(const FPhysicsBodyModifier& other) const = default;
+
 	UPROPERTY(EditAnywhere, BlueprintReadWrite, Category = PhysicsControl)
 	FName BoneName;
 
@@ -296,6 +306,8 @@ struct PHYSICSCONTROL_API FPhysicsBodyModifierCreationData
 {
 	GENERATED_BODY();
 
+	bool operator==(const FPhysicsBodyModifierCreationData& other) const = default;
+
 	UPROPERTY(EditAnywhere, BlueprintReadWrite, Category = PhysicsControl)
 	FPhysicsBodyModifier Modifier;
 
@@ -351,6 +363,8 @@ struct PHYSICSCONTROL_API FPhysicsControlSparseData
 	{
 	}
 
+	bool operator==(const FPhysicsControlSparseData& other) const = default;
+
 	/** The strength used to drive linear motion */
 	UPROPERTY(EditAnywhere, BlueprintReadWrite, Category = PhysicsControl, meta = (ClampMin = "0.0", editcondition = "bEnableLinearStrength"))
 	float LinearStrength;
@@ -543,6 +557,8 @@ struct PHYSICSCONTROL_API FPhysicsControlData
 	{
 	}
 
+	bool operator==(const FPhysicsControlData& other) const = default;
+
 	/** Applies the values that have been flagged as enabled from the sparse data */
 	void UpdateFromSparseData(const FPhysicsControlSparseData& SparseData);
 
@@ -679,6 +695,8 @@ struct PHYSICSCONTROL_API FPhysicsControlSparseMultiplier
 	{
 	}
 
+	bool operator==(const FPhysicsControlSparseMultiplier& other) const = default;
+
 	// Per-direction multiplier on the linear strength.
 	UPROPERTY(EditAnywhere, BlueprintReadWrite, Category = PhysicsControl, meta = (ClampMin = "0.0", editcondition = "bEnableLinearStrengthMultiplier"))
 	FVector LinearStrengthMultiplier;
@@ -765,6 +783,8 @@ struct PHYSICSCONTROL_API FPhysicsControlMultiplier
 	{
 	}
 
+	bool operator==(const FPhysicsControlMultiplier& other) const = default;
+
 	/** Applies the values that have been flagged as enabled from the sparse data */
 	void UpdateFromSparseData(const FPhysicsControlSparseMultiplier& SparseData);
 
@@ -870,6 +890,8 @@ struct PHYSICSCONTROL_API FPhysicsControl
 	{
 	}
 
+	bool operator==(const FPhysicsControl& other) const = default;
+
 	// Indicates if the control is enabled
 	bool IsEnabled() const { return ControlData.bEnabled; }
 
@@ -906,6 +928,8 @@ struct PHYSICSCONTROL_API FPhysicsControlCreationData
 {
 	GENERATED_BODY();
 
+	bool operator==(const FPhysicsControlCreationData& other) const = default;
+
 	UPROPERTY(EditAnywhere, BlueprintReadWrite, Category = PhysicsControl)
 	FPhysicsControl Control;
 
@@ -926,6 +950,8 @@ struct PHYSICSCONTROL_API FPhysicsControlNamedControlParameters
 	FPhysicsControlNamedControlParameters(FName InName, const FPhysicsControlSparseData& InData)
 		: Name(InName), Data(InData) {}
 
+	bool operator==(const FPhysicsControlNamedControlParameters& other) const = default;
+
 	// The name of the control (or set of controls) to update
 	UPROPERTY(EditAnywhere, BlueprintReadWrite, Category = PhysicsControl)
 	FName Name;
@@ -947,6 +973,8 @@ struct PHYSICSCONTROL_API FPhysicsControlNamedControlMultiplierParameters
 	FPhysicsControlNamedControlMultiplierParameters(FName InName, const FPhysicsControlSparseMultiplier& InData)
 		: Name(InName), Data(InData) {}
 
+	bool operator==(const FPhysicsControlNamedControlMultiplierParameters& other) const = default;
+
 	// The name of the control (or set of controls) to update
 	UPROPERTY(EditAnywhere, BlueprintReadWrite, Category = PhysicsControl)
 	FName Name;
@@ -969,6 +997,8 @@ struct PHYSICSCONTROL_API FPhysicsControlNamedModifierParameters
 	FPhysicsControlNamedModifierParameters(FName InName, const FPhysicsControlModifierSparseData& InData)
 		: Name(InName), Data(InData) {}
 
+	bool operator==(const FPhysicsControlNamedModifierParameters& other) const = default;
+
 	// The name of the modifier (or set of modifiers) to update
 	UPROPERTY(EditAnywhere, BlueprintReadWrite, Category = PhysicsControl)
 	FName Name;
@@ -1022,6 +1052,8 @@ struct PHYSICSCONTROL_API FPhysicsControlControlAndModifierUpdates
 {
 	GENERATED_BODY();
 
+	bool operator==(const FPhysicsControlControlAndModifierUpdates& other) const = default;
+
 	/** Modifications to the underlying controls */
 	UPROPERTY(EditAnywhere, BlueprintReadWrite, Category = PhysicsControl)
 	TArray<FPhysicsControlNamedControlParameters> ControlUpdates;
@@ -1045,6 +1077,8 @@ struct PHYSICSCONTROL_API FPhysicsControlAndBodyModifierCreationDatas
 
 	FPhysicsControlAndBodyModifierCreationDatas& operator+=(const FPhysicsControlAndBodyModifierCreationDatas& other);
 
+	bool operator==(const FPhysicsControlAndBodyModifierCreationDatas& other) const;
+
 	UPROPERTY(EditAnywhere, BlueprintReadWrite, Category = PhysicsControl)
 	TMap<FName, FPhysicsControlCreationData> Controls;
 
diff --git a/Engine/Plugins/Experimental/PhysicsControl/Source/PhysicsControl/Public/PhysicsControlLimbData.h b/Engine/Plugins/Experimental/PhysicsControl/Source/PhysicsControl/Public/PhysicsControlLimbData.h
--- a/Engine/Plugins/Experimental/PhysicsControl/Source/PhysicsControl/Public/PhysicsControlLimbData.h
+++ b/Engine/Plugins/Experimental/PhysicsControl/Source/PhysicsControl/Public/PhysicsControlLimbData.h
@@ -27,6 +27,8 @@ struct PHYSICSCONTROL_API FPhysicsControlLimbSetupData
 		, bCreateParentSpaceControls(true), bCreateBodyModifiers(true)
 	{}
 
+	bool operator==(const FPhysicsControlLimbSetupData& other) const = default;
+
 	/** The name of the limb that this will be used to create */
 	UPROPERTY(EditAnywhere, BlueprintReadWrite, Category = PhysicsControl)
 	FName LimbName;
@@ -107,6 +109,8 @@ struct PHYSICSCONTROL_API FPhysicsControlCharacterSetupData
 
 	FPhysicsControlCharacterSetupData& operator+=(const FPhysicsControlCharacterSetupData& other);
 
+	bool operator==(const FPhysicsControlCharacterSetupData& other) const = default;
+
 	UPROPERTY(EditAnywhere, Category = ControlSetup)
 	TArray<FPhysicsControlLimbSetupData> LimbSetupData;
 
diff --git a/Engine/Plugins/Experimental/PhysicsControl/Source/PhysicsControl/Public/RigidBodyPoseData.h b/Engine/Plugins/Experimental/PhysicsControl/Source/PhysicsControl/Public/RigidBodyPoseData.h
--- a/Engine/Plugins/Experimental/PhysicsControl/Source/PhysicsControl/Public/RigidBodyPoseData.h
+++ b/Engine/Plugins/Experimental/PhysicsControl/Source/PhysicsControl/Public/RigidBodyPoseData.h
@@ -19,6 +19,7 @@ struct FOutputBoneData
 {
 	FOutputBoneData()
 		: CompactPoseBoneIndex(INDEX_NONE), CompactPoseParentBoneIndex(INDEX_NONE)
+		, BodyIndex(0), ParentBodyIndex(0)
 	{}
 
 	TArray<FCompactPoseBoneIndex> BoneIndicesToParentBody;
diff --git a/Engine/Plugins/Experimental/PhysicsControl/Source/PhysicsControlEditor/PhysicsControlEditor.Build.cs b/Engine/Plugins/Experimental/PhysicsControl/Source/PhysicsControlEditor/PhysicsControlEditor.Build.cs
--- a/Engine/Plugins/Experimental/PhysicsControl/Source/PhysicsControlEditor/PhysicsControlEditor.Build.cs
+++ b/Engine/Plugins/Experimental/PhysicsControl/Source/PhysicsControlEditor/PhysicsControlEditor.Build.cs
@@ -35,21 +35,26 @@ public class PhysicsControlEditor : ModuleRules
 		PrivateDependencyModuleNames.AddRange(
 			new string[]
 			{
+				"AdvancedPreviewScene",
 				"AnimationCore",
 				"AnimationEditMode",
 				"AnimGraph",
 				"AnimGraphRuntime",
 				"ApplicationCore",
 				"BlueprintGraph",
+				"Chaos",
 				"ComponentVisualizers",
 				"Core",
 				"CoreUObject",
 				"EditorFramework",
+				"EditorStyle",
 				"Engine",
 				"InputCore",
 				"Persona",
+				"PinnedCommandList",
 				"PhysicsControl",
 				"PhysicsControlUncookedOnly",
+				"SkeletonEditor",
 				"Slate",
 				"SlateCore",
 				"ToolMenus",
@@ -63,12 +68,12 @@ public class PhysicsControlEditor : ModuleRules
 			PrivateDependencyModuleNames.AddRange(
 				new string[]
 				{
-						"AnimationEditor",
-						"AnimGraph",
-						"BlueprintGraph",
-						"EditorFramework",
-						"Kismet",
-						"UnrealEd",
+					"AnimationEditor",
+					"AnimGraph",
+					"BlueprintGraph",
+					"EditorFramework",
+					"Kismet",
+					"UnrealEd",
 				}
 			);
 		}
diff --git a/Engine/Plugins/Experimental/PhysicsControl/Source/PhysicsControlEditor/Private/PhysicsControlAssetActions.cpp b/Engine/Plugins/Experimental/PhysicsControl/Source/PhysicsControlEditor/Private/PhysicsControlAssetActions.cpp
--- a/Engine/Plugins/Experimental/PhysicsControl/Source/PhysicsControlEditor/Private/PhysicsControlAssetActions.cpp
+++ b/Engine/Plugins/Experimental/PhysicsControl/Source/PhysicsControlEditor/Private/PhysicsControlAssetActions.cpp
@@ -2,7 +2,7 @@
 
 #include "PhysicsControlAssetActions.h"
 #include "PhysicsControlAsset.h"
-#include "PhysicsControlAssetEditorToolkit.h"
+#include "PhysicsControlAssetEditor.h"
 
 //======================================================================================================================
 UClass* FPhysicsControlAssetActions::GetSupportedClass() const
@@ -13,7 +13,7 @@ UClass* FPhysicsControlAssetActions::GetSupportedClass() const
 //======================================================================================================================
 FText FPhysicsControlAssetActions::GetName() const
 {
-	return INVTEXT("Physics Control Profile");
+	return INVTEXT("Physics Control Asset");
 }
 
 //======================================================================================================================
@@ -40,7 +40,7 @@ void FPhysicsControlAssetActions::OpenAssetEditor(
 	{
 		if (UPhysicsControlAsset* Asset = Cast<UPhysicsControlAsset>(Object))
 		{
-			TSharedRef<FPhysicsControlAssetEditorToolkit> NewEditor(new FPhysicsControlAssetEditorToolkit());
+			TSharedRef<FPhysicsControlAssetEditor> NewEditor(new FPhysicsControlAssetEditor());
 			NewEditor->InitAssetEditor(Mode, EditWithinLevelEditor, Asset);
 		}
 	}
diff --git a/Engine/Plugins/Experimental/PhysicsControl/Source/PhysicsControlEditor/Private/PhysicsControlAssetApplicationMode.cpp b/Engine/Plugins/Experimental/PhysicsControl/Source/PhysicsControlEditor/Private/PhysicsControlAssetApplicationMode.cpp
--- a/Engine/Plugins/Experimental/PhysicsControl/Source/PhysicsControlEditor/Private/PhysicsControlAssetApplicationMode.cpp
+++ b/Engine/Plugins/Experimental/PhysicsControl/Source/PhysicsControlEditor/Private/PhysicsControlAssetApplicationMode.cpp
@@ -1,8 +1,12 @@
 // Copyright Epic Games, Inc. All Rights Reserved.
 
 #include "PhysicsControlAssetApplicationMode.h"
-#include "PhysicsControlAssetEditorToolkit.h"
+#include "PhysicsControlAssetEditor.h"
+#include "PhysicsControlAssetEditorProfilesTabSummoners.h"
+#include "PhysicsControlAsset.h"
+
 #include "PersonaModule.h"
+#include "ISkeletonEditorModule.h"
 #include "Modules/ModuleManager.h"
 #include "PersonaTabs.h"
 
@@ -10,35 +14,74 @@
 
 FName FPhysicsControlAssetApplicationMode::ModeName("PhysicsControlAssetEditMode");
 
+static const FName PhysicsControlAssetEditorTabName("PhysicsControlAssetEditorTab");
+
 //======================================================================================================================
 FPhysicsControlAssetApplicationMode::FPhysicsControlAssetApplicationMode(
 	TSharedRef<FWorkflowCentricApplication> InHostingApp,
+	TSharedPtr<ISkeletonTree>               SkeletonTree,
 	TSharedRef<IPersonaPreviewScene>        InPreviewScene)
 	: 
-	FApplicationMode(PhysicsControlAssetEditorModes::Editor)
+	FApplicationMode(PhysicsControlAssetEditorModes::PhysicsControlAssetEditorMode)
 {
-	EditorToolkit = StaticCastSharedRef<FPhysicsControlAssetEditorToolkit>(InHostingApp);
-	TSharedRef<FPhysicsControlAssetEditorToolkit> PhysicsControlAssetEditor = 
-		StaticCastSharedRef<FPhysicsControlAssetEditorToolkit>(InHostingApp);
+	PhysicsControlAssetEditor = StaticCastSharedRef<FPhysicsControlAssetEditor>(InHostingApp);
+	TSharedRef<FPhysicsControlAssetEditor> PhysicsControlAssetEditorSharedRef = 
+		StaticCastSharedRef<FPhysicsControlAssetEditor>(InHostingApp);
+
+	if (SkeletonTree.IsValid())
+	{
+		ISkeletonEditorModule& SkeletonEditorModule = 
+			FModuleManager::LoadModuleChecked<ISkeletonEditorModule>("SkeletonEditor");
+		TabFactories.RegisterFactory(
+			SkeletonEditorModule.CreateSkeletonTreeTabFactory(InHostingApp, SkeletonTree.ToSharedRef()));
+	}
+
+	TArray<TSharedPtr<FExtender>> ViewportExtenders;
+	ViewportExtenders.Add(MakeShared<FExtender>());
 
 	FPersonaViewportArgs ViewportArgs(InPreviewScene);
 	ViewportArgs.bAlwaysShowTransformToolbar = true;
 	ViewportArgs.bShowStats = false;
-	ViewportArgs.bShowTimeline = false;
+	ViewportArgs.bShowTimeline = true;
 	ViewportArgs.bShowLODMenu = true;
-	ViewportArgs.bShowPlaySpeedMenu = false;
-	ViewportArgs.bShowPhysicsMenu = false;
+	ViewportArgs.bShowPlaySpeedMenu = true;
+	ViewportArgs.bShowPhysicsMenu = true;
 	ViewportArgs.ContextName = TEXT("PhysicsControlAssetEditor.Viewport");
+	ViewportArgs.Extenders = ViewportExtenders;
 	ViewportArgs.OnViewportCreated = FOnViewportCreated::CreateSP(
-		PhysicsControlAssetEditor, &FPhysicsControlAssetEditorToolkit::HandleViewportCreated);
+		PhysicsControlAssetEditorSharedRef, &FPhysicsControlAssetEditor::HandleViewportCreated);
 
 	// Register Persona tabs.
 	FPersonaModule& PersonaModule = FModuleManager::LoadModuleChecked<FPersonaModule>("Persona");
 	TabFactories.RegisterFactory(PersonaModule.CreatePersonaViewportTabFactory(InHostingApp, ViewportArgs));
-	TabFactories.RegisterFactory(PersonaModule.CreateDetailsTabFactory(InHostingApp, FOnDetailsCreated::CreateSP(&PhysicsControlAssetEditor.Get(), &FPhysicsControlAssetEditorToolkit::HandleDetailsCreated)));
+	TabFactories.RegisterFactory(PersonaModule.CreateAdvancedPreviewSceneTabFactory(InHostingApp, InPreviewScene));
+
+	TabFactories.RegisterFactory(
+		MakeShared<FPhysicsControlAssetEditorSetupTabSummoner>(
+			InHostingApp, 
+			CastChecked<UPhysicsControlAsset>(
+				(*PhysicsControlAssetEditorSharedRef->GetObjectsCurrentlyBeingEdited())[0])));
+
+	TabFactories.RegisterFactory(
+		MakeShared<FPhysicsControlAssetEditorProfileTabSummoner>(
+			InHostingApp,
+			CastChecked<UPhysicsControlAsset>(
+				(*PhysicsControlAssetEditorSharedRef->GetObjectsCurrentlyBeingEdited())[0])));
+
+	TabFactories.RegisterFactory(
+		MakeShared<FPhysicsControlAssetEditorPreviewTabSummoner>(
+			InHostingApp,
+			CastChecked<UPhysicsControlAsset>(
+				(*PhysicsControlAssetEditorSharedRef->GetObjectsCurrentlyBeingEdited())[0])));
+
+
+	// For standard tabs, these are provided by Persona. For custom tabs, the contents of the tab is
+	// provided by a summoner. When Summoners are made, they register their name in the constructor.
+	// These names then hook into the tab names below. See MLDeformerApplicationMode.cpp - that puts
+	// the names into a TabID field in the summoner.
 
 	// Create tab layout.
-	TabLayout = FTabManager::NewLayout("Standalone_PhysicsControlAssetEditor_Layout_v0.001")
+	TabLayout = FTabManager::NewLayout("Standalone_PhysicsControlAssetEditor_Layout_v1.0")
 		->AddArea
 		(
 			FTabManager::NewPrimaryArea()
@@ -59,14 +102,22 @@ FPhysicsControlAssetApplicationMode::FPhysicsControlAssetApplicationMode(
 						->SetSizeCoefficient(0.7f)
 						->SetHideTabWell(true)
 						->AddTab(FPersonaTabs::PreviewViewportID, ETabState::OpenedTab)
+						->AddTab(FPersonaTabs::SkeletonTreeViewID, ETabState::OpenedTab)
 					)
 				)
 				->Split
 				(
 					FTabManager::NewStack()
 					->SetSizeCoefficient(0.3f)
-					->AddTab(FPersonaTabs::DetailsID, ETabState::OpenedTab)
-					->SetForegroundTab(FPersonaTabs::DetailsID)
+					//->AddTab(FPersonaTabs::DetailsID, ETabState::OpenedTab)
+					->AddTab(FPersonaTabs::AdvancedPreviewSceneSettingsID, ETabState::OpenedTab)
+					->AddTab(
+						FPhysicsControlAssetEditorSetupTabSummoner::TabName, ETabState::OpenedTab)
+					->AddTab(
+						FPhysicsControlAssetEditorProfileTabSummoner::TabName, ETabState::OpenedTab)
+					->AddTab(
+						FPhysicsControlAssetEditorPreviewTabSummoner::TabName, ETabState::OpenedTab)
+					->SetForegroundTab(FPhysicsControlAssetEditorSetupTabSummoner::TabName)
 				)
 			)
 		);
@@ -80,7 +131,7 @@ FPhysicsControlAssetApplicationMode::FPhysicsControlAssetApplicationMode(
 //======================================================================================================================
 void FPhysicsControlAssetApplicationMode::RegisterTabFactories(TSharedPtr<FTabManager> InTabManager)
 {
-	TSharedPtr<FPhysicsControlAssetEditorToolkit> Editor = EditorToolkit.Pin();
+	TSharedPtr<FPhysicsControlAssetEditor> Editor = PhysicsControlAssetEditor.Pin();
 	Editor->RegisterTabSpawners(InTabManager.ToSharedRef());
 	Editor->PushTabFactories(TabFactories);
 	FApplicationMode::RegisterTabFactories(InTabManager);
diff --git a/Engine/Plugins/Experimental/PhysicsControl/Source/PhysicsControlEditor/Private/PhysicsControlAssetEditor.cpp b/Engine/Plugins/Experimental/PhysicsControl/Source/PhysicsControlEditor/Private/PhysicsControlAssetEditor.cpp
new file mode 100644
--- /dev/null
+++ b/Engine/Plugins/Experimental/PhysicsControl/Source/PhysicsControlEditor/Private/PhysicsControlAssetEditor.cpp
@@ -0,0 +1,1145 @@
+// Copyright Epic Games, Inc. All Rights Reserved.
+
+#include "PhysicsControlAssetEditor.h"
+
+#include "PhysicsControlAsset.h"
+#include "PhysicsControlAssetApplicationMode.h"
+#include "PhysicsControlAssetEditorCommands.h"
+#include "PhysicsControlAssetEditorData.h"
+#include "PhysicsControlAssetEditorEditMode.h"
+#include "PhysicsControlAssetEditorPhysicsHandleComponent.h"
+#include "PhysicsControlAssetEditorSkeletalMeshComponent.h"
+#include "PhysicsControlAssetEditorSkeletonTreeBuilder.h"
+#include "PhysicsControlAssetEditorToolMenuContext.h"
+#include "PhysicsControlComponent.h"
+#include "PhysicsControlEditorModule.h"
+#include "PhysicsControlLog.h"
+
+#include "Animation/DebugSkelMeshComponent.h"
+#include "AnimationEditorPreviewActor.h"
+#include "AnimPreviewInstance.h"
+#include "Components/StaticMeshComponent.h"
+#include "Editor/PropertyEditor/Private/PropertyNode.h"
+#include "EditorModeManager.h"
+#include "IDetailsView.h"
+#include "IPersonaToolkit.h"
+#include "IPinnedCommandList.h"
+#include "ISkeletonEditorModule.h"
+#include "Modules/ModuleManager.h"
+#include "PersonaModule.h"
+#include "PersonaToolMenuContext.h"
+#include "PhysicsEngine/SkeletalBodySetup.h"
+#include "Preferences/PersonaOptions.h"
+#include "PropertyEditorModule.h"
+#include "SkeletonTreePhysicsControlBodyItem.h"
+#include "SkeletonTreePhysicsControlShapeItem.h"
+#include "SkeletonTreeSelection.h"
+#include "ToolMenus.h"
+#include "UICommandList_Pinnable.h"
+#include "Widgets/Input/SNumericEntryBox.h"
+
+#define LOCTEXT_NAMESPACE "PhysicsControlAssetEditor"
+
+const FName PhysicsControlAssetEditorModes::PhysicsControlAssetEditorMode("PhysicsControlAssetEditorMode");
+const FName PhysicsControlAssetEditorAppName = FName(TEXT("PhysicsControlAssetEditorApp"));
+
+//======================================================================================================================
+namespace PhysicsControlAssetEditor
+{
+	static TSharedPtr<FPhysicsControlAssetEditor> GetPhysicsControlAssetEditorFromToolContext(const FToolMenuContext& InMenuContext)
+	{
+		if (UPhysicsControlAssetEditorToolMenuContext* Context = InMenuContext.FindContext<UPhysicsControlAssetEditorToolMenuContext>())
+		{
+			return Context->PhysicsControlAssetEditor.Pin();
+		}
+		return TSharedPtr<FPhysicsControlAssetEditor>();
+	}
+}
+
+
+//======================================================================================================================
+void FPhysicsControlAssetEditor::InitAssetEditor(
+	const EToolkitMode::Type        Mode,
+	const TSharedPtr<IToolkitHost>& InitToolkitHost,
+	UPhysicsControlAsset*           InPhysicsControlAsset)
+{
+	bIsInitialized = false;
+
+	// Initialise EditorData
+	{
+		EditorData = MakeShared<FPhysicsControlAssetEditorData>();
+		EditorData->PhysicsControlAsset = InPhysicsControlAsset;
+		EditorData->CachePreviewMesh();
+	}
+
+	USkeletalMesh* SkeletalMesh = InPhysicsControlAsset->GetPreviewMesh();
+	USkeleton* Skeleton = SkeletalMesh ? SkeletalMesh->GetSkeleton() : nullptr;
+
+	// Create Persona toolkit
+	FPersonaToolkitArgs PersonaToolkitArgs;
+	PersonaToolkitArgs.OnPreviewSceneCreated = FOnPreviewSceneCreated::FDelegate::CreateSP(
+		this, &FPhysicsControlAssetEditor::HandlePreviewSceneCreated);
+	FPersonaModule& PersonaModule = FModuleManager::LoadModuleChecked<FPersonaModule>("Persona");
+	PersonaToolkit = PersonaModule.CreatePersonaToolkit(InPhysicsControlAsset, PersonaToolkitArgs, Skeleton);
+	PersonaModule.RecordAssetOpened(FAssetData(InPhysicsControlAsset));
+
+	// Make the skeleton tree
+	{
+		FSkeletonTreeArgs SkeletonTreeArgs;
+		SkeletonTreeArgs.OnSelectionChanged = FOnSkeletonTreeSelectionChanged::CreateSP(this, &FPhysicsControlAssetEditor::HandleSelectionChanged);
+		SkeletonTreeArgs.PreviewScene = PersonaToolkit->GetPreviewScene();
+		SkeletonTreeArgs.bShowBlendProfiles = false;
+		SkeletonTreeArgs.bShowDebugVisualizationOptions = true;
+		SkeletonTreeArgs.bAllowMeshOperations = false;
+		SkeletonTreeArgs.bAllowSkeletonOperations = false;
+		SkeletonTreeArgs.bHideBonesByDefault = true;
+		SkeletonTreeArgs.OnGetFilterText = FOnGetFilterText::CreateSP(this, &FPhysicsControlAssetEditor::HandleGetFilterLabel);
+		SkeletonTreeArgs.Extenders = MakeShared<FExtender>();
+		SkeletonTreeArgs.Extenders->AddMenuExtension("FilterOptions", EExtensionHook::After, GetToolkitCommands(), FMenuExtensionDelegate::CreateSP(this, &FPhysicsControlAssetEditor::HandleExtendFilterMenu));
+		SkeletonTreeArgs.Extenders->AddMenuExtension("SkeletonTreeContextMenu", EExtensionHook::After, GetToolkitCommands(), FMenuExtensionDelegate::CreateSP(this, &FPhysicsControlAssetEditor::HandleExtendContextMenu));
+
+		if (SkeletalMesh)
+		{
+			if (UPhysicsAsset* PhysicsAsset = SkeletalMesh->GetPhysicsAsset())
+			{
+				SkeletonTreeArgs.Builder = SkeletonTreeBuilder = MakeShared<FPhysicsControlAssetEditorSkeletonTreeBuilder>(PhysicsAsset);
+			}
+		}
+		SkeletonTreeArgs.ContextName = GetToolkitFName();
+
+		ISkeletonEditorModule& SkeletonEditorModule = FModuleManager::GetModuleChecked<ISkeletonEditorModule>("SkeletonEditor");
+
+		GetMutableDefault<UPersonaOptions>()->bFlattenSkeletonHierarchyWhenFiltering = false;
+		GetMutableDefault<UPersonaOptions>()->bHideParentsWhenFiltering = true;
+
+		SkeletonTree = SkeletonEditorModule.CreateSkeletonTree(PersonaToolkit->GetSkeleton(), SkeletonTreeArgs);
+	}
+
+	bSelecting = false;
+
+	GEditor->RegisterForUndo(this);
+
+	// Register our commands. This will only register them if not previously registered
+	FPhysicsControlAssetEditorCommands::Register();
+
+	BindCommands();
+
+	// Initialise the asset editor
+	const bool bCreateDefaultStandaloneMenu = true;
+	const bool bCreateDefaultToolbar = true;
+	FAssetEditorToolkit::InitAssetEditor(
+		Mode,
+		InitToolkitHost,
+		PhysicsControlAssetEditorAppName,
+		FTabManager::FLayout::NullLayout,
+		bCreateDefaultStandaloneMenu,
+		bCreateDefaultToolbar,
+		InPhysicsControlAsset);
+
+	// Create and set the application mode.
+	AddApplicationMode(
+		PhysicsControlAssetEditorModes::PhysicsControlAssetEditorMode,
+		MakeShareable(new FPhysicsControlAssetApplicationMode(
+			SharedThis(this), SkeletonTree.ToSharedRef(), PersonaToolkit->GetPreviewScene())));
+	SetCurrentMode(PhysicsControlAssetEditorModes::PhysicsControlAssetEditorMode);
+
+	// Activate the editor mode.
+	GetEditorModeManager().SetDefaultMode(FPhysicsControlAssetEditorEditMode::ModeName);
+	GetEditorModeManager().ActivateMode(FPhysicsControlAssetEditorEditMode::ModeName);
+	GetEditorModeManager().ActivateMode(FPersonaEditModes::SkeletonSelection);
+
+	FPhysicsControlAssetEditorEditMode* EditorMode = GetEditorModeManager().
+		GetActiveModeTyped<FPhysicsControlAssetEditorEditMode>(FPhysicsControlAssetEditorEditMode::ModeName);
+	EditorMode->SetEditorData(SharedThis(this), EditorData);
+
+	// Need to load the module before extending menus etc
+	//FPhysicsControlEditorModule& EditorModule = 
+	//	FModuleManager::LoadModuleChecked<FPhysicsControlEditorModule>("PhysicsControlEditor");
+
+	ExtendMenu();
+	ExtendToolbar();
+	ExtendViewportMenus();
+	RegenerateMenusAndToolbars();
+
+	bIsInitialized = true;
+}
+
+//======================================================================================================================
+TSharedPtr<FPhysicsControlAssetEditorData> FPhysicsControlAssetEditor::GetEditorData() const
+{
+	return EditorData;
+}
+
+
+//======================================================================================================================
+void FPhysicsControlAssetEditor::ExtendMenu()
+{
+	// In PhAT, entries in this function appear in the main Edit menu list
+	// For the moment, we don't have anything to add.
+}
+
+//======================================================================================================================
+void FPhysicsControlAssetEditor::InitToolMenuContext(FToolMenuContext& MenuContext)
+{
+	FAssetEditorToolkit::InitToolMenuContext(MenuContext);
+
+	UPhysicsControlAssetEditorToolMenuContext* PhysicsControlAssetEditorContext = 
+		NewObject<UPhysicsControlAssetEditorToolMenuContext>();
+	PhysicsControlAssetEditorContext->PhysicsControlAssetEditor = SharedThis(this);
+	MenuContext.AddObject(PhysicsControlAssetEditorContext);
+
+	UPersonaToolMenuContext* PersonaContext = NewObject<UPersonaToolMenuContext>();
+	PersonaContext->SetToolkit(GetPersonaToolkit());
+	MenuContext.AddObject(PersonaContext);
+
+	// Danny TODO I don't think we need this
+	MenuContext.AppendCommandList(ViewportCommandList);
+}
+
+//======================================================================================================================
+FText FPhysicsControlAssetEditor::GetSimulationToolTip() const
+{
+	if (EditorData->bNoGravitySimulation)
+	{
+		return FPhysicsControlAssetEditorCommands::Get().SimulationNoGravity->GetDescription();
+	}
+	else
+	{
+		return FPhysicsControlAssetEditorCommands::Get().Simulation->GetDescription();
+	}
+}
+
+//======================================================================================================================
+FSlateIcon FPhysicsControlAssetEditor::GetSimulationIcon() const
+{
+	if (EditorData->bNoGravitySimulation)
+	{
+		return FPhysicsControlAssetEditorCommands::Get().SimulationNoGravity->GetIcon();
+	}
+	else
+	{
+		return FPhysicsControlAssetEditorCommands::Get().Simulation->GetIcon();
+	}
+}
+
+//======================================================================================================================
+void FPhysicsControlAssetEditor::ExtendToolbar()
+{
+	// Extends the toolbar to the right of the save/locate buttons, selection of preview anim etc
+	struct Local
+	{
+		static TSharedRef< SWidget > FillSimulateOptions(TSharedRef<FUICommandList> InCommandList)
+		{
+			const bool bShouldCloseWindowAfterMenuSelection = true;
+			FMenuBuilder MenuBuilder(bShouldCloseWindowAfterMenuSelection, InCommandList);
+
+			const FPhysicsControlAssetEditorCommands& Commands = FPhysicsControlAssetEditorCommands::Get();
+
+			//selected simulation
+			MenuBuilder.BeginSection("SimulationOptions", LOCTEXT("SimulationOptionsHeader", "Simulation Options"));
+			{
+				MenuBuilder.AddMenuEntry(Commands.SimulationNoGravity);
+				MenuBuilder.AddMenuEntry(Commands.SimulationFloorCollision);
+			}
+			MenuBuilder.EndSection();
+
+			return MenuBuilder.MakeWidget();
+		}
+	};
+
+
+	FName ParentName;
+	static const FName MenuName = GetToolMenuToolbarName(ParentName);
+
+	UToolMenu* ToolMenu = UToolMenus::Get()->ExtendMenu(MenuName);
+	const FToolMenuInsert SectionInsertLocation("Asset", EToolMenuInsertType::After);
+
+	ToolMenu->AddDynamicSection("Persona", FNewToolMenuDelegate::CreateLambda([](UToolMenu* InToolMenu)
+		{
+			FPersonaModule& PersonaModule = FModuleManager::LoadModuleChecked<FPersonaModule>("Persona");
+			FPersonaModule::FCommonToolbarExtensionArgs Args;
+			Args.bReferencePose = true;
+			PersonaModule.AddCommonToolbarExtensions(InToolMenu, Args);
+		}), SectionInsertLocation);
+
+	const FPhysicsControlAssetEditorCommands& Commands = FPhysicsControlAssetEditorCommands::Get();
+
+	{
+		FToolMenuSection& Section = ToolMenu->AddSection("Compile", FText());
+		Section.AddEntry(FToolMenuEntry::InitToolBarButton(
+			Commands.Compile,
+			LOCTEXT("Compile_Label", "Compile"),
+			LOCTEXT("Compile_ToolTip", "Compiles all the data from this and any parent asset into runtime form")
+		));
+	}
+
+	ToolMenu->AddDynamicSection("Simulation", FNewToolMenuDelegate::CreateLambda([&Commands](UToolMenu* InToolMenu)
+		{
+			TSharedPtr<FPhysicsControlAssetEditor> PhysicsControlAssetEditor =
+				PhysicsControlAssetEditor::GetPhysicsControlAssetEditorFromToolContext(InToolMenu->Context);
+			if (PhysicsControlAssetEditor)
+			{
+				FToolMenuSection& Section = InToolMenu->AddSection("Simulation", FText());
+				// Simulate
+				Section.AddEntry(FToolMenuEntry::InitToolBarButton(
+					Commands.Simulation,
+					LOCTEXT("Simulation", "Simulate"),
+					TAttribute< FText >::Create(TAttribute< FText >::FGetter::CreateSP(
+						PhysicsControlAssetEditor.Get(), &FPhysicsControlAssetEditor::GetSimulationToolTip)),
+					TAttribute< FSlateIcon >::Create(TAttribute< FSlateIcon >::FGetter::CreateSP(
+						PhysicsControlAssetEditor.Get(), &FPhysicsControlAssetEditor::GetSimulationIcon))));
+
+				Section.AddEntry(FToolMenuEntry::InitComboButton("SimulationMode",
+					FUIAction(FExecuteAction(), FCanExecuteAction::CreateSP(
+						PhysicsControlAssetEditor.Get(), &FPhysicsControlAssetEditor::IsNotRunningSimulation)),
+					FOnGetContent::CreateLambda([WeakPhysicsControlAssetEditor = PhysicsControlAssetEditor.ToWeakPtr()]()
+						{
+							return Local::FillSimulateOptions(WeakPhysicsControlAssetEditor.Pin()->GetToolkitCommands());
+						}),
+					LOCTEXT("SimulateCombo_Label", "Simulate Options"),
+					LOCTEXT("SimulateComboToolTip", "Options for Simulation"),
+					FSlateIcon(),
+					true));
+			}
+		}), SectionInsertLocation);
+}
+
+//======================================================================================================================
+TSharedRef<SWidget> FPhysicsControlAssetEditor::MakeConstraintScaleWidget()
+{
+	return
+		SNew(SBox)
+		.HAlign(HAlign_Right)
+		[
+			SNew(SBox)
+				.Padding(FMargin(4.0f, 0.0f, 0.0f, 0.0f))
+				.WidthOverride(100.0f)
+				[
+					SNew(SNumericEntryBox<float>)
+						.Font(FAppStyle::GetFontStyle(TEXT("MenuItem.Font")))
+						.AllowSpin(true)
+						.MinSliderValue(0.0f)
+						.MaxSliderValue(4.0f)
+						.Value_Lambda([this]() { return EditorData->EditorOptions->ConstraintDrawSize; })
+						.OnValueChanged_Lambda([this](float InValue) 
+							{ 
+								EditorData->EditorOptions->ConstraintDrawSize = InValue;
+								RefreshPreviewViewport(); 
+							})
+						.OnValueCommitted_Lambda([this](float InValue, ETextCommit::Type InCommitType)
+							{
+								EditorData->EditorOptions->ConstraintDrawSize = InValue;
+								EditorData->EditorOptions->SaveConfig();
+								ViewportCommandList->WidgetInteraction(TEXT("ConstraintScaleWidget"));
+								RefreshPreviewViewport();
+							})
+				]
+		];
+}
+
+//======================================================================================================================
+TSharedRef<SWidget> FPhysicsControlAssetEditor::MakeCollisionOpacityWidget()
+{
+	return
+		SNew(SBox)
+		.HAlign(HAlign_Right)
+		[
+			SNew(SBox)
+				.Padding(FMargin(4.0f, 0.0f, 0.0f, 0.0f))
+				.WidthOverride(100.0f)
+				[
+					SNew(SNumericEntryBox<float>)
+						.Font(FAppStyle::GetFontStyle(TEXT("MenuItem.Font")))
+						.AllowSpin(true)
+						.MinValue(0.0f)
+						.MaxValue(1.0f)
+						.MinSliderValue(0.0f)
+						.MaxSliderValue(1.0f)
+						.Value_Lambda([this]() { return EditorData->EditorOptions->CollisionOpacity; })
+						.OnValueChanged_Lambda([this](float InValue) 
+							{ 
+								EditorData->EditorOptions->CollisionOpacity = InValue; 
+								RefreshPreviewViewport();
+							})
+						.OnValueCommitted_Lambda([this](float InValue, ETextCommit::Type InCommitType)
+							{
+								EditorData->EditorOptions->CollisionOpacity = InValue;
+								EditorData->EditorOptions->SaveConfig();
+								ViewportCommandList->WidgetInteraction(TEXT("CollisionOpacityWidget"));
+								RefreshPreviewViewport();
+							})
+				]
+		];
+}
+
+//======================================================================================================================
+void FPhysicsControlAssetEditor::ExtendViewportMenus()
+{
+	// This scope ensures that menus we add/extend are scoped to us - i.e. only dipslay in our
+	// editor, even if the menu is shared by other editors.
+	FToolMenuOwnerScoped OwnerScoped(this);
+
+	// Extend the "Character" menu in the viewport. By default that just contains the "Scene Elements" section.
+	static const FName CharacterMenuName("Persona.AnimViewportCharacterMenu");
+	UToolMenu* ExtendableCharacterMenu = UToolMenus::Get()->ExtendMenu(CharacterMenuName);
+	ExtendableCharacterMenu->AddDynamicSection(
+		"PhysicsControlCharacterMenu", FNewToolMenuDelegate::CreateLambda([](UToolMenu* CharacterMenu)
+			{
+				TSharedPtr<FPhysicsControlAssetEditor> PhysicsControlAssetEditor =
+					PhysicsControlAssetEditor::GetPhysicsControlAssetEditorFromToolContext(CharacterMenu->Context);
+				if (PhysicsControlAssetEditor)
+				{
+					FToolMenuSection& Section = CharacterMenu->AddSection(
+						"PhysicsAssetShowCommands",
+						LOCTEXT("PhysicsShowCommands", "Physics Rendering"),
+						FToolMenuInsert("AnimViewportSceneElements", EToolMenuInsertType::Before));
+					Section.AddSubMenu(TEXT("MeshRenderModeSubMenu"), LOCTEXT("MeshRenderModeSubMenu", "Mesh"), FText::GetEmpty(),
+						FNewToolMenuDelegate::CreateLambda([](UToolMenu* InSubMenu)
+							{
+								const FPhysicsControlAssetEditorCommands& Commands = FPhysicsControlAssetEditorCommands::Get();
+								{
+									FToolMenuSection& Section = InSubMenu->AddSection(
+										"PhysicsControlAssetEditorRenderingMode",
+										LOCTEXT("MeshRenderModeHeader", "Mesh Drawing (Edit)"));
+									Section.AddMenuEntry(Commands.MeshRenderingMode_Solid);
+									Section.AddMenuEntry(Commands.MeshRenderingMode_Wireframe);
+									Section.AddMenuEntry(Commands.MeshRenderingMode_None);
+								}
+
+								{
+									FToolMenuSection& Section = InSubMenu->AddSection(
+										"PhysicsControlAssetEditorRenderingModeSim",
+										LOCTEXT("MeshRenderModeSimHeader", "Mesh Drawing (Simulation)"));
+									Section.AddMenuEntry(Commands.MeshRenderingMode_Simulation_Solid);
+									Section.AddMenuEntry(Commands.MeshRenderingMode_Simulation_Wireframe);
+									Section.AddMenuEntry(Commands.MeshRenderingMode_Simulation_None);
+								}
+							}));
+
+					Section.AddSubMenu(TEXT("CollisionRenderModeSubMenu"),
+						LOCTEXT("CollisionRenderModeSubMenu", "Bodies"), FText::GetEmpty(),
+						FNewToolMenuDelegate::CreateLambda(
+							[WeakPhysicsControlAssetEditor = PhysicsControlAssetEditor.ToWeakPtr()](UToolMenu* InSubMenu)
+							{
+								const FPhysicsControlAssetEditorCommands& Commands = FPhysicsControlAssetEditorCommands::Get();
+								{
+									FToolMenuSection& Section = InSubMenu->AddSection(
+										"PhysicsControlAssetEditorCollisionRenderSettings",
+										LOCTEXT("CollisionRenderSettingsHeader", "Body Drawing"));
+									Section.AddEntry(FToolMenuEntry::InitWidget(TEXT("CollisionOpacity"),
+										WeakPhysicsControlAssetEditor.Pin()->MakeCollisionOpacityWidget(),
+										LOCTEXT("CollisionOpacityLabel", "Collision Opacity")));
+								}
+
+								{
+									FToolMenuSection& Section = InSubMenu->AddSection(
+										"PhysicsControlAssetEditorCollisionMode",
+										LOCTEXT("CollisionRenderModeHeader", "Body Drawing (Edit)"));
+									Section.AddMenuEntry(Commands.CollisionRenderingMode_Solid);
+									Section.AddMenuEntry(Commands.CollisionRenderingMode_Wireframe);
+									Section.AddMenuEntry(Commands.CollisionRenderingMode_SolidWireframe);
+									Section.AddMenuEntry(Commands.CollisionRenderingMode_None);
+								}
+
+								{
+									FToolMenuSection& Section = InSubMenu->AddSection(
+										"PhysicsControlAssetEditorCollisionModeSim",
+										LOCTEXT("CollisionRenderModeSimHeader", "Body Drawing (Simulation)"));
+									Section.AddMenuEntry(Commands.CollisionRenderingMode_Simulation_Solid);
+									Section.AddMenuEntry(Commands.CollisionRenderingMode_Simulation_Wireframe);
+									Section.AddMenuEntry(Commands.CollisionRenderingMode_Simulation_SolidWireframe);
+									Section.AddMenuEntry(Commands.CollisionRenderingMode_Simulation_None);
+								}
+							}));
+
+					Section.AddSubMenu(TEXT("ConstraintConstraintModeSubMenu"),
+						LOCTEXT("ConstraintConstraintModeSubMenu", "Constraints"), FText::GetEmpty(),
+						FNewToolMenuDelegate::CreateLambda(
+							[WeakPhysicsControlAssetEditor = PhysicsControlAssetEditor.ToWeakPtr()](UToolMenu* InSubMenu)
+							{
+								const FPhysicsControlAssetEditorCommands& Commands = FPhysicsControlAssetEditorCommands::Get();
+								{
+									FToolMenuSection& Section = InSubMenu->AddSection(
+										"PhysicsControlAssetEditorConstraints",
+										LOCTEXT("ConstraintHeader", "Constraints"));
+									Section.AddMenuEntry(Commands.DrawViolatedLimits);
+									Section.AddEntry(FToolMenuEntry::InitWidget(TEXT(
+										"ConstraintScale"), WeakPhysicsControlAssetEditor.Pin()->MakeConstraintScaleWidget(),
+										LOCTEXT("ConstraintScaleLabel", "Constraint Scale")));
+								}
+								{
+									FToolMenuSection& Section = InSubMenu->AddSection(
+										"PhysicsControlAssetEditorConstraintMode",
+										LOCTEXT("ConstraintRenderModeHeader", "Constraint Drawing (Edit)"));
+									Section.AddMenuEntry(Commands.ConstraintRenderingMode_None);
+									Section.AddMenuEntry(Commands.ConstraintRenderingMode_AllPositions);
+									Section.AddMenuEntry(Commands.ConstraintRenderingMode_AllLimits);
+								}
+
+								{
+									FToolMenuSection& Section = InSubMenu->AddSection(
+										"PhysicsControlAssetEditorConstraintModeSim",
+										LOCTEXT("ConstraintRenderModeSimHeader", "Constraint Drawing (Simulation)"));
+									Section.AddMenuEntry(Commands.ConstraintRenderingMode_Simulation_None);
+									Section.AddMenuEntry(Commands.ConstraintRenderingMode_Simulation_AllPositions);
+									Section.AddMenuEntry(Commands.ConstraintRenderingMode_Simulation_AllLimits);
+								}
+							}));
+				}
+			}));
+
+	// This extends the menu "Physics" in the viewport. This is empty by default.
+	static const FName PhysicsMenuName("Persona.AnimViewportPhysicsMenu");
+	UToolMenu* ExtendablePhysicsMenu = UToolMenus::Get()->ExtendMenu(PhysicsMenuName);
+	ExtendablePhysicsMenu->AddDynamicSection(
+		"AnimViewportPhysicsControlMenu", FNewToolMenuDelegate::CreateLambda([](UToolMenu* PhysicsMenu)
+			{
+				TSharedPtr<FPhysicsControlAssetEditor> PhysicsControlAssetEditor =
+					PhysicsControlAssetEditor::GetPhysicsControlAssetEditorFromToolContext(PhysicsMenu->Context);
+				if (PhysicsControlAssetEditor)
+				{
+					FToolMenuSection& Section = PhysicsMenu->AddSection(
+						"AnimViewportPhysicsMenu", LOCTEXT("ViewMenu_AnimViewportPhysicsMenu", "Physics Menu"));
+
+					FPropertyEditorModule& PropertyEditorModule =
+						FModuleManager::GetModuleChecked<FPropertyEditorModule>("PropertyEditor");
+
+					FDetailsViewArgs DetailsViewArgs;
+					DetailsViewArgs.bAllowSearch = false;
+					DetailsViewArgs.NameAreaSettings = FDetailsViewArgs::HideNameArea;
+
+					// I'd prefer to do this using SetIsPropertyVisibleDelegate but for some reason
+					// any delegate registered gets cleared by the time the visibilty check call
+					// gets made.
+					DetailsViewArgs.ShouldForceHideProperty.BindLambda([](const TSharedRef<FPropertyNode>& PropertyNode)->bool
+						{
+							const FName PropertyName = PropertyNode->GetProperty()->GetFName();
+							if (PropertyName == GET_MEMBER_NAME_CHECKED(UPhysicsAssetEditorOptions, PhysicsBlend) ||
+								PropertyName == GET_MEMBER_NAME_CHECKED(UPhysicsAssetEditorOptions, bRenderOnlySelectedConstraints) ||
+								PropertyName == GET_MEMBER_NAME_CHECKED(UPhysicsAssetEditorOptions, bUpdateJointsFromAnimation) ||
+								PropertyName == GET_MEMBER_NAME_CHECKED(UPhysicsAssetEditorOptions, PhysicsUpdateMode))
+							{
+								return true;
+							}
+							return false;
+						});
+
+
+					TSharedPtr<IDetailsView> OptionsDetailsView = PropertyEditorModule.CreateDetailView(DetailsViewArgs);
+					OptionsDetailsView->SetObject(PhysicsControlAssetEditor->GetEditorData()->EditorOptions);
+					OptionsDetailsView->OnFinishedChangingProperties().AddLambda(
+						[WeakPhysicsControlAssetEditor = PhysicsControlAssetEditor.ToWeakPtr()](const FPropertyChangedEvent& InEvent)
+						{
+							WeakPhysicsControlAssetEditor.Pin()->GetEditorData()->EditorOptions->SaveConfig();
+						});
+
+					Section.AddEntry(FToolMenuEntry::InitWidget("PhysicsEditorOptions", OptionsDetailsView.ToSharedRef(), FText()));
+				}
+			}));
+}
+
+//======================================================================================================================
+void FPhysicsControlAssetEditor::OnCompile()
+{
+	EditorData->PhysicsControlAsset->Compile();
+}
+
+//======================================================================================================================
+bool FPhysicsControlAssetEditor::IsCompilationNeeded()
+{
+	return EditorData->PhysicsControlAsset->IsCompilationNeeded();
+}
+
+//======================================================================================================================
+void FPhysicsControlAssetEditor::OnToggleSimulation()
+{
+	static const float PrevMaxFPS = GEngine->GetMaxFPS();
+
+	if (!EditorData->bRunningSimulation)
+	{
+		GEngine->SetMaxFPS((float)EditorData->EditorOptions->MaxFPS);
+	}
+	else
+	{
+		GEngine->SetMaxFPS(PrevMaxFPS);
+	}
+
+	EditorData->ToggleSimulation();
+}
+
+//======================================================================================================================
+void FPhysicsControlAssetEditor::OnToggleSimulationNoGravity()
+{
+	EditorData->bNoGravitySimulation = !EditorData->bNoGravitySimulation;
+}
+
+//======================================================================================================================
+bool FPhysicsControlAssetEditor::IsNoGravitySimulationEnabled() const
+{
+	return EditorData->bNoGravitySimulation;
+}
+
+//======================================================================================================================
+void FPhysicsControlAssetEditor::OnToggleSimulationFloorCollision()
+{
+	if (EditorData && EditorData->EditorOptions)
+	{
+		EditorData->EditorOptions->bSimulationFloorCollisionEnabled = !EditorData->EditorOptions->bSimulationFloorCollisionEnabled;
+
+		// Update collision for floor
+		if (PersonaToolkit)
+		{
+			TSharedRef<IPersonaPreviewScene> PersonaPreviewScene = PersonaToolkit->GetPreviewScene();
+
+			if (UStaticMeshComponent* FloorMeshComponent = const_cast<UStaticMeshComponent*>(PersonaPreviewScene->GetFloorMeshComponent()))
+			{
+				if (EditorData->EditorOptions->bSimulationFloorCollisionEnabled)
+				{
+					FloorMeshComponent->SetCollisionEnabled(ECollisionEnabled::QueryAndPhysics);
+				}
+				else
+				{
+					FloorMeshComponent->SetCollisionEnabled(ECollisionEnabled::NoCollision);
+				}
+			}
+		}
+	}
+}
+
+//======================================================================================================================
+bool FPhysicsControlAssetEditor::IsSimulationFloorCollisionEnabled() const
+{
+	return EditorData && EditorData->EditorOptions && EditorData->EditorOptions->bSimulationFloorCollisionEnabled;
+}
+
+//======================================================================================================================
+void FPhysicsControlAssetEditor::OnMeshRenderingMode(EPhysicsAssetEditorMeshViewMode Mode, bool bSimulation)
+{
+	if (bSimulation)
+	{
+		EditorData->EditorOptions->SimulationMeshViewMode = Mode;
+	}
+	else
+	{
+		EditorData->EditorOptions->MeshViewMode = Mode;
+	}
+
+	EditorData->EditorOptions->SaveConfig();
+
+	// Changing the mesh rendering mode requires the skeletal mesh component to change its render state, which is an operation
+	// which is deferred until after render. Hence we need to trigger another viewport refresh on the following frame.
+	RefreshPreviewViewport();
+}
+
+//======================================================================================================================
+bool FPhysicsControlAssetEditor::IsMeshRenderingMode(EPhysicsAssetEditorMeshViewMode Mode, bool bSimulation) const
+{
+	return Mode == EditorData->GetCurrentMeshViewMode(bSimulation);
+}
+
+//======================================================================================================================
+void FPhysicsControlAssetEditor::OnCollisionRenderingMode(EPhysicsAssetEditorCollisionViewMode Mode, bool bSimulation)
+{
+	if (bSimulation)
+	{
+		EditorData->EditorOptions->SimulationCollisionViewMode = Mode;
+	}
+	else
+	{
+		EditorData->EditorOptions->CollisionViewMode = Mode;
+	}
+	EditorData->EditorOptions->SaveConfig();
+	RefreshPreviewViewport();
+}
+
+//======================================================================================================================
+bool FPhysicsControlAssetEditor::IsCollisionRenderingMode(EPhysicsAssetEditorCollisionViewMode Mode, bool bSimulation) const
+{
+	return Mode == EditorData->GetCurrentCollisionViewMode(bSimulation);
+}
+
+//======================================================================================================================
+void FPhysicsControlAssetEditor::OnConstraintRenderingMode(EPhysicsAssetEditorConstraintViewMode Mode, bool bSimulation)
+{
+	if (bSimulation)
+	{
+		EditorData->EditorOptions->SimulationConstraintViewMode = Mode;
+	}
+	else
+	{
+		EditorData->EditorOptions->ConstraintViewMode = Mode;
+	}
+
+	EditorData->EditorOptions->SaveConfig();
+
+	RefreshPreviewViewport();
+}
+
+//======================================================================================================================
+bool FPhysicsControlAssetEditor::IsConstraintRenderingMode(EPhysicsAssetEditorConstraintViewMode Mode, bool bSimulation) const
+{
+	return Mode == EditorData->GetCurrentConstraintViewMode(bSimulation);
+}
+
+//======================================================================================================================
+void FPhysicsControlAssetEditor::ToggleDrawViolatedLimits()
+{
+	EditorData->EditorOptions->bDrawViolatedLimits = !EditorData->EditorOptions->bDrawViolatedLimits;
+	EditorData->EditorOptions->SaveConfig();
+	RefreshPreviewViewport();
+}
+
+//======================================================================================================================
+bool FPhysicsControlAssetEditor::IsDrawingViolatedLimits() const
+{
+	return EditorData->EditorOptions->bDrawViolatedLimits;
+}
+
+//======================================================================================================================
+bool FPhysicsControlAssetEditor::IsRunningSimulation() const
+{
+	return EditorData->bRunningSimulation;
+}
+
+//======================================================================================================================
+bool FPhysicsControlAssetEditor::IsNotRunningSimulation() const
+{
+	return !EditorData->bRunningSimulation;
+}
+
+//======================================================================================================================
+void FPhysicsControlAssetEditor::BindCommands()
+{
+	const FPhysicsControlAssetEditorCommands& Commands = FPhysicsControlAssetEditorCommands::Get();
+
+	ToolkitCommands->MapAction(
+		Commands.Compile,
+		FExecuteAction::CreateSP(this, &FPhysicsControlAssetEditor::OnCompile),
+		FCanExecuteAction(),
+		FIsActionChecked::CreateSP(this, &FPhysicsControlAssetEditor::IsCompilationNeeded));
+
+	ToolkitCommands->MapAction(
+		Commands.Simulation,
+		FExecuteAction::CreateSP(this, &FPhysicsControlAssetEditor::OnToggleSimulation),
+		FCanExecuteAction(),
+		FIsActionChecked::CreateSP(this, &FPhysicsControlAssetEditor::IsRunningSimulation));
+
+	ToolkitCommands->MapAction(
+		Commands.SimulationNoGravity,
+		FExecuteAction::CreateSP(this, &FPhysicsControlAssetEditor::OnToggleSimulationNoGravity),
+		FCanExecuteAction(),
+		FIsActionChecked::CreateSP(this, &FPhysicsControlAssetEditor::IsNoGravitySimulationEnabled));
+
+	ToolkitCommands->MapAction(
+		Commands.SimulationFloorCollision,
+		FExecuteAction::CreateSP(this, &FPhysicsControlAssetEditor::OnToggleSimulationFloorCollision),
+		FCanExecuteAction(),
+		FIsActionChecked::CreateSP(this, &FPhysicsControlAssetEditor::IsSimulationFloorCollisionEnabled));
+
+	// Viewport commands
+	ViewportCommandList = MakeShared<FUICommandList_Pinnable>();
+
+	ViewportCommandList->BeginGroup(TEXT("MeshRenderingMode"));
+
+	ViewportCommandList->MapAction(
+		Commands.MeshRenderingMode_Solid,
+		FExecuteAction::CreateSP(this, &FPhysicsControlAssetEditor::OnMeshRenderingMode, EPhysicsAssetEditorMeshViewMode::Solid, false),
+		FCanExecuteAction(),
+		FIsActionChecked::CreateSP(this, &FPhysicsControlAssetEditor::IsMeshRenderingMode, EPhysicsAssetEditorMeshViewMode::Solid, false));
+
+	ViewportCommandList->MapAction(
+		Commands.MeshRenderingMode_Wireframe,
+		FExecuteAction::CreateSP(this, &FPhysicsControlAssetEditor::OnMeshRenderingMode, EPhysicsAssetEditorMeshViewMode::Wireframe, false),
+		FCanExecuteAction(),
+		FIsActionChecked::CreateSP(this, &FPhysicsControlAssetEditor::IsMeshRenderingMode, EPhysicsAssetEditorMeshViewMode::Wireframe, false));
+
+	ViewportCommandList->MapAction(
+		Commands.MeshRenderingMode_None,
+		FExecuteAction::CreateSP(this, &FPhysicsControlAssetEditor::OnMeshRenderingMode, EPhysicsAssetEditorMeshViewMode::None, false),
+		FCanExecuteAction(),
+		FIsActionChecked::CreateSP(this, &FPhysicsControlAssetEditor::IsMeshRenderingMode, EPhysicsAssetEditorMeshViewMode::None, false));
+
+	ViewportCommandList->EndGroup();
+
+	ViewportCommandList->BeginGroup(TEXT("CollisionRenderingMode"));
+
+	ViewportCommandList->MapAction(
+		Commands.CollisionRenderingMode_Solid,
+		FExecuteAction::CreateSP(this, &FPhysicsControlAssetEditor::OnCollisionRenderingMode, EPhysicsAssetEditorCollisionViewMode::Solid, false),
+		FCanExecuteAction(),
+		FIsActionChecked::CreateSP(this, &FPhysicsControlAssetEditor::IsCollisionRenderingMode, EPhysicsAssetEditorCollisionViewMode::Solid, false));
+
+	ViewportCommandList->MapAction(
+		Commands.CollisionRenderingMode_Wireframe,
+		FExecuteAction::CreateSP(this, &FPhysicsControlAssetEditor::OnCollisionRenderingMode, EPhysicsAssetEditorCollisionViewMode::Wireframe, false),
+		FCanExecuteAction(),
+		FIsActionChecked::CreateSP(this, &FPhysicsControlAssetEditor::IsCollisionRenderingMode, EPhysicsAssetEditorCollisionViewMode::Wireframe, false));
+
+	ViewportCommandList->MapAction(
+		Commands.CollisionRenderingMode_SolidWireframe,
+		FExecuteAction::CreateSP(this, &FPhysicsControlAssetEditor::OnCollisionRenderingMode, EPhysicsAssetEditorCollisionViewMode::SolidWireframe, false),
+		FCanExecuteAction(),
+		FIsActionChecked::CreateSP(this, &FPhysicsControlAssetEditor::IsCollisionRenderingMode, EPhysicsAssetEditorCollisionViewMode::SolidWireframe, false));
+
+	ViewportCommandList->MapAction(
+		Commands.CollisionRenderingMode_None,
+		FExecuteAction::CreateSP(this, &FPhysicsControlAssetEditor::OnCollisionRenderingMode, EPhysicsAssetEditorCollisionViewMode::None, false),
+		FCanExecuteAction(),
+		FIsActionChecked::CreateSP(this, &FPhysicsControlAssetEditor::IsCollisionRenderingMode, EPhysicsAssetEditorCollisionViewMode::None, false));
+
+	ViewportCommandList->EndGroup();
+
+	ViewportCommandList->BeginGroup(TEXT("ConstraintRenderingMode"));
+
+	ViewportCommandList->MapAction(
+		Commands.ConstraintRenderingMode_None,
+		FExecuteAction::CreateSP(this, &FPhysicsControlAssetEditor::OnConstraintRenderingMode, EPhysicsAssetEditorConstraintViewMode::None, false),
+		FCanExecuteAction(),
+		FIsActionChecked::CreateSP(this, &FPhysicsControlAssetEditor::IsConstraintRenderingMode, EPhysicsAssetEditorConstraintViewMode::None, false));
+
+	ViewportCommandList->MapAction(
+		Commands.ConstraintRenderingMode_AllPositions,
+		FExecuteAction::CreateSP(this, &FPhysicsControlAssetEditor::OnConstraintRenderingMode, EPhysicsAssetEditorConstraintViewMode::AllPositions, false),
+		FCanExecuteAction(),
+		FIsActionChecked::CreateSP(this, &FPhysicsControlAssetEditor::IsConstraintRenderingMode, EPhysicsAssetEditorConstraintViewMode::AllPositions, false));
+
+	ViewportCommandList->MapAction(
+		Commands.ConstraintRenderingMode_AllLimits,
+		FExecuteAction::CreateSP(this, &FPhysicsControlAssetEditor::OnConstraintRenderingMode, EPhysicsAssetEditorConstraintViewMode::AllLimits, false),
+		FCanExecuteAction(),
+		FIsActionChecked::CreateSP(this, &FPhysicsControlAssetEditor::IsConstraintRenderingMode, EPhysicsAssetEditorConstraintViewMode::AllLimits, false));
+
+	ViewportCommandList->EndGroup();
+
+	ViewportCommandList->BeginGroup(TEXT("MeshRenderingMode_Simulation"));
+
+	ViewportCommandList->MapAction(
+		Commands.MeshRenderingMode_Simulation_Solid,
+		FExecuteAction::CreateSP(this, &FPhysicsControlAssetEditor::OnMeshRenderingMode, EPhysicsAssetEditorMeshViewMode::Solid, true),
+		FCanExecuteAction(),
+		FIsActionChecked::CreateSP(this, &FPhysicsControlAssetEditor::IsMeshRenderingMode, EPhysicsAssetEditorMeshViewMode::Solid, true));
+
+	ViewportCommandList->MapAction(
+		Commands.MeshRenderingMode_Simulation_Wireframe,
+		FExecuteAction::CreateSP(this, &FPhysicsControlAssetEditor::OnMeshRenderingMode, EPhysicsAssetEditorMeshViewMode::Wireframe, true),
+		FCanExecuteAction(),
+		FIsActionChecked::CreateSP(this, &FPhysicsControlAssetEditor::IsMeshRenderingMode, EPhysicsAssetEditorMeshViewMode::Wireframe, true));
+
+	ViewportCommandList->MapAction(
+		Commands.MeshRenderingMode_Simulation_None,
+		FExecuteAction::CreateSP(this, &FPhysicsControlAssetEditor::OnMeshRenderingMode, EPhysicsAssetEditorMeshViewMode::None, true),
+		FCanExecuteAction(),
+		FIsActionChecked::CreateSP(this, &FPhysicsControlAssetEditor::IsMeshRenderingMode, EPhysicsAssetEditorMeshViewMode::None, true));
+
+	ViewportCommandList->EndGroup();
+
+	ViewportCommandList->BeginGroup(TEXT("CollisionRenderingMode_Simulation"));
+
+	ViewportCommandList->MapAction(
+		Commands.CollisionRenderingMode_Simulation_Solid,
+		FExecuteAction::CreateSP(this, &FPhysicsControlAssetEditor::OnCollisionRenderingMode, EPhysicsAssetEditorCollisionViewMode::Solid, true),
+		FCanExecuteAction(),
+		FIsActionChecked::CreateSP(this, &FPhysicsControlAssetEditor::IsCollisionRenderingMode, EPhysicsAssetEditorCollisionViewMode::Solid, true));
+
+	ViewportCommandList->MapAction(
+		Commands.CollisionRenderingMode_Simulation_Wireframe,
+		FExecuteAction::CreateSP(this, &FPhysicsControlAssetEditor::OnCollisionRenderingMode, EPhysicsAssetEditorCollisionViewMode::Wireframe, true),
+		FCanExecuteAction(),
+		FIsActionChecked::CreateSP(this, &FPhysicsControlAssetEditor::IsCollisionRenderingMode, EPhysicsAssetEditorCollisionViewMode::Wireframe, true));
+
+	ViewportCommandList->MapAction(
+		Commands.CollisionRenderingMode_Simulation_SolidWireframe,
+		FExecuteAction::CreateSP(this, &FPhysicsControlAssetEditor::OnCollisionRenderingMode, EPhysicsAssetEditorCollisionViewMode::SolidWireframe, true),
+		FCanExecuteAction(),
+		FIsActionChecked::CreateSP(this, &FPhysicsControlAssetEditor::IsCollisionRenderingMode, EPhysicsAssetEditorCollisionViewMode::SolidWireframe, true));
+
+	ViewportCommandList->MapAction(
+		Commands.CollisionRenderingMode_Simulation_None,
+		FExecuteAction::CreateSP(this, &FPhysicsControlAssetEditor::OnCollisionRenderingMode, EPhysicsAssetEditorCollisionViewMode::None, true),
+		FCanExecuteAction(),
+		FIsActionChecked::CreateSP(this, &FPhysicsControlAssetEditor::IsCollisionRenderingMode, EPhysicsAssetEditorCollisionViewMode::None, true));
+
+	ViewportCommandList->EndGroup();
+
+	ViewportCommandList->BeginGroup(TEXT("ConstraintRenderingMode_Simulation"));
+
+	ViewportCommandList->MapAction(
+		Commands.ConstraintRenderingMode_Simulation_None,
+		FExecuteAction::CreateSP(this, &FPhysicsControlAssetEditor::OnConstraintRenderingMode, EPhysicsAssetEditorConstraintViewMode::None, true),
+		FCanExecuteAction(),
+		FIsActionChecked::CreateSP(this, &FPhysicsControlAssetEditor::IsConstraintRenderingMode, EPhysicsAssetEditorConstraintViewMode::None, true));
+
+	ViewportCommandList->MapAction(
+		Commands.ConstraintRenderingMode_Simulation_AllPositions,
+		FExecuteAction::CreateSP(this, &FPhysicsControlAssetEditor::OnConstraintRenderingMode, EPhysicsAssetEditorConstraintViewMode::AllPositions, true),
+		FCanExecuteAction(),
+		FIsActionChecked::CreateSP(this, &FPhysicsControlAssetEditor::IsConstraintRenderingMode, EPhysicsAssetEditorConstraintViewMode::AllPositions, true));
+
+	ViewportCommandList->MapAction(
+		Commands.ConstraintRenderingMode_Simulation_AllLimits,
+		FExecuteAction::CreateSP(this, &FPhysicsControlAssetEditor::OnConstraintRenderingMode, EPhysicsAssetEditorConstraintViewMode::AllLimits, true),
+		FCanExecuteAction(),
+		FIsActionChecked::CreateSP(this, &FPhysicsControlAssetEditor::IsConstraintRenderingMode, EPhysicsAssetEditorConstraintViewMode::AllLimits, true));
+
+	ViewportCommandList->EndGroup();
+
+	ViewportCommandList->MapAction(
+		Commands.DrawViolatedLimits,
+		FExecuteAction::CreateSP(this, &FPhysicsControlAssetEditor::ToggleDrawViolatedLimits),
+		FCanExecuteAction(),
+		FIsActionChecked::CreateSP(this, &FPhysicsControlAssetEditor::IsDrawingViolatedLimits));
+
+}
+
+//======================================================================================================================
+void FPhysicsControlAssetEditor::BuildMenuWidgetBody(FMenuBuilder& InMenuBuilder)
+{
+	// TODO Danny
+}
+
+//======================================================================================================================
+// Danny TODO note that these appear (on right-click) but are not yet functional
+void FPhysicsControlAssetEditor::BuildMenuWidgetSelection(FMenuBuilder& InMenuBuilder)
+{
+	InMenuBuilder.PushCommandList(GetToolkitCommands());
+	{
+		const FPhysicsControlAssetEditorCommands& Commands = FPhysicsControlAssetEditorCommands::Get();
+
+		InMenuBuilder.BeginSection("Selection", LOCTEXT("Selection", "Selection"));
+		InMenuBuilder.AddMenuEntry(Commands.SelectAllBodies);
+		InMenuBuilder.AddMenuEntry(Commands.SelectSimulatedBodies);
+		InMenuBuilder.AddMenuEntry(Commands.SelectKinematicBodies);
+		InMenuBuilder.EndSection();
+	}
+	InMenuBuilder.PopCommandList();
+}
+
+
+//======================================================================================================================
+// Danny TODO selection needs to be implemented/handled/made useful
+void FPhysicsControlAssetEditor::HandleSelectionChanged(const TArrayView<TSharedPtr<ISkeletonTreeItem>>& InSelectedItems, ESelectInfo::Type InSelectInfo)
+{
+	if (!bSelecting)
+	{
+		TGuardValue<bool> RecursionGuard(bSelecting, true);
+
+		// Always set the details customization object, regardless of selection type
+		// We do this because the tree may have been rebuilt and objects invalidated
+		TArray<UObject*> Objects;
+		Algo::TransformIf(InSelectedItems, Objects, [](const TSharedPtr<ISkeletonTreeItem>& InItem) { return InItem->GetObject() != nullptr; }, [](const TSharedPtr<ISkeletonTreeItem>& InItem) { return InItem->GetObject(); });
+
+		if (DetailsView.IsValid())
+		{
+			DetailsView->SetObjects(Objects);
+		}
+	}
+}
+
+//======================================================================================================================
+void FPhysicsControlAssetEditor::HandleGetFilterLabel(TArray<FText>& InOutItems) const
+{
+	if (SkeletonTreeBuilder->bShowBodies)
+	{
+		InOutItems.Add(LOCTEXT("BodiesFilterLabel", "Bodies"));
+	}
+
+	if (SkeletonTreeBuilder->bShowPrimitives)
+	{
+		InOutItems.Add(LOCTEXT("PrimitivesFilterLabel", "Primitives"));
+	}
+}
+
+//======================================================================================================================
+// Danny TODO
+void FPhysicsControlAssetEditor::HandleExtendFilterMenu(FMenuBuilder& InMenuBuilder)
+{
+}
+
+//======================================================================================================================
+// Danny TODO
+void FPhysicsControlAssetEditor::HandleExtendContextMenu(FMenuBuilder& InMenuBuilder)
+{
+	TArray<TSharedPtr<ISkeletonTreeItem>> SelectedItems = SkeletonTree->GetSelectedItems();
+	FSkeletonTreeSelection Selection(SelectedItems);
+
+	TArray<TSharedPtr<FSkeletonTreePhysicsControlBodyItem>> SelectedBodies = Selection.GetSelectedItems<FSkeletonTreePhysicsControlBodyItem>();
+	TArray<TSharedPtr<FSkeletonTreePhysicsControlShapeItem>> SelectedShapes = Selection.GetSelectedItems<FSkeletonTreePhysicsControlShapeItem>();
+	TArray<TSharedPtr<ISkeletonTreeItem>> SelectedBones = Selection.GetSelectedItemsByTypeId("FSkeletonTreeBoneItem");
+}
+
+
+//======================================================================================================================
+void FPhysicsControlAssetEditor::RegisterTabSpawners(const TSharedRef<FTabManager>& InTabManager)
+{
+	WorkspaceMenuCategory = InTabManager->AddLocalWorkspaceMenuCategory(
+		LOCTEXT("WorkspaceMenu_PhysicsControlAssetEditor", "PhysicsControlAssetEditor"));
+	FAssetEditorToolkit::RegisterTabSpawners(InTabManager);
+}
+
+//======================================================================================================================
+void FPhysicsControlAssetEditor::UnregisterTabSpawners(const TSharedRef<FTabManager>& InTabManager)
+{
+	FAssetEditorToolkit::UnregisterTabSpawners(InTabManager);
+}
+
+//======================================================================================================================
+FName FPhysicsControlAssetEditor::GetToolkitFName() const
+{
+	return FName("PhysicsControlAssetEditor");
+}
+
+//======================================================================================================================
+FText FPhysicsControlAssetEditor::GetBaseToolkitName() const
+{
+	return LOCTEXT("PhysicsControlAssetEditorAppLabel", "Physics Control Asset Editor");
+}
+
+//======================================================================================================================
+FText FPhysicsControlAssetEditor::GetToolkitName() const
+{
+	FFormatNamedArguments Args;
+	Args.Add(TEXT("AssetName"), FText::FromString(EditorData->PhysicsControlAsset->GetName()));
+	return FText::Format(LOCTEXT("PhysicsControlAssetEditorName", "{AssetName}"), Args);
+}
+
+//======================================================================================================================
+FLinearColor FPhysicsControlAssetEditor::GetWorldCentricTabColorScale() const
+{
+	return FLinearColor::White;
+}
+
+//======================================================================================================================
+FString FPhysicsControlAssetEditor::GetWorldCentricTabPrefix() const
+{
+	return TEXT("PhysicsControlAssetEditor");
+}
+
+//======================================================================================================================
+void FPhysicsControlAssetEditor::AddReferencedObjects(FReferenceCollector& Collector)
+{
+	EditorData->AddReferencedObjects(Collector);
+}
+
+//======================================================================================================================
+TStatId FPhysicsControlAssetEditor::GetStatId() const
+{
+	RETURN_QUICK_DECLARE_CYCLE_STAT(FPhysicsControlAssetEditor, STATGROUP_Tickables);
+}
+
+//======================================================================================================================
+TSharedRef<IPersonaToolkit> FPhysicsControlAssetEditor::GetPersonaToolkit() const
+{
+	return PersonaToolkit.ToSharedRef();
+}
+
+//======================================================================================================================
+// For inspiration here, see:
+// - PhysicsControlAssetEditor
+// - MLDeformerEditorToolkit
+// - IKRigToolkit
+// though note that they all do things differently!
+void FPhysicsControlAssetEditor::HandlePreviewSceneCreated(
+	const TSharedRef<IPersonaPreviewScene>& InPersonaPreviewScene)
+{
+	EditorData->PreviewScene = InPersonaPreviewScene;
+
+	AAnimationEditorPreviewActor* Actor = InPersonaPreviewScene->GetWorld()->SpawnActor<AAnimationEditorPreviewActor>(
+		AAnimationEditorPreviewActor::StaticClass(), FTransform::Identity);
+	Actor->SetFlags(RF_Transient);
+	InPersonaPreviewScene->SetActor(Actor);
+
+	// Create the preview skeletal mesh component
+	EditorData->EditorSkelComp = NewObject<UPhysicsControlAssetEditorSkeletalMeshComponent>(Actor);
+	EditorData->EditorSkelComp->SetCollisionEnabled(ECollisionEnabled::QueryAndPhysics);
+	EditorData->EditorSkelComp->EditorData = EditorData;
+	EditorData->EditorSkelComp->SetSkeletalMesh(EditorData->PhysicsControlAsset->GetPreviewMesh());
+	EditorData->EditorSkelComp->SetPhysicsAsset(EditorData->PhysicsControlAsset->GetPhysicsAsset(), true);
+	EditorData->EditorSkelComp->SetDisablePostProcessBlueprint(true);
+	EditorData->EditorSkelComp->Stop();
+
+	EditorData->EditorSkelComp->bSelectable = false;
+
+	// set root component, so we can attach to it. 
+	Actor->SetRootComponent(EditorData->EditorSkelComp);
+
+	// Set the skeletal mesh on the component, using the asset. Note that this will change if/when
+	// the asset doesn't hold a mesh.
+	USkeletalMesh* Mesh = EditorData->PhysicsControlAsset->GetPreviewMesh();
+	EditorData->EditorSkelComp->SetSkeletalMesh(Mesh);
+
+	// apply mesh to the preview scene
+	InPersonaPreviewScene->SetPreviewMeshComponent(EditorData->EditorSkelComp);
+	InPersonaPreviewScene->SetAllowMeshHitProxies(false);
+	InPersonaPreviewScene->SetAdditionalMeshesSelectable(false);
+	InPersonaPreviewScene->SetPreviewMesh(Mesh);
+	InPersonaPreviewScene->AddComponent(EditorData->EditorSkelComp, FTransform::Identity);
+
+	EditorData->PhysicsControlComponent = NewObject<UPhysicsControlComponent>(Actor);
+	InPersonaPreviewScene->AddComponent(EditorData->PhysicsControlComponent, FTransform::Identity);
+
+	// Register handle component
+	EditorData->MouseHandle->RegisterComponentWithWorld(InPersonaPreviewScene->GetWorld());
+	EditorData->EnableSimulation(false);
+}
+
+//======================================================================================================================
+void FPhysicsControlAssetEditor::HandleViewportCreated(const TSharedRef<IPersonaViewport>& InPersonaViewport)
+{
+	PersonaViewport = InPersonaViewport;
+	InPersonaViewport.Get().GetPinnedCommandList()->BindCommandList(ViewportCommandList.ToSharedRef());
+	InPersonaViewport.Get().GetPinnedCommandList()->RegisterCustomWidget(IPinnedCommandList::FOnGenerateCustomWidget::CreateSP(
+		this, &FPhysicsControlAssetEditor::MakeConstraintScaleWidget), TEXT("ConstraintScaleWidget"), LOCTEXT("ConstraintScaleLabel", "Constraint Scale"));
+	InPersonaViewport.Get().GetPinnedCommandList()->RegisterCustomWidget(IPinnedCommandList::FOnGenerateCustomWidget::CreateSP(
+		this, &FPhysicsControlAssetEditor::MakeCollisionOpacityWidget), TEXT("CollisionOpacityWidget"), LOCTEXT("CollisionOpacityLabel", "Collision Opacity"));
+}
+
+//======================================================================================================================
+void FPhysicsControlAssetEditor::ShowEmptyDetails() const
+{
+	DetailsView->SetObject(EditorData->PhysicsControlAsset);
+}
+
+//======================================================================================================================
+void FPhysicsControlAssetEditor::HandleDetailsCreated(const TSharedRef<class IDetailsView>& InDetailsView)
+{
+	DetailsView = InDetailsView;
+	DetailsView->OnFinishedChangingProperties().AddSP(
+		this, &FPhysicsControlAssetEditor::OnFinishedChangingDetails);
+	ShowEmptyDetails();
+}
+
+//======================================================================================================================
+void FPhysicsControlAssetEditor::OnFinishedChangingDetails(const FPropertyChangedEvent& PropertyChangedEvent)
+{
+	const bool bPreviewMeshChanged = 
+		PropertyChangedEvent.GetPropertyName() == UPhysicsControlAsset::GetPreviewMeshPropertyName();
+	if (bPreviewMeshChanged)
+	{
+		USkeletalMesh* Mesh = EditorData->PhysicsControlAsset->GetPreviewMesh();
+		EditorData->EditorSkelComp->SetSkeletalMesh(Mesh);
+		EditorData->CachePreviewMesh();
+	}
+}
+
+//======================================================================================================================
+void FPhysicsControlAssetEditor::RefreshHierachyTree()
+{
+	if (SkeletonTree.IsValid())
+	{
+		SkeletonTree->Refresh();
+	}
+}
+
+//======================================================================================================================
+void FPhysicsControlAssetEditor::RefreshPreviewViewport()
+{
+	if (PersonaToolkit.IsValid())
+	{
+		PersonaToolkit->GetPreviewScene()->InvalidateViews();
+	}
+}
+
+
+
+#undef LOCTEXT_NAMESPACE
+
diff --git a/Engine/Plugins/Experimental/PhysicsControl/Source/PhysicsControlEditor/Private/PhysicsControlAssetEditorAnimInstance.cpp b/Engine/Plugins/Experimental/PhysicsControl/Source/PhysicsControlEditor/Private/PhysicsControlAssetEditorAnimInstance.cpp
new file mode 100644
--- /dev/null
+++ b/Engine/Plugins/Experimental/PhysicsControl/Source/PhysicsControlEditor/Private/PhysicsControlAssetEditorAnimInstance.cpp
@@ -0,0 +1,50 @@
+// Copyright Epic Games, Inc. All Rights Reserved.
+
+#include "PhysicsControlAssetEditorAnimInstance.h"
+#include "PhysicsControlAssetEditorAnimInstanceProxy.h"
+
+/////////////////////////////////////////////////////
+// UPhysicsControlAssetEditorAnimInstance
+/////////////////////////////////////////////////////
+
+UPhysicsControlAssetEditorAnimInstance::UPhysicsControlAssetEditorAnimInstance(const FObjectInitializer& ObjectInitializer)
+	: Super(ObjectInitializer)
+{
+	bUseMultiThreadedAnimationUpdate = true;
+}
+
+FAnimInstanceProxy* UPhysicsControlAssetEditorAnimInstance::CreateAnimInstanceProxy()
+{
+	return new FPhysicsControlAssetEditorAnimInstanceProxy(this);
+}
+
+void UPhysicsControlAssetEditorAnimInstance::Grab(FName InBoneName, const FVector& Location, const FRotator& Rotation, bool bRotationConstrained)
+{
+	FPhysicsControlAssetEditorAnimInstanceProxy& Proxy = GetProxyOnGameThread<FPhysicsControlAssetEditorAnimInstanceProxy>();
+	Proxy.Grab(InBoneName, Location, Rotation, bRotationConstrained);
+}
+
+void UPhysicsControlAssetEditorAnimInstance::Ungrab()
+{
+	FPhysicsControlAssetEditorAnimInstanceProxy& Proxy = GetProxyOnGameThread<FPhysicsControlAssetEditorAnimInstanceProxy>();
+	Proxy.Ungrab();
+}
+
+void UPhysicsControlAssetEditorAnimInstance::UpdateHandleTransform(const FTransform& NewTransform)
+{
+	FPhysicsControlAssetEditorAnimInstanceProxy& Proxy = GetProxyOnGameThread<FPhysicsControlAssetEditorAnimInstanceProxy>();
+	Proxy.UpdateHandleTransform(NewTransform);
+}
+
+void UPhysicsControlAssetEditorAnimInstance::UpdateDriveSettings(bool bLinearSoft, float LinearStiffness, float LinearDamping)
+{
+	FPhysicsControlAssetEditorAnimInstanceProxy& Proxy = GetProxyOnGameThread<FPhysicsControlAssetEditorAnimInstanceProxy>();
+	Proxy.UpdateDriveSettings(bLinearSoft, LinearStiffness, LinearDamping);
+}
+
+void UPhysicsControlAssetEditorAnimInstance::CreateSimulationFloor(FBodyInstance* FloorBodyInstance, const FTransform& Transform)
+{
+	FPhysicsControlAssetEditorAnimInstanceProxy& Proxy = GetProxyOnGameThread<FPhysicsControlAssetEditorAnimInstanceProxy>();
+	Proxy.CreateSimulationFloor(FloorBodyInstance, Transform);
+}
+
diff --git a/Engine/Plugins/Experimental/PhysicsControl/Source/PhysicsControlEditor/Private/PhysicsControlAssetEditorAnimInstance.h b/Engine/Plugins/Experimental/PhysicsControl/Source/PhysicsControlEditor/Private/PhysicsControlAssetEditorAnimInstance.h
new file mode 100644
--- /dev/null
+++ b/Engine/Plugins/Experimental/PhysicsControl/Source/PhysicsControlEditor/Private/PhysicsControlAssetEditorAnimInstance.h
@@ -0,0 +1,30 @@
+// Copyright Epic Games, Inc. All Rights Reserved.
+
+/**
+ *
+ * Used by Preview in PhysicsControlAssetEditor, allows us to switch between immediate mode and vanilla physx
+ */
+
+#pragma once
+#include "AnimPreviewInstance.h"
+#include "PhysicsControlAssetEditorAnimInstance.generated.h"
+
+class UAnimSequence;
+
+UCLASS(transient, NotBlueprintable)
+class UPhysicsControlAssetEditorAnimInstance : public UAnimPreviewInstance
+{
+	GENERATED_UCLASS_BODY()
+
+	void Grab(FName InBoneName, const FVector& Location, const FRotator& Rotation, bool bRotationConstrained);
+	void Ungrab();
+	void UpdateHandleTransform(const FTransform& NewTransform);
+	void UpdateDriveSettings(bool bLinearSoft, float LinearStiffness, float LinearDamping);
+	void CreateSimulationFloor(FBodyInstance* FloorBodyInstance, const FTransform& Transform);
+
+protected:
+	virtual FAnimInstanceProxy* CreateAnimInstanceProxy() override;
+};
+
+
+
diff --git a/Engine/Plugins/Experimental/PhysicsControl/Source/PhysicsControlEditor/Private/PhysicsControlAssetEditorAnimInstanceProxy.cpp b/Engine/Plugins/Experimental/PhysicsControl/Source/PhysicsControlEditor/Private/PhysicsControlAssetEditorAnimInstanceProxy.cpp
new file mode 100644
--- /dev/null
+++ b/Engine/Plugins/Experimental/PhysicsControl/Source/PhysicsControlEditor/Private/PhysicsControlAssetEditorAnimInstanceProxy.cpp
@@ -0,0 +1,201 @@
+// Copyright Epic Games, Inc. All Rights Reserved.
+
+#include "PhysicsControlAssetEditorAnimInstanceProxy.h"
+#include "PhysicsControlAssetEditorAnimInstance.h"
+
+#include "Physics/ImmediatePhysics/ImmediatePhysicsActorHandle.h"
+#include "Physics/ImmediatePhysics/ImmediatePhysicsJointHandle.h"
+#include "Physics/ImmediatePhysics/ImmediatePhysicsSimulation.h"
+
+FPhysicsControlAssetEditorAnimInstanceProxy::FPhysicsControlAssetEditorAnimInstanceProxy()
+	: TargetActor(nullptr)
+	, HandleActor(nullptr)
+	, HandleJoint(nullptr)
+	, FloorActor(nullptr)
+{
+}
+
+FPhysicsControlAssetEditorAnimInstanceProxy::FPhysicsControlAssetEditorAnimInstanceProxy(UAnimInstance* InAnimInstance)
+	: FAnimPreviewInstanceProxy(InAnimInstance)
+	, TargetActor(nullptr)
+	, HandleActor(nullptr)
+	, HandleJoint(nullptr)
+	, FloorActor(nullptr)
+{
+}
+
+FPhysicsControlAssetEditorAnimInstanceProxy::~FPhysicsControlAssetEditorAnimInstanceProxy()
+{
+}
+
+void FPhysicsControlAssetEditorAnimInstanceProxy::Initialize(UAnimInstance* InAnimInstance)
+{
+	FAnimPreviewInstanceProxy::Initialize(InAnimInstance);
+	ConstructNodes();
+
+	UPhysicsAsset* PhysicsAsset = InAnimInstance->GetSkelMeshComponent()->GetPhysicsAsset();
+	if (PhysicsAsset != nullptr)
+	{
+		SolverSettings = PhysicsAsset->SolverSettings;
+		SolverIterations = PhysicsAsset->SolverIterations;
+	}
+
+	FloorActor = nullptr;
+}
+
+void FPhysicsControlAssetEditorAnimInstanceProxy::ConstructNodes()
+{
+	ComponentToLocalSpace.ComponentPose.SetLinkNode(&RagdollNode);
+	
+	RagdollNode.SimulationSpace = ESimulationSpace::WorldSpace;
+	RagdollNode.ActualAlpha = 1.0f;
+}
+
+FAnimNode_Base* FPhysicsControlAssetEditorAnimInstanceProxy::GetCustomRootNode()
+{
+	return &ComponentToLocalSpace;
+}
+
+void FPhysicsControlAssetEditorAnimInstanceProxy::GetCustomNodes(TArray<FAnimNode_Base*>& OutNodes)
+{
+	OutNodes.Add(&RagdollNode);
+	OutNodes.Add(&ComponentToLocalSpace);
+}
+
+void FPhysicsControlAssetEditorAnimInstanceProxy::UpdateAnimationNode(const FAnimationUpdateContext& InContext)
+{
+	if (CurrentAsset != nullptr)
+	{
+		FAnimPreviewInstanceProxy::UpdateAnimationNode(InContext);
+	}
+	else
+	{
+		ComponentToLocalSpace.Update_AnyThread(InContext);
+	}
+}
+
+bool FPhysicsControlAssetEditorAnimInstanceProxy::Evaluate_WithRoot(FPoseContext& Output, FAnimNode_Base* InRootNode)
+{
+	ImmediatePhysics::FSimulation* Simulation = RagdollNode.GetSimulation();
+	if (Simulation != nullptr)
+	{
+		Simulation->SetSolverSettings(
+			SolverSettings.FixedTimeStep,
+			SolverSettings.CullDistance,
+			SolverSettings.MaxDepenetrationVelocity,
+			SolverSettings.bUseLinearJointSolver,
+			SolverSettings.PositionIterations,
+			SolverSettings.VelocityIterations,
+			SolverSettings.ProjectionIterations,
+			SolverSettings.bUseManifolds);
+	}
+
+	if (CurrentAsset != nullptr)
+	{
+		return FAnimPreviewInstanceProxy::Evaluate_WithRoot(Output, InRootNode);
+	}
+	else
+	{
+		if ((InRootNode != NULL) && (InRootNode == GetRootNode()))
+		{
+			EvaluationCounter.Increment();
+		}
+
+		InRootNode->Evaluate_AnyThread(Output);
+		return true;
+	}
+}
+
+void FPhysicsControlAssetEditorAnimInstanceProxy::AddImpulseAtLocation(FVector Impulse, FVector Location, FName BoneName)
+{
+	RagdollNode.AddImpulseAtLocation(Impulse, Location, BoneName);
+}
+
+void FPhysicsControlAssetEditorAnimInstanceProxy::Grab(FName InBoneName, const FVector& Location, const FRotator& Rotation, bool bRotationConstrained)
+{
+	ImmediatePhysics::FSimulation* Simulation = RagdollNode.GetSimulation();
+
+	if (TargetActor != nullptr)
+	{
+		Ungrab();
+	}
+
+	for (int ActorIndex = 0; ActorIndex < Simulation->NumActors(); ++ActorIndex)
+	{
+		if (Simulation->GetActorHandle(ActorIndex)->GetName() == InBoneName)
+		{
+			TargetActor = Simulation->GetActorHandle(ActorIndex);
+			break;
+		}
+	}
+
+	if (TargetActor != nullptr)
+	{
+		FTransform HandleTransform = FTransform(Rotation, Location);
+		HandleActor = Simulation->CreateKinematicActor(nullptr, HandleTransform);
+		HandleActor->SetWorldTransform(HandleTransform);
+		HandleActor->SetKinematicTarget(HandleTransform);
+
+		HandleJoint = Simulation->CreateJoint(nullptr, TargetActor, HandleActor);
+	}
+}
+
+void FPhysicsControlAssetEditorAnimInstanceProxy::Ungrab()
+{
+	if (TargetActor != nullptr)
+	{
+		ImmediatePhysics::FSimulation* Simulation = RagdollNode.GetSimulation();
+		Simulation->DestroyJoint(HandleJoint);
+		Simulation->DestroyActor(HandleActor);
+		TargetActor = nullptr;
+		HandleActor = nullptr;
+		HandleJoint = nullptr;
+	}
+}
+
+void FPhysicsControlAssetEditorAnimInstanceProxy::UpdateHandleTransform(const FTransform& NewTransform)
+{
+	if (HandleActor != nullptr)
+	{
+		HandleActor->SetKinematicTarget(NewTransform);
+	}
+}
+
+void FPhysicsControlAssetEditorAnimInstanceProxy::UpdateDriveSettings(
+	bool bLinearSoft, float LinearStiffness, float LinearDamping)
+{
+	if (HandleJoint != nullptr)
+	{
+		HandleJoint->SetSoftLinearSettings(bLinearSoft, LinearStiffness, LinearDamping);
+	}
+}
+
+void FPhysicsControlAssetEditorAnimInstanceProxy::CreateSimulationFloor(
+	FBodyInstance* FloorBodyInstance, const FTransform& Transform)
+{
+	using namespace Chaos;
+
+	DestroySimulationFloor();
+
+	ImmediatePhysics::FSimulation* Simulation = RagdollNode.GetSimulation();
+	if (Simulation != nullptr)
+	{
+		FloorActor = Simulation->CreateKinematicActor(FloorBodyInstance, Transform);
+		if (FloorActor != nullptr)
+		{
+			Simulation->AddToCollidingPairs(FloorActor);
+		}
+	}
+}
+
+void FPhysicsControlAssetEditorAnimInstanceProxy::DestroySimulationFloor()
+{
+	using namespace Chaos;
+	ImmediatePhysics::FSimulation* Simulation = RagdollNode.GetSimulation();
+	if ((Simulation != nullptr) && (FloorActor != nullptr))
+	{
+		Simulation->DestroyActor(FloorActor);
+		FloorActor = nullptr;
+	}
+}
+
diff --git a/Engine/Plugins/Experimental/PhysicsControl/Source/PhysicsControlEditor/Private/PhysicsControlAssetEditorAnimInstanceProxy.h b/Engine/Plugins/Experimental/PhysicsControl/Source/PhysicsControlEditor/Private/PhysicsControlAssetEditorAnimInstanceProxy.h
new file mode 100644
--- /dev/null
+++ b/Engine/Plugins/Experimental/PhysicsControl/Source/PhysicsControlEditor/Private/PhysicsControlAssetEditorAnimInstanceProxy.h
@@ -0,0 +1,58 @@
+// Copyright Epic Games, Inc. All Rights Reserved.
+
+#pragma once
+
+#include "AnimPreviewInstance.h"
+#include "Animation/AnimNode_LinkedInputPose.h"
+#include "Animation/AnimNode_SequencePlayer.h"
+#include "Animation/AnimNodeSpaceConversions.h"
+#include "BoneControllers/AnimNode_RigidBody.h"
+#include "Physics/ImmediatePhysics/ImmediatePhysicsDeclares.h"
+
+#include "PhysicsControlAssetEditorAnimInstanceProxy.generated.h"
+
+class UAnimSequence;
+
+/** Proxy override for this UAnimInstance-derived class */
+USTRUCT()
+struct FPhysicsControlAssetEditorAnimInstanceProxy : public FAnimPreviewInstanceProxy
+{
+	GENERATED_BODY()
+
+public:
+	FPhysicsControlAssetEditorAnimInstanceProxy();
+
+	FPhysicsControlAssetEditorAnimInstanceProxy(UAnimInstance* InAnimInstance);
+
+	virtual ~FPhysicsControlAssetEditorAnimInstanceProxy();
+
+	virtual void Initialize(UAnimInstance* InAnimInstance) override;
+	virtual FAnimNode_Base* GetCustomRootNode() override;
+	virtual void GetCustomNodes(TArray<FAnimNode_Base*>& OutNodes) override;
+
+	virtual void UpdateAnimationNode(const FAnimationUpdateContext& InContext) override;
+	virtual bool Evaluate_WithRoot(FPoseContext& Output, FAnimNode_Base* InRootNode) override;
+
+	virtual void AddImpulseAtLocation(FVector Impulse, FVector Location, FName BoneName = NAME_None) override;
+
+	void Grab(FName InBoneName, const FVector& Location, const FRotator& Rotation, bool bRotationConstrained);
+	void Ungrab();
+	void UpdateHandleTransform(const FTransform& NewTransform);
+	void UpdateDriveSettings(bool bLinearSoft, float LinearStiffness, float LinearDamping);
+
+	void CreateSimulationFloor(FBodyInstance* FloorBodyInstance, const FTransform& Transform);
+	void DestroySimulationFloor();
+
+private:
+	void ConstructNodes();
+
+	FAnimNode_RigidBody RagdollNode;
+ 	FAnimNode_ConvertComponentToLocalSpace ComponentToLocalSpace;
+	FPhysicsAssetSolverSettings SolverSettings;
+	FSolverIterations SolverIterations;
+
+	ImmediatePhysics::FActorHandle* TargetActor;
+	ImmediatePhysics::FActorHandle* HandleActor;
+	ImmediatePhysics::FJointHandle* HandleJoint;
+	ImmediatePhysics::FActorHandle* FloorActor;
+};
diff --git a/Engine/Plugins/Experimental/PhysicsControl/Source/PhysicsControlEditor/Private/PhysicsControlAssetEditorCommands.cpp b/Engine/Plugins/Experimental/PhysicsControl/Source/PhysicsControlEditor/Private/PhysicsControlAssetEditorCommands.cpp
new file mode 100644
--- /dev/null
+++ b/Engine/Plugins/Experimental/PhysicsControl/Source/PhysicsControlEditor/Private/PhysicsControlAssetEditorCommands.cpp
@@ -0,0 +1,48 @@
+// Copyright Epic Games, Inc. All Rights Reserved.
+
+#include "PhysicsControlAssetEditorCommands.h"
+
+#include "Framework/Commands/InputChord.h"
+#include "Framework/Commands/UICommandInfo.h"
+#include "GenericPlatform/GenericApplication.h"
+#include "InputCoreTypes.h"
+
+#define LOCTEXT_NAMESPACE "PhysicsControlAssetEditorCommands"
+
+void FPhysicsControlAssetEditorCommands::RegisterCommands()
+{
+	UI_COMMAND(Compile, "Compile", "Compiles all the data from this and any parent asset into runtime form", EUserInterfaceActionType::RadioButton, FInputChord());
+	UI_COMMAND(Simulation, "Simulate", "Previews Physics Simulation", EUserInterfaceActionType::RadioButton, FInputChord());
+	UI_COMMAND(SimulationNoGravity, "No Gravity Simulation", "Run Physics Simulation without gravity.", EUserInterfaceActionType::ToggleButton, FInputChord());
+	UI_COMMAND(SimulationFloorCollision, "Enable Floor Collisions", "Run Physics Simulation with collisions between physics bodies and the floor enabled.", EUserInterfaceActionType::ToggleButton, FInputChord());
+
+	UI_COMMAND(SelectAllBodies, "Select All Bodies", "Select All Bodies", EUserInterfaceActionType::Button, FInputChord(EModifierKey::Control, EKeys::A));
+	UI_COMMAND(SelectSimulatedBodies, "Select Simulated Bodies", "Select Simulated Bodies", EUserInterfaceActionType::Button, FInputChord(EModifierKey::Control, EKeys::J));
+	UI_COMMAND(SelectKinematicBodies, "Select Kinematic Bodies", "Select Kinematic Bodies", EUserInterfaceActionType::Button, FInputChord(EModifierKey::Control, EKeys::K));
+
+	UI_COMMAND(MeshRenderingMode_Solid, "Solid", "Solid Mesh Rendering Mode (Edit)", EUserInterfaceActionType::RadioButton, FInputChord());
+	UI_COMMAND(MeshRenderingMode_Wireframe, "Wireframe", "Wireframe Mesh Rendering Mode (Edit)", EUserInterfaceActionType::RadioButton, FInputChord());
+	UI_COMMAND(MeshRenderingMode_None, "None", "No Mesh Rendering Mode (Edit)", EUserInterfaceActionType::RadioButton, FInputChord());
+	UI_COMMAND(CollisionRenderingMode_Solid, "Solid", "Solid Collision Rendering Mode (Edit)", EUserInterfaceActionType::RadioButton, FInputChord());
+	UI_COMMAND(CollisionRenderingMode_Wireframe, "Wireframe", "Wireframe Collision Rendering Mode (Edit)", EUserInterfaceActionType::RadioButton, FInputChord());
+	UI_COMMAND(CollisionRenderingMode_SolidWireframe, "Solid + Wireframe", "Solid + Wireframe Collision Rendering Mode (Edit)", EUserInterfaceActionType::RadioButton, FInputChord());
+	UI_COMMAND(CollisionRenderingMode_None, "None", "No Collision Rendering Mode (Edit)", EUserInterfaceActionType::RadioButton, FInputChord());
+	UI_COMMAND(ConstraintRenderingMode_None, "None", "No Constraint Rendering Mode (Edit)", EUserInterfaceActionType::RadioButton, FInputChord());
+	UI_COMMAND(ConstraintRenderingMode_AllPositions, "All Positions", "All Positions Constraint Rendering Mode (Edit)", EUserInterfaceActionType::RadioButton, FInputChord());
+	UI_COMMAND(ConstraintRenderingMode_AllLimits, "All Limits", "All Limits Constraint Rendering Mode (Edit)", EUserInterfaceActionType::RadioButton, FInputChord());
+
+	UI_COMMAND(MeshRenderingMode_Simulation_Solid, "Solid", "Solid Mesh Rendering Mode (Simulation)", EUserInterfaceActionType::RadioButton, FInputChord());
+	UI_COMMAND(MeshRenderingMode_Simulation_Wireframe, "Wireframe", "Wireframe Mesh Rendering Mode (Simulation)", EUserInterfaceActionType::RadioButton, FInputChord());
+	UI_COMMAND(MeshRenderingMode_Simulation_None, "None", "No Mesh Rendering Mode (Simulation)", EUserInterfaceActionType::RadioButton, FInputChord());
+	UI_COMMAND(CollisionRenderingMode_Simulation_Solid, "Solid", "Solid Collision Rendering Mode (Simulation)", EUserInterfaceActionType::RadioButton, FInputChord());
+	UI_COMMAND(CollisionRenderingMode_Simulation_Wireframe, "Wireframe", "Wireframe Collision Rendering Mode (Simulation)", EUserInterfaceActionType::RadioButton, FInputChord());
+	UI_COMMAND(CollisionRenderingMode_Simulation_SolidWireframe, "Solid + Wireframe", "Solid + Wireframe Collision Rendering Mode (Simulation)", EUserInterfaceActionType::RadioButton, FInputChord());
+	UI_COMMAND(CollisionRenderingMode_Simulation_None, "None", "No Collision Rendering Mode (Simulation)", EUserInterfaceActionType::RadioButton, FInputChord());
+	UI_COMMAND(ConstraintRenderingMode_Simulation_None, "None", "No Constraint Rendering Mode (Simulation)", EUserInterfaceActionType::RadioButton, FInputChord());
+	UI_COMMAND(ConstraintRenderingMode_Simulation_AllPositions, "All Positions", "All Positions Constraint Rendering Mode (Simulation)", EUserInterfaceActionType::RadioButton, FInputChord());
+	UI_COMMAND(ConstraintRenderingMode_Simulation_AllLimits, "All Limits", "All Limits Constraint Rendering Mode (Simulation)", EUserInterfaceActionType::RadioButton, FInputChord());
+
+	UI_COMMAND(DrawViolatedLimits, "Draw Violated Limits", "Draw Violated Limits", EUserInterfaceActionType::ToggleButton, FInputChord());
+}
+
+#undef LOCTEXT_NAMESPACE
diff --git a/Engine/Plugins/Experimental/PhysicsControl/Source/PhysicsControlEditor/Private/PhysicsControlAssetEditorCommands.h b/Engine/Plugins/Experimental/PhysicsControl/Source/PhysicsControlEditor/Private/PhysicsControlAssetEditorCommands.h
new file mode 100644
--- /dev/null
+++ b/Engine/Plugins/Experimental/PhysicsControl/Source/PhysicsControlEditor/Private/PhysicsControlAssetEditorCommands.h
@@ -0,0 +1,64 @@
+// Copyright Epic Games, Inc. All Rights Reserved.
+
+#pragma once
+
+#include "Framework/Commands/Commands.h"
+#include "Internationalization/Internationalization.h"
+#include "Styling/AppStyle.h"
+#include "Templates/SharedPointer.h"
+#include "UObject/NameTypes.h"
+#include "UObject/UnrealNames.h"
+
+class FUICommandInfo;
+
+class FPhysicsControlAssetEditorCommands : public TCommands<FPhysicsControlAssetEditorCommands>
+{
+public:
+	/** Constructor */
+	FPhysicsControlAssetEditorCommands()
+		: TCommands<FPhysicsControlAssetEditorCommands>(
+			"PhysicsControlAssetEditor", 
+			NSLOCTEXT("Contexts", "PhysicsControlAssetEditor", "PhysicsControlAssetEditor"), 
+			NAME_None, FAppStyle::GetAppStyleSetName())
+	{
+	}
+	
+	/** See tooltips in cpp for documentation */
+	TSharedPtr<FUICommandInfo> Compile;
+	TSharedPtr<FUICommandInfo> Simulation;
+	TSharedPtr<FUICommandInfo> SimulationNoGravity;
+	TSharedPtr<FUICommandInfo> SimulationFloorCollision;
+
+	/** Selection commands */
+	TSharedPtr<FUICommandInfo> SelectAllBodies;
+	TSharedPtr<FUICommandInfo> SelectSimulatedBodies;
+	TSharedPtr<FUICommandInfo> SelectKinematicBodies;
+
+	/** Rendering commands */
+	TSharedPtr<FUICommandInfo> MeshRenderingMode_Solid;
+	TSharedPtr<FUICommandInfo> MeshRenderingMode_Wireframe;
+	TSharedPtr<FUICommandInfo> MeshRenderingMode_None;
+	TSharedPtr<FUICommandInfo> CollisionRenderingMode_Solid;
+	TSharedPtr<FUICommandInfo> CollisionRenderingMode_Wireframe;
+	TSharedPtr<FUICommandInfo> CollisionRenderingMode_SolidWireframe;
+	TSharedPtr<FUICommandInfo> CollisionRenderingMode_None;
+	TSharedPtr<FUICommandInfo> ConstraintRenderingMode_None;
+	TSharedPtr<FUICommandInfo> ConstraintRenderingMode_AllPositions;
+	TSharedPtr<FUICommandInfo> ConstraintRenderingMode_AllLimits;
+
+	TSharedPtr<FUICommandInfo> MeshRenderingMode_Simulation_Solid;
+	TSharedPtr<FUICommandInfo> MeshRenderingMode_Simulation_Wireframe;
+	TSharedPtr<FUICommandInfo> MeshRenderingMode_Simulation_None;
+	TSharedPtr<FUICommandInfo> CollisionRenderingMode_Simulation_Solid;
+	TSharedPtr<FUICommandInfo> CollisionRenderingMode_Simulation_Wireframe;
+	TSharedPtr<FUICommandInfo> CollisionRenderingMode_Simulation_SolidWireframe;
+	TSharedPtr<FUICommandInfo> CollisionRenderingMode_Simulation_None;
+	TSharedPtr<FUICommandInfo> ConstraintRenderingMode_Simulation_None;
+	TSharedPtr<FUICommandInfo> ConstraintRenderingMode_Simulation_AllPositions;
+	TSharedPtr<FUICommandInfo> ConstraintRenderingMode_Simulation_AllLimits;
+
+	TSharedPtr<FUICommandInfo> DrawViolatedLimits;
+
+	/** Initialize commands */
+	virtual void RegisterCommands() override;
+};
diff --git a/Engine/Plugins/Experimental/PhysicsControl/Source/PhysicsControlEditor/Private/PhysicsControlAssetEditorData.cpp b/Engine/Plugins/Experimental/PhysicsControl/Source/PhysicsControlEditor/Private/PhysicsControlAssetEditorData.cpp
--- a/Engine/Plugins/Experimental/PhysicsControl/Source/PhysicsControlEditor/Private/PhysicsControlAssetEditorData.cpp
+++ b/Engine/Plugins/Experimental/PhysicsControl/Source/PhysicsControlEditor/Private/PhysicsControlAssetEditorData.cpp
@@ -2,14 +2,47 @@
 
 #include "PhysicsControlAssetEditorData.h"
 #include "PhysicsControlAsset.h"
+#include "PhysicsControlAssetEditorSkeletalMeshComponent.h"
+#include "PhysicsControlAssetEditorPhysicsHandleComponent.h"
+#include "PhysicsControlComponent.h"
+
+#include "Components/StaticMeshComponent.h"
+#include "IPersonaPreviewScene.h"
+#include "PhysicsEngine/PhysicsAsset.h"
+#include "PhysicsEngine/SkeletalBodySetup.h"
+#include "AnimPreviewInstance.h"
 #include "Engine/SkeletalMesh.h"
 #include "Misc/MessageDialog.h"
+#include "UObject/Package.h"
 
 #define LOCTEXT_NAMESPACE "PhysicsControlAssetEditorData"
 
 //======================================================================================================================
 FPhysicsControlAssetEditorData::FPhysicsControlAssetEditorData()
+	: bSuspendSelectionBroadcast(false)
+	, InsideSelChange(0)
 {
+	bRunningSimulation = false;
+	bNoGravitySimulation = false;
+
+	bManipulating = false;
+
+	// Construct mouse handle
+	MouseHandle = NewObject<UPhysicsControlAssetEditorPhysicsHandleComponent>();
+
+	// Construct sim options.
+	EditorOptions = NewObject<UPhysicsAssetEditorOptions>(
+		GetTransientPackage(), 
+		MakeUniqueObjectName(
+			GetTransientPackage(), UPhysicsAssetEditorOptions::StaticClass(), 
+			FName(TEXT("EditorOptions"))), RF_Transactional);
+	check(EditorOptions);
+
+	EditorOptions->LoadConfig();
+
+	// Set some options that we don't want to have the user modify
+	EditorOptions->bUpdateJointsFromAnimation = true;
+	EditorOptions->PhysicsUpdateMode = EPhysicsTransformUpdateMode::ComponentTransformIsKinematic;
 }
 
 //======================================================================================================================
@@ -17,9 +50,11 @@ void FPhysicsControlAssetEditorData::Initialize(const TSharedRef<IPersonaPreview
 {
 	PreviewScene = InPreviewScene;
 
+	bRunningSimulation = false;
+	bNoGravitySimulation = false;
+
 	EditorSkelComp = nullptr;
 	PhysicsControlComponent = nullptr;
-	FSoftObjectPath PreviewMeshStringRef = PhysicsControlAsset->PreviewSkeletalMesh.ToSoftObjectPath();
 
 	// Support undo/redo
 	PhysicsControlAsset->SetFlags(RF_Transactional);
@@ -28,7 +63,8 @@ void FPhysicsControlAssetEditorData::Initialize(const TSharedRef<IPersonaPreview
 //======================================================================================================================
 void FPhysicsControlAssetEditorData::CachePreviewMesh()
 {
-	USkeletalMesh* PreviewMesh = PhysicsControlAsset->PreviewSkeletalMesh.LoadSynchronous();
+	// This loads it if necessary
+	USkeletalMesh* PreviewMesh = PhysicsControlAsset->GetPreviewMesh();
 
 	if (PreviewMesh == nullptr)
 	{
@@ -38,10 +74,10 @@ void FPhysicsControlAssetEditorData::CachePreviewMesh()
 		PreviewMesh = (USkeletalMesh*)StaticLoadObject(USkeletalMesh::StaticClass(), NULL, TEXT("/Engine/EngineMeshes/SkeletalCube.SkeletalCube"), NULL, LOAD_None, NULL);
 		check(PreviewMesh);
 
-		PhysicsControlAsset->PreviewSkeletalMesh = PreviewMesh;
+		PhysicsControlAsset->SetPreviewMesh(PreviewMesh);
 
 		FMessageDialog::Open(EAppMsgType::Ok, FText::Format(
-			LOCTEXT("Error_PhysicsControlAssetHasNoSkelMesh", "Warning: Physics Control Profile Asset has no skeletal mesh assigned.\nFor now, a simple default skeletal mesh ({0}) will be used.\nYou can fix this by opening the asset and choosing another skeletal mesh from the toolbar."),
+			LOCTEXT("Error_PhysicsControlAssetHasNoSkelMesh", "Warning: Physics Control Asset has no skeletal mesh assigned.\nFor now, a simple default skeletal mesh ({0}) will be used.\nYou can fix this by opening the asset and choosing another skeletal mesh from the toolbar."),
 			FText::FromString(PreviewMesh->GetFullName())));
 	}
 	else if (PreviewMesh->GetSkeleton() == nullptr)
@@ -50,13 +86,300 @@ void FPhysicsControlAssetEditorData::CachePreviewMesh()
 		PreviewMesh = (USkeletalMesh*)StaticLoadObject(USkeletalMesh::StaticClass(), NULL, TEXT("/Engine/EngineMeshes/SkeletalCube.SkeletalCube"), NULL, LOAD_None, NULL);
 		check(PreviewMesh);
 
-		PhysicsControlAsset->PreviewSkeletalMesh = PreviewMesh;
+		PhysicsControlAsset->SetPreviewMesh(PreviewMesh);
 
 		FMessageDialog::Open(EAppMsgType::Ok, FText::Format(
-			LOCTEXT("Error_PhysicsControlAssetHasNoSkelMeshSkeleton", "Warning: Physics Control Profile Asset has a skeletal mesh with no skeleton assigned.\nFor now, a simple default skeletal mesh ({0}) will be used.\nYou can fix this by opening the asset and choosing another skeletal mesh from the toolbar, or repairing the skeleton."),
+			LOCTEXT("Error_PhysicsControlAssetHasNoSkelMeshSkeleton", "Warning: Physics Control Asset has a skeletal mesh with no skeleton assigned.\nFor now, a simple default skeletal mesh ({0}) will be used.\nYou can fix this by opening the asset and choosing another skeletal mesh from the toolbar, or repairing the skeleton."),
 			FText::FromString(PreviewMesh->GetFullName())));
 	}
 }
 
+//======================================================================================================================
+void FPhysicsControlAssetEditorData::ToggleSimulation()
+{
+	if (!bManipulating)
+	{
+		EnableSimulation(!bRunningSimulation);
+	}
+}
+
+
+//======================================================================================================================
+void FPhysicsControlAssetEditorData::EnableSimulation(bool bEnableSimulation)
+{
+	// keep the EditorSkelComp animation asset if any set 
+	UAnimationAsset* PreviewAnimationAsset = nullptr;
+	if (EditorSkelComp->PreviewInstance)
+	{
+		PreviewAnimationAsset = EditorSkelComp->PreviewInstance->CurrentAsset;
+	}
+
+	UPhysicsAsset* PA = PhysicsControlAsset->GetPhysicsAsset();
+
+	if (bEnableSimulation && PA)
+	{
+		// in Chaos, we have to manipulate the RBAN node in the Anim Instance (at least until we get
+		// SkelMeshComp implemented)
+		const bool bUseRBANSolver = (PA->SolverType == EPhysicsAssetSolverType::RBAN);
+		MouseHandle->SetAnimInstanceMode(bUseRBANSolver);
+
+		if (!bUseRBANSolver)
+		{
+			// We should not already have an instance (destroyed when stopping sim).
+			EditorSkelComp->SetCollisionEnabled(ECollisionEnabled::QueryAndPhysics);
+			EditorSkelComp->SetSimulatePhysics(true);
+			EditorSkelComp->ResetAllBodiesSimulatePhysics();
+			EditorSkelComp->SetPhysicsBlendWeight(EditorOptions->PhysicsBlend);
+			// Make it start simulating
+			EditorSkelComp->WakeAllRigidBodies();
+
+			PhysicsControlComponent->PhysicsControlAsset = PhysicsControlAsset;
+			PhysicsControlComponent->CreateControlsAndBodyModifiersFromControlProfileAsset(
+				EditorSkelComp, nullptr, FName());
+		}
+		else
+		{
+			// Enable the PreviewInstance (containing the AnimNode_RigidBody)
+			EditorSkelComp->SetAnimationMode(EAnimationMode::AnimationCustomMode);
+			EditorSkelComp->InitAnim(true);
+
+			// Disable main solver physics
+			EditorSkelComp->SetAllBodiesSimulatePhysics(false);
+
+			// make sure we enable the preview animation is any compatible with the skeleton
+			if (PreviewAnimationAsset && EditorSkelComp->GetSkeletalMeshAsset() && 
+				PreviewAnimationAsset->GetSkeleton() == EditorSkelComp->GetSkeletalMeshAsset()->GetSkeleton())
+			{
+				EditorSkelComp->EnablePreview(true, PreviewAnimationAsset);
+				EditorSkelComp->Play(true);
+			}
+
+			// Add the floor
+			TSharedPtr<IPersonaPreviewScene> Scene = PreviewScene.Pin();
+			if (Scene != nullptr)
+			{
+				UStaticMeshComponent* FloorMeshComponent = const_cast<UStaticMeshComponent*>(Scene->GetFloorMeshComponent());
+				if ((FloorMeshComponent != nullptr) && (FloorMeshComponent->GetBodyInstance() != nullptr))
+				{
+					EditorSkelComp->CreateSimulationFloor(FloorMeshComponent->GetBodyInstance(), FloorMeshComponent->GetBodyInstance()->GetUnrealWorldTransform());
+				}
+			}
+		}
+
+		if (EditorOptions->bResetClothWhenSimulating)
+		{
+			EditorSkelComp->RecreateClothingActors();
+		}
+	}
+	else
+	{
+		// Disable the PreviewInstance
+		EditorSkelComp->SetAnimationMode(EAnimationMode::AnimationSingleNode);
+
+		// Undo ends up recreating the anim script instance, so we need to remove it here (otherwise the AnimNode_RigidBody simulation starts when we undo)
+		EditorSkelComp->ClearAnimScriptInstance();
+
+		EditorSkelComp->SetPhysicsBlendWeight(0.f);
+		EditorSkelComp->ResetAllBodiesSimulatePhysics();
+		EditorSkelComp->SetSimulatePhysics(false);
+		ForceDisableSimulation();
+
+		// Since simulation, actor location changes. Reset to identity 
+		EditorSkelComp->SetWorldTransform(ResetTM);
+		// Force an update of the skeletal mesh to get it back to ref pose
+		EditorSkelComp->RefreshBoneTransforms();
+
+		// restore the EditorSkelComp animation asset 
+		if (PreviewAnimationAsset)
+		{
+			EditorSkelComp->EnablePreview(true, PreviewAnimationAsset);
+		}
+
+		PhysicsControlComponent->DestroyControlsInSet("All");
+		PhysicsControlComponent->DestroyBodyModifiersInSet("All");
+
+		BroadcastPreviewChanged();
+	}
+
+	bRunningSimulation = bEnableSimulation;
+}
+
+//======================================================================================================================
+EPhysicsAssetEditorMeshViewMode FPhysicsControlAssetEditorData::GetCurrentMeshViewMode(bool bSimulation)
+{
+	if (bSimulation)
+	{
+		return EditorOptions->SimulationMeshViewMode;
+	}
+	else
+	{
+		return EditorOptions->MeshViewMode;
+	}
+}
+
+//======================================================================================================================
+EPhysicsAssetEditorCollisionViewMode FPhysicsControlAssetEditorData::GetCurrentCollisionViewMode(bool bSimulation)
+{
+	if (bSimulation)
+	{
+		return EditorOptions->SimulationCollisionViewMode;
+	}
+	else
+	{
+		return EditorOptions->CollisionViewMode;
+	}
+}
+
+//======================================================================================================================
+EPhysicsAssetEditorConstraintViewMode FPhysicsControlAssetEditorData::GetCurrentConstraintViewMode(bool bSimulation)
+{
+	if (bSimulation)
+	{
+		return EditorOptions->SimulationConstraintViewMode;
+	}
+	else
+	{
+		return EditorOptions->ConstraintViewMode;
+	}
+}
+
+//======================================================================================================================
+void FPhysicsControlAssetEditorData::AddReferencedObjects(FReferenceCollector& Collector)
+{
+	Collector.AddReferencedObject(PhysicsControlAsset);
+	Collector.AddReferencedObject(EditorSkelComp);
+	Collector.AddReferencedObject(PhysicsControlComponent);
+	Collector.AddReferencedObject(EditorOptions);
+	Collector.AddReferencedObject(MouseHandle);
+
+	if (PreviewScene != nullptr)
+	{
+		PreviewScene.Pin()->AddReferencedObjects(Collector);
+	}
+}
+
+//======================================================================================================================
+void FPhysicsControlAssetEditorData::ForceDisableSimulation()
+{
+	// Reset simulation state of body instances so we dont actually simulate outside of 'simulation mode'
+
+	UPhysicsAsset* PA = PhysicsControlAsset->GetPhysicsAsset();
+	if (PA)
+	{
+		for (int32 BodyIdx = 0; BodyIdx < EditorSkelComp->Bodies.Num(); ++BodyIdx)
+		{
+			if (FBodyInstance* BodyInst = EditorSkelComp->Bodies[BodyIdx])
+			{
+				if (UBodySetup* PhysAssetBodySetup = PA->SkeletalBodySetups[BodyIdx])
+				{
+					BodyInst->SetInstanceSimulatePhysics(false);
+				}
+			}
+		}
+	}
+}
+
+//======================================================================================================================
+void FPhysicsControlAssetEditorData::BroadcastPreviewChanged()
+{
+	PreviewChangedEvent.Broadcast();
+}
+
+//======================================================================================================================
+void FPhysicsControlAssetEditorData::BroadcastSelectionChanged()
+{
+	if (!bSuspendSelectionBroadcast)
+	{
+		SelectionChangedEvent.Broadcast(SelectedBodies);
+	}
+}
+
+//======================================================================================================================
+void FPhysicsControlAssetEditorData::HitBone(int32 BodyIndex, EAggCollisionShape::Type PrimType, int32 PrimIndex, bool bGroupSelect)
+{
+	if (!bRunningSimulation)
+	{
+		FPhysicsControlAssetEditorData::FSelection Selection(BodyIndex, PrimType, PrimIndex);
+		if (bGroupSelect)
+		{
+			if (IsBodySelected(Selection))
+			{
+				SetSelectedBody(Selection, false);
+			}
+			else
+			{
+				SetSelectedBody(Selection, true);
+			}
+		}
+		else
+		{
+			ClearSelectedBody();
+			SetSelectedBody(Selection, true);
+		}
+	}
+}
+
+
+//======================================================================================================================
+FPhysicsControlAssetEditorData::FSelection* FPhysicsControlAssetEditorData::GetSelectedBody()
+{
+	int32 Count = SelectedBodies.Num();
+	return Count ? &SelectedBodies[Count - 1] : NULL;
+}
+
+//======================================================================================================================
+void FPhysicsControlAssetEditorData::ClearSelectedBody()
+{
+	SelectedBodies.Empty();
+	BroadcastSelectionChanged();
+}
+
+//======================================================================================================================
+void FPhysicsControlAssetEditorData::SetSelectedBody(const FSelection& Body, bool bSelected)
+{
+	SetSelectedBodies({ Body }, bSelected);
+}
+
+//======================================================================================================================
+void FPhysicsControlAssetEditorData::SetSelectedBodies(const TArray<FSelection>& Bodies, bool bSelected)
+{
+	if (InsideSelChange || Bodies.Num() == 0)
+	{
+		return;
+	}
+
+	if (bSelected)
+	{
+		for (const FSelection& Body : Bodies)
+		{
+			SelectedBodies.AddUnique(Body);
+		}
+	}
+	else
+	{
+		for (const FSelection& Body : Bodies)
+		{
+			SelectedBodies.Remove(Body);
+		}
+	}
+
+	BroadcastSelectionChanged();
+
+	if (!GetSelectedBody())
+	{
+		return;
+	}
+
+	++InsideSelChange;
+	BroadcastPreviewChanged();
+	--InsideSelChange;
+}
+
+//======================================================================================================================
+bool FPhysicsControlAssetEditorData::IsBodySelected(const FSelection& Body) const
+{
+	return SelectedBodies.Contains(Body);
+}
+
 
 #undef LOCTEXT_NAMESPACE
diff --git a/Engine/Plugins/Experimental/PhysicsControl/Source/PhysicsControlEditor/Private/PhysicsControlAssetEditorData.h b/Engine/Plugins/Experimental/PhysicsControl/Source/PhysicsControlEditor/Private/PhysicsControlAssetEditorData.h
--- a/Engine/Plugins/Experimental/PhysicsControl/Source/PhysicsControlEditor/Private/PhysicsControlAssetEditorData.h
+++ b/Engine/Plugins/Experimental/PhysicsControl/Source/PhysicsControlEditor/Private/PhysicsControlAssetEditorData.h
@@ -3,37 +3,132 @@
 #pragma once
 
 #include "CoreMinimal.h"
+#include "PhysicsEngine/ShapeElem.h"
 #include "UObject/ObjectPtr.h"
+#include "Preferences/PhysicsAssetEditorOptions.h"
 
 class IPersonaPreviewScene;
 class UPhysicsControlAsset;
 class UPhysicsControlAssetEditorSkeletalMeshComponent;
+class UPhysicsControlAssetEditorPhysicsHandleComponent;
+class UPhysicsControlComponent;
 
 /**
- * Helper/container for data used by the Physics Control Profile Editor 
+ * Helper/container for data used by the Physics Control Asset Editor 
  */
 class FPhysicsControlAssetEditorData
 {
 public:
 	FPhysicsControlAssetEditorData();
 
+	/** Encapsulates a selected set of bodies */
+	struct FSelection
+	{
+		int32 Index;
+		EAggCollisionShape::Type PrimitiveType;
+		int32 PrimitiveIndex;
+		FTransform WidgetTM;
+		FTransform ManipulateTM;
+
+		FSelection(int32 GivenBodyIndex, EAggCollisionShape::Type GivenPrimitiveType, int32 GivenPrimitiveIndex) :
+			Index(GivenBodyIndex), PrimitiveType(GivenPrimitiveType), PrimitiveIndex(GivenPrimitiveIndex),
+			WidgetTM(FTransform::Identity), ManipulateTM(FTransform::Identity)
+		{
+		}
+
+		bool operator==(const FSelection& rhs) const
+		{
+			return Index == rhs.Index && PrimitiveType == rhs.PrimitiveType && PrimitiveIndex == rhs.PrimitiveIndex;
+		}
+	};
+
 	/** Initializes members */
 	void Initialize(const TSharedRef<IPersonaPreviewScene>& InPreviewScene);
 
 	/** Caches a preview mesh. Sets us to a default mesh if none is set yet (or if an older one got deleted) */
 	void CachePreviewMesh();
 
+	/** Accessor for mesh view mode, allows access for simulation and non-simulation modes */
+	EPhysicsAssetEditorMeshViewMode GetCurrentMeshViewMode(bool bSimulation);
+
+	/** Accessor for collision view mode, allows access for simulation and non-simulation modes */
+	EPhysicsAssetEditorCollisionViewMode GetCurrentCollisionViewMode(bool bSimulation);
+
+	/** Accessor for constraint view mode, allows access for simulation and non-simulation modes */
+	EPhysicsAssetEditorConstraintViewMode GetCurrentConstraintViewMode(bool bSimulation);
+
+	/** Prevents GC from collecting our objects */
+	void AddReferencedObjects(FReferenceCollector& Collector);
+
+	void EnableSimulation(bool bEnableSimulation);
+
+	/** Force simulation off for all bodies, regardless of physics type */
+	void ForceDisableSimulation();
+
+	/** Toggle simulation on and off */
+	void ToggleSimulation();
+
+	/** broadcast a change in the preview*/
+	void BroadcastPreviewChanged();
+
+	/** broadcast a selection change ( if bSuspendSelectionBroadcast is false ) */
+	void BroadcastSelectionChanged();
+
+	/** Handle clicking on a body */
+	void HitBone(int32 BodyIndex, EAggCollisionShape::Type PrimType, int32 PrimIndex, bool bGroupSelect);
+
+	/** Selection */
+	FSelection* GetSelectedBody();
+	void ClearSelectedBody();
+	void SetSelectedBody(const FSelection& Body, bool bSelected);
+	void SetSelectedBodies(const TArray<FSelection>& Bodies, bool bSelected);
+	bool IsBodySelected(const FSelection& Body) const;
+
 public:
-	/** The asset being inspected */
+	/** The PhysicsControlAsset being inspected */
 	TObjectPtr<UPhysicsControlAsset> PhysicsControlAsset;
 
 	/** Skeletal mesh component specialized for this asset editor */
-	UPhysicsControlAssetEditorSkeletalMeshComponent* EditorSkelComp;
+	TObjectPtr<UPhysicsControlAssetEditorSkeletalMeshComponent> EditorSkelComp;
 
 	/** The physics control component used for testing/simulating on the character */
-	class UPhysicsControlComponent* PhysicsControlComponent;
+	TObjectPtr<UPhysicsControlComponent> PhysicsControlComponent;
+
+	// viewport anim instance Danny TODO needed? This is for when running with RBAN (which we will want to do)
+	//TObjectPtr<UAnimPreviewInstance> ViewportAnimInstance;
 
 	/** Preview scene */
 	TWeakPtr<IPersonaPreviewScene> PreviewScene;
 
+	/** Editor options */
+	TObjectPtr<UPhysicsAssetEditorOptions> EditorOptions;
+
+	/** Helps define how the asset behaves given user interaction in simulation mode*/
+	TObjectPtr<UPhysicsControlAssetEditorPhysicsHandleComponent> MouseHandle;
+
+	/** Callback for triggering a refresh of the preview viewport */
+	DECLARE_EVENT(FPhysicsControlAssetEditorSharedData, FPreviewChanged);
+	FPreviewChanged PreviewChangedEvent;
+
+	/** Callback for handling selection changes */
+	DECLARE_EVENT_OneParam(FPhysicsAssetEditorSharedData, FSelectionChanged, const TArray<FSelection>&);
+	FSelectionChanged SelectionChangedEvent;
+
+	TArray<FSelection> SelectedBodies;
+
+	/** Misc toggles */
+	bool bRunningSimulation;
+	bool bNoGravitySimulation;
+
+	/** Manipulation (rotate, translate, scale) */
+	bool bManipulating;
+
+	/** when true, we don't broadcast every selection change - allows for bulk changes without so much overhead */
+	bool bSuspendSelectionBroadcast;
+
+	/** Used to prevent recursion with tree hierarchy ... needs to be rewritten! */
+	int32 InsideSelChange;
+
+	// Where we put the component back to after simulating. It will just be identity.
+	const FTransform ResetTM;
 };
diff --git a/Engine/Plugins/Experimental/PhysicsControl/Source/PhysicsControlEditor/Private/PhysicsControlAssetEditorEditMode.cpp b/Engine/Plugins/Experimental/PhysicsControl/Source/PhysicsControlEditor/Private/PhysicsControlAssetEditorEditMode.cpp
new file mode 100644
--- /dev/null
+++ b/Engine/Plugins/Experimental/PhysicsControl/Source/PhysicsControlEditor/Private/PhysicsControlAssetEditorEditMode.cpp
@@ -0,0 +1,579 @@
+// Copyright Epic Games, Inc. All Rights Reserved.
+#include "PhysicsControlAssetEditorEditMode.h"
+
+#include "PhysicsControlAsset.h"
+#include "PhysicsControlAssetEditor.h"
+#include "PhysicsControlAssetEditorData.h"
+#include "PhysicsControlAssetEditorHitProxies.h"
+#include "PhysicsControlAssetEditorSkeletalMeshComponent.h"
+#include "PhysicsControlAssetEditorPhysicsHandleComponent.h"
+#include "PhysicsControlEditorModule.h"
+
+#include "AssetEditorModeManager.h"
+#include "Engine/Font.h"
+#include "GameFramework/WorldSettings.h"
+#include "PhysicsEngine/PhysicsAsset.h"
+#include "PhysicsEngine/PhysicsSettings.h"
+#include "PhysicsEngine/SkeletalBodySetup.h"
+#include "SEditorViewport.h"
+#include "Framework/Application/SlateApplication.h"
+
+FName FPhysicsControlAssetEditorEditMode::ModeName("PhysicsControlAssetEditMode");
+
+//======================================================================================================================
+FPhysicsControlAssetEditorEditMode::FPhysicsControlAssetEditorEditMode()
+	: SimHoldDistanceChangeDelta(20.0f)
+	, SimMinHoldDistance(10.0f)
+	, SimGrabMoveSpeed(1.0f)
+{
+	// Disable grid drawing for this mode as the viewport handles this
+	bDrawGrid = false;
+
+	LastClickPos = FIntPoint::ZeroValue;
+	LastClickOrigin = FVector::ZeroVector;
+	LastClickDirection = FVector::UpVector;
+	LastClickHitPos = FVector::ZeroVector;
+	LastClickHitNormal = FVector::UpVector;
+	bLastClickHit = false;
+
+	PhysicsControlAssetEditorFont = GEngine->GetSmallFont();
+	check(PhysicsControlAssetEditorFont);
+}
+
+
+//======================================================================================================================
+void FPhysicsControlAssetEditorEditMode::SetEditorData(
+	const TSharedRef<FPhysicsControlAssetEditor> InPhysicsControlAssetEditor,
+	TSharedPtr<FPhysicsControlAssetEditorData> InEditorData)
+{
+	PhysicsControlAssetEditor = InPhysicsControlAssetEditor;
+	EditorData = InEditorData;
+}
+
+//======================================================================================================================
+bool FPhysicsControlAssetEditorEditMode::GetCameraTarget(FSphere& OutTarget) const
+{
+	bool bHandled = false;
+
+	float NumSelected = 0.0f;
+	FBox Bounds(ForceInit);
+	if (const UPhysicsAsset* PA = EditorData->PhysicsControlAsset->GetPhysicsAsset())
+	{
+		for (int32 BodyIndex = 0; BodyIndex < EditorData->SelectedBodies.Num(); ++BodyIndex)
+		{
+			FPhysicsControlAssetEditorData::FSelection& SelectedObject = EditorData->SelectedBodies[BodyIndex];
+			int32 BoneIndex = EditorData->EditorSkelComp->GetBoneIndex(PA->SkeletalBodySetups[SelectedObject.Index]->BoneName);
+			UBodySetup* BodySetup = PA->SkeletalBodySetups[SelectedObject.Index];
+			const FKAggregateGeom& AggGeom = BodySetup->AggGeom;
+
+			FTransform BoneTM = EditorData->EditorSkelComp->GetBoneTransform(BoneIndex);
+			const float Scale = (float) BoneTM.GetScale3D().GetAbsMax();
+			BoneTM.RemoveScaling();
+
+			if (SelectedObject.PrimitiveType == EAggCollisionShape::Sphere)
+			{
+				Bounds += AggGeom.SphereElems[SelectedObject.PrimitiveIndex].CalcAABB(BoneTM, Scale);
+			}
+			else if (SelectedObject.PrimitiveType == EAggCollisionShape::Box)
+			{
+				Bounds += AggGeom.BoxElems[SelectedObject.PrimitiveIndex].CalcAABB(BoneTM, Scale);
+			}
+			else if (SelectedObject.PrimitiveType == EAggCollisionShape::Sphyl)
+			{
+				Bounds += AggGeom.SphylElems[SelectedObject.PrimitiveIndex].CalcAABB(BoneTM, Scale);
+			}
+			else if (SelectedObject.PrimitiveType == EAggCollisionShape::Convex)
+			{
+				Bounds += AggGeom.ConvexElems[SelectedObject.PrimitiveIndex].CalcAABB(BoneTM, BoneTM.GetScale3D());
+			}
+			else if (SelectedObject.PrimitiveType == EAggCollisionShape::TaperedCapsule)
+			{
+				Bounds += AggGeom.TaperedCapsuleElems[SelectedObject.PrimitiveIndex].CalcAABB(BoneTM, Scale);
+			}
+			else if (SelectedObject.PrimitiveType == EAggCollisionShape::LevelSet)
+			{
+				Bounds += AggGeom.LevelSetElems[SelectedObject.PrimitiveIndex].CalcAABB(BoneTM, BoneTM.GetScale3D());
+			}
+			else if (SelectedObject.PrimitiveType == EAggCollisionShape::SkinnedLevelSet)
+			{
+				Bounds += AggGeom.SkinnedLevelSetElems[SelectedObject.PrimitiveIndex].CalcAABB(BoneTM, BoneTM.GetScale3D());
+			}
+
+			bHandled = true;
+		}
+	}
+
+	OutTarget.Center = Bounds.GetCenter();
+	OutTarget.W = Bounds.GetExtent().Size();	// @TODO: calculate correct bounds
+
+	return bHandled;
+}
+
+//======================================================================================================================
+IPersonaPreviewScene& FPhysicsControlAssetEditorEditMode::GetAnimPreviewScene() const
+{
+	return *static_cast<IPersonaPreviewScene*>(static_cast<FAssetEditorModeManager*>(Owner)->GetPreviewScene());
+}
+
+//======================================================================================================================
+void FPhysicsControlAssetEditorEditMode::Exit()
+{
+	IPersonaEditMode::Exit();
+}
+
+//======================================================================================================================
+void FPhysicsControlAssetEditorEditMode::Render(const FSceneView* View, FViewport* Viewport, FPrimitiveDrawInterface* PDI)
+{
+	FEdMode::Render(View, Viewport, PDI);
+	// Danny TODO EditorData needs to be a smart pointer - this can fail on shut down
+	EPhysicsAssetEditorMeshViewMode MeshViewMode = EditorData->GetCurrentMeshViewMode(EditorData->bRunningSimulation);
+
+	if (MeshViewMode != EPhysicsAssetEditorMeshViewMode::None)
+	{
+		EditorData->EditorSkelComp->SetVisibility(true);
+
+		if (MeshViewMode == EPhysicsAssetEditorMeshViewMode::Wireframe)
+		{
+			EditorData->EditorSkelComp->SetForceWireframe(true);
+		}
+		else
+		{
+			EditorData->EditorSkelComp->SetForceWireframe(false);
+		}
+	}
+	else
+	{
+		EditorData->EditorSkelComp->SetVisibility(false);
+	}
+
+	// Draw phat skeletal component.
+	EditorData->EditorSkelComp->DebugDraw(View, PDI);
+}
+
+//======================================================================================================================
+void FPhysicsControlAssetEditorEditMode::Tick(FEditorViewportClient* ViewportClient, float DeltaTime)
+{
+	FEdMode::Tick(ViewportClient, DeltaTime);
+
+	if (EditorData->bRunningSimulation)
+	{
+		// check if PIE disabled the realtime viewport and quit sim if so
+		if (!ViewportClient->IsRealtime())
+		{
+			EditorData->ToggleSimulation();
+
+			ViewportClient->Invalidate();
+		}
+
+		UWorld* World = EditorData->PreviewScene.Pin()->GetWorld();
+		AWorldSettings* Setting = World->GetWorldSettings();
+		Setting->WorldGravityZ = EditorData->bNoGravitySimulation ? 
+			0.0f : 
+			(EditorData->EditorOptions->bUseGravityOverride ? 
+				EditorData->EditorOptions->GravityOverrideZ : 
+				(UPhysicsSettings::Get()->DefaultGravityZ * EditorData->EditorOptions->GravScale));
+		Setting->bWorldGravitySet = true;
+
+		// We back up the transforms array now
+		EditorData->EditorSkelComp->AnimationSpaceBases = EditorData->EditorSkelComp->GetComponentSpaceTransforms();
+
+		// We don't apply the physics blend, since that comes from body modifiers.
+		EditorData->EditorSkelComp->bUpdateJointsFromAnimation = EditorData->EditorOptions->bUpdateJointsFromAnimation;
+		EditorData->EditorSkelComp->PhysicsTransformUpdateMode = EditorData->EditorOptions->PhysicsUpdateMode;
+	}
+}
+
+//======================================================================================================================
+void FPhysicsControlAssetEditorEditMode::DrawHUD(FEditorViewportClient* ViewportClient, FViewport* Viewport, const FSceneView* View, FCanvas* Canvas)
+{
+	FEdMode::DrawHUD(ViewportClient, Viewport, View, Canvas);
+}
+
+//======================================================================================================================
+void FPhysicsControlAssetEditorEditMode::HitNothing(FEditorViewportClient* InViewportClient)
+{
+	if (InViewportClient->IsCtrlPressed() == false)	//we only want to deselect if Ctrl is not used
+	{
+		EditorData->ClearSelectedBody();
+	}
+
+	InViewportClient->Invalidate();
+	PhysicsControlAssetEditor.Pin()->RefreshHierachyTree();
+}
+
+//======================================================================================================================
+bool FPhysicsControlAssetEditorEditMode::HandleClick(FEditorViewportClient* InViewportClient, HHitProxy* HitProxy, const FViewportClick& Click)
+{
+	if (!EditorData->bRunningSimulation)
+	{
+		if (Click.GetKey() == EKeys::LeftMouseButton)
+		{
+			if (HitProxy && HitProxy->IsA(HPhysicsControlAssetEditorEdBoneProxy::StaticGetType()))
+			{
+				HPhysicsControlAssetEditorEdBoneProxy* BoneProxy = (HPhysicsControlAssetEditorEdBoneProxy*)HitProxy;
+
+				EditorData->HitBone(BoneProxy->BodyIndex, BoneProxy->PrimType, BoneProxy->PrimIndex, 
+					InViewportClient->IsCtrlPressed() || InViewportClient->IsShiftPressed());
+				return true;
+			}
+			else
+			{
+				HitNothing(InViewportClient);
+			}
+		}
+		else if (Click.GetKey() == EKeys::RightMouseButton)
+		{
+			if (HitProxy && HitProxy->IsA(HPhysicsControlAssetEditorEdBoneProxy::StaticGetType()))
+			{
+				HPhysicsControlAssetEditorEdBoneProxy* BoneProxy = (HPhysicsControlAssetEditorEdBoneProxy*)HitProxy;
+
+				// Select body under cursor if not already selected	(if ctrl is held down we only add, not remove)
+				FPhysicsControlAssetEditorData::FSelection Selection(BoneProxy->BodyIndex, BoneProxy->PrimType, BoneProxy->PrimIndex);
+				if (!EditorData->IsBodySelected(Selection))
+				{
+					if (!InViewportClient->IsCtrlPressed())
+					{
+						EditorData->ClearSelectedBody();
+					}
+
+					EditorData->SetSelectedBody(Selection, true);
+				}
+
+				// Pop up menu, if we have a body selected.
+				if (EditorData->GetSelectedBody())
+				{
+					OpenBodyMenu(InViewportClient);
+				}
+
+				return true;
+			}
+			else
+			{
+				OpenSelectionMenu(InViewportClient);
+				return true;
+			}
+		}
+	}
+
+	return false;
+}
+
+//======================================================================================================================
+/** Helper function to open a viewport context menu */
+static void OpenContextMenu(
+	const TSharedRef<FPhysicsControlAssetEditor>& PhysicsControlAssetEditor,
+	FEditorViewportClient*                        InViewportClient, 
+	TFunctionRef<void(FMenuBuilder&)>             InBuildMenu)
+{
+	FMenuBuilder MenuBuilder(true, PhysicsControlAssetEditor->GetToolkitCommands());
+
+	InBuildMenu(MenuBuilder);
+
+	TSharedPtr<SWidget> MenuWidget = MenuBuilder.MakeWidget();
+	TSharedPtr<SWidget> ParentWidget = InViewportClient->GetEditorViewportWidget();
+
+	if (MenuWidget.IsValid() && ParentWidget.IsValid())
+	{
+		const FVector2D MouseCursorLocation = FSlateApplication::Get().GetCursorPos();
+
+		FSlateApplication::Get().PushMenu(
+			ParentWidget.ToSharedRef(),
+			FWidgetPath(),
+			MenuWidget.ToSharedRef(),
+			MouseCursorLocation,
+			FPopupTransitionEffect(FPopupTransitionEffect::ContextMenu)
+		);
+	}
+}
+
+//======================================================================================================================
+void FPhysicsControlAssetEditorEditMode::OpenBodyMenu(FEditorViewportClient* InViewportClient)
+{
+	OpenContextMenu(PhysicsControlAssetEditor.Pin().ToSharedRef(), InViewportClient,
+		[this](FMenuBuilder& InMenuBuilder)
+		{
+			PhysicsControlAssetEditor.Pin()->BuildMenuWidgetBody(InMenuBuilder);
+			PhysicsControlAssetEditor.Pin()->BuildMenuWidgetSelection(InMenuBuilder);
+		});
+}
+
+//======================================================================================================================
+void FPhysicsControlAssetEditorEditMode::OpenSelectionMenu(FEditorViewportClient* InViewportClient)
+{
+	OpenContextMenu(PhysicsControlAssetEditor.Pin().ToSharedRef(), InViewportClient,
+		[this](FMenuBuilder& InMenuBuilder)
+		{
+			PhysicsControlAssetEditor.Pin()->BuildMenuWidgetSelection(InMenuBuilder);
+		});
+}
+
+
+//======================================================================================================================
+bool FPhysicsControlAssetEditorEditMode::ReceivedFocus(FEditorViewportClient* ViewportClient, FViewport* Viewport)
+{
+	return IPersonaEditMode::ReceivedFocus(ViewportClient, Viewport);
+}
+
+//======================================================================================================================
+bool FPhysicsControlAssetEditorEditMode::LostFocus(FEditorViewportClient* ViewportClient, FViewport* Viewport)
+{
+	return IPersonaEditMode::LostFocus(ViewportClient, Viewport);
+}
+
+//======================================================================================================================
+bool FPhysicsControlAssetEditorEditMode::InputKey(FEditorViewportClient* InViewportClient, FViewport* InViewport, FKey Key, EInputEvent Event)
+{
+	int32 HitX = InViewport->GetMouseX();
+	int32 HitY = InViewport->GetMouseY();
+	bool bCtrlDown = InViewport->KeyState(EKeys::LeftControl) || InViewport->KeyState(EKeys::RightControl);
+	bool bShiftDown = InViewport->KeyState(EKeys::LeftShift) || InViewport->KeyState(EKeys::RightShift);
+
+	bool bHandled = false;
+	if (EditorData->bRunningSimulation)
+	{
+		if (Key == EKeys::RightMouseButton || Key == EKeys::LeftMouseButton)
+		{
+			if (Event == IE_Pressed)
+			{
+				bHandled = SimMousePress(InViewportClient, Key);
+			}
+			else if (Event == IE_Released)
+			{
+				bHandled = SimMouseRelease();
+			}
+			else
+			{
+				// Handle repeats/double clicks etc. so we dont fall through
+				bHandled = true;
+			}
+		}
+		else if (Key == EKeys::MouseScrollUp)
+		{
+			bHandled = SimMouseWheelUp(InViewportClient);
+		}
+		else if (Key == EKeys::MouseScrollDown)
+		{
+			bHandled = SimMouseWheelDown(InViewportClient);
+		}
+		else if (InViewportClient->IsFlightCameraActive())
+		{
+			// If the flight camera is active (user is looking or moving around the scene)
+			// consume the event so hotkeys don't fire.
+			bHandled = true;
+		}
+	}
+
+	if (!bHandled)
+	{
+		bHandled = IPersonaEditMode::InputKey(InViewportClient, InViewport, Key, Event);
+	}
+
+	if (bHandled)
+	{
+		InViewportClient->Invalidate();
+	}
+
+	return bHandled;
+}
+
+//======================================================================================================================
+bool FPhysicsControlAssetEditorEditMode::InputAxis(FEditorViewportClient* InViewportClient, FViewport* InViewport, int32 ControllerId, FKey Key, float Delta, float DeltaTime)
+{
+	bool bHandled = false;
+	// If we are 'manipulating' don't move the camera but do something else with mouse input.
+	if (EditorData->bManipulating)
+	{
+		bool bCtrlDown = InViewport->KeyState(EKeys::LeftControl) || InViewport->KeyState(EKeys::RightControl);
+
+		if (EditorData->bRunningSimulation)
+		{
+			if (Key == EKeys::MouseX)
+			{
+				SimMouseMove(InViewportClient, Delta, 0.0f);
+			}
+			else if (Key == EKeys::MouseY)
+			{
+				SimMouseMove(InViewportClient, 0.0f, Delta);
+			}
+			bHandled = true;
+		}
+	}
+
+	if (!bHandled)
+	{
+		bHandled = IPersonaEditMode::InputAxis(InViewportClient, InViewport, ControllerId, Key, Delta, DeltaTime);
+	}
+
+	InViewportClient->Invalidate();
+
+	return bHandled;
+}
+
+
+//======================================================================================================================
+bool FPhysicsControlAssetEditorEditMode::SimMousePress(FEditorViewportClient* InViewportClient, FKey Key)
+{
+	bool bHandled = false;
+
+	FViewport* Viewport = InViewportClient->Viewport;
+
+	bool bCtrlDown = Viewport->KeyState(EKeys::LeftControl) || Viewport->KeyState(EKeys::RightControl);
+	bool bShiftDown = Viewport->KeyState(EKeys::LeftShift) || Viewport->KeyState(EKeys::RightShift);
+
+	FSceneViewFamilyContext ViewFamily(FSceneViewFamily::ConstructionValues(
+		Viewport, InViewportClient->GetScene(), InViewportClient->EngineShowFlags));
+	FSceneView* View = InViewportClient->CalcSceneView(&ViewFamily);
+
+	const FViewportClick Click(
+		View, InViewportClient, EKeys::Invalid, IE_Released, Viewport->GetMouseX(), Viewport->GetMouseY());
+	FHitResult Result(1.f);
+	bool bHit = EditorData->EditorSkelComp->LineTraceComponent(
+		Result, Click.GetOrigin(), Click.GetOrigin() + Click.GetDirection() * EditorData->EditorOptions->InteractionDistance, FCollisionQueryParams(NAME_None, true));
+
+	LastClickPos = Click.GetClickPos();
+	LastClickOrigin = Click.GetOrigin();
+	LastClickDirection = Click.GetDirection();
+	bLastClickHit = bHit;
+	if (bHit)
+	{
+		LastClickHitPos = Result.Location;
+		LastClickHitNormal = Result.Normal;
+	}
+
+	const UPhysicsAsset* PA = EditorData->PhysicsControlAsset->GetPhysicsAsset();
+	if (bHit && PA)
+	{
+		check(Result.Item != INDEX_NONE);
+		FName BoneName = PA->SkeletalBodySetups[Result.Item]->BoneName;
+
+		UE_LOG(LogPhysics, Log, TEXT("Physics Asset Editor Click Hit Bone (%s)"), *BoneName.ToString());
+
+		if (bCtrlDown || bShiftDown)
+		{
+			// Right mouse is for dragging things around
+			if (Key == EKeys::RightMouseButton)
+			{
+				EditorData->bManipulating = true;
+				DragX = 0.0f;
+				DragY = 0.0f;
+				SimGrabPush = 0.0f;
+
+				// Update mouse force properties from sim options.
+				EditorData->MouseHandle->LinearDamping = EditorData->EditorOptions->HandleLinearDamping;
+				EditorData->MouseHandle->LinearStiffness = EditorData->EditorOptions->HandleLinearStiffness;
+				EditorData->MouseHandle->AngularDamping = EditorData->EditorOptions->HandleAngularDamping;
+				EditorData->MouseHandle->AngularStiffness = EditorData->EditorOptions->HandleAngularStiffness;
+				EditorData->MouseHandle->InterpolationSpeed = EditorData->EditorOptions->InterpolationSpeed;
+
+				// Create handle to object.
+				EditorData->MouseHandle->GrabComponentAtLocationWithRotation(
+					EditorData->EditorSkelComp, BoneName, Result.Location, FRotator::ZeroRotator);
+
+				FMatrix	InvViewMatrix = View->ViewMatrices.GetInvViewMatrix();
+
+				SimGrabMinPush = SimMinHoldDistance - (Result.Time * EditorData->EditorOptions->InteractionDistance);
+
+				SimGrabLocation = Result.Location;
+				SimGrabX = InvViewMatrix.GetUnitAxis(EAxis::X);
+				SimGrabY = InvViewMatrix.GetUnitAxis(EAxis::Y);
+				SimGrabZ = InvViewMatrix.GetUnitAxis(EAxis::Z);
+			}
+			// Left mouse is for poking things
+			else if (Key == EKeys::LeftMouseButton)
+			{
+				EditorData->MouseHandle->AddImpulseAtLocation(
+					EditorData->EditorSkelComp,
+					Click.GetDirection() * EditorData->EditorOptions->PokeStrength, Result.Location, BoneName);
+			}
+
+			bHandled = true;
+		}
+	}
+
+	return bHandled;
+}
+
+//======================================================================================================================
+void FPhysicsControlAssetEditorEditMode::SimMouseMove(FEditorViewportClient* InViewportClient, float DeltaX, float DeltaY)
+{
+	DragX = (float) (InViewportClient->Viewport->GetMouseX() - LastClickPos.X);
+	DragY = (float) (InViewportClient->Viewport->GetMouseY() - LastClickPos.Y);
+
+	if (!EditorData->MouseHandle->GrabbedComponent)
+	{
+		return;
+	}
+
+	//We need to convert Pixel Delta into Screen position (deal with different viewport sizes)
+	FSceneViewFamilyContext ViewFamily(FSceneViewFamily::ConstructionValues(InViewportClient->Viewport, EditorData->PreviewScene.Pin()->GetScene(), InViewportClient->EngineShowFlags));
+	FSceneView* View = InViewportClient->CalcSceneView(&ViewFamily);
+	FVector4 ScreenOldPos = View->PixelToScreen((float) LastClickPos.X, (float) LastClickPos.Y, 1.f);
+	FVector4 ScreenNewPos = View->PixelToScreen(DragX + (float) LastClickPos.X, DragY + (float) LastClickPos.Y, 1.f);
+	FVector4 ScreenDelta = ScreenNewPos - ScreenOldPos;
+	FVector4 ProjectedDelta = View->ScreenToWorld(ScreenDelta);
+	FVector4 WorldDelta;
+
+	//Now we project new ScreenPos to xy-plane of SimGrabLocation
+	FVector LocalOffset = View->ViewMatrices.GetViewMatrix().TransformPosition(SimGrabLocation + SimGrabZ * SimGrabPush);
+	//in the ortho case we don't need to do any fixup because there is no perspective
+	double ZDistance = InViewportClient->GetViewportType() == ELevelViewportType::LVT_Perspective ? FMath::Abs(LocalOffset.Z) : 1.0;
+	WorldDelta = ProjectedDelta * ZDistance;
+
+	//Now we convert back into WorldPos
+	FVector WorldPos = SimGrabLocation + WorldDelta + SimGrabZ * SimGrabPush;
+	FVector NewLocation = WorldPos;
+	float QuickRadius = 5 - SimGrabPush / SimHoldDistanceChangeDelta;
+	QuickRadius = QuickRadius < 2 ? 2 : QuickRadius;
+
+	DrawDebugPoint(GetWorld(), NewLocation, QuickRadius, FColorList::Red, false, 0.3f);
+
+	EditorData->MouseHandle->SetTargetLocation(NewLocation);
+	EditorData->MouseHandle->GrabbedComponent->WakeRigidBody(EditorData->MouseHandle->GrabbedBoneName);
+}
+
+//======================================================================================================================
+bool FPhysicsControlAssetEditorEditMode::SimMouseRelease()
+{
+	EditorData->bManipulating = false;
+
+	if (!EditorData->MouseHandle->GrabbedComponent)
+	{
+		return false;
+	}
+
+	EditorData->MouseHandle->GrabbedComponent->WakeRigidBody(EditorData->MouseHandle->GrabbedBoneName);
+	EditorData->MouseHandle->ReleaseComponent();
+
+	return true;
+}
+
+//======================================================================================================================
+bool FPhysicsControlAssetEditorEditMode::SimMouseWheelUp(FEditorViewportClient* InViewportClient)
+{
+	if (!EditorData->MouseHandle->GrabbedComponent)
+	{
+		return false;
+	}
+
+	SimGrabPush += SimHoldDistanceChangeDelta;
+
+	SimMouseMove(InViewportClient, 0.0f, 0.0f);
+
+	return true;
+}
+
+//======================================================================================================================
+bool FPhysicsControlAssetEditorEditMode::SimMouseWheelDown(FEditorViewportClient* InViewportClient)
+{
+	if (!EditorData->MouseHandle->GrabbedComponent)
+	{
+		return false;
+	}
+
+	SimGrabPush -= SimHoldDistanceChangeDelta;
+	SimGrabPush = FMath::Max(SimGrabMinPush, SimGrabPush);
+
+	SimMouseMove(InViewportClient, 0.0f, 0.0f);
+
+	return true;
+}
diff --git a/Engine/Plugins/Experimental/PhysicsControl/Source/PhysicsControlEditor/Private/PhysicsControlAssetEditorHitProxies.cpp b/Engine/Plugins/Experimental/PhysicsControl/Source/PhysicsControlEditor/Private/PhysicsControlAssetEditorHitProxies.cpp
new file mode 100644
--- /dev/null
+++ b/Engine/Plugins/Experimental/PhysicsControl/Source/PhysicsControlEditor/Private/PhysicsControlAssetEditorHitProxies.cpp
@@ -0,0 +1,6 @@
+// Copyright Epic Games, Inc. All Rights Reserved.
+
+#include "PhysicsControlAssetEditorHitProxies.h"
+
+IMPLEMENT_HIT_PROXY(HPhysicsControlAssetEditorEdBoneProxy,HHitProxy);
+IMPLEMENT_HIT_PROXY(HPhysicsControlAssetEditorEdBoneNameProxy,HHitProxy);
diff --git a/Engine/Plugins/Experimental/PhysicsControl/Source/PhysicsControlEditor/Private/PhysicsControlAssetEditorHitProxies.h b/Engine/Plugins/Experimental/PhysicsControl/Source/PhysicsControlEditor/Private/PhysicsControlAssetEditorHitProxies.h
new file mode 100644
--- /dev/null
+++ b/Engine/Plugins/Experimental/PhysicsControl/Source/PhysicsControlEditor/Private/PhysicsControlAssetEditorHitProxies.h
@@ -0,0 +1,53 @@
+// Copyright Epic Games, Inc. All Rights Reserved.
+
+#pragma once
+
+#include "GenericPlatform/ICursor.h"
+#include "HAL/Platform.h"
+#include "HitProxies.h"
+#include "PhysicsEngine/ShapeElem.h"
+
+/*-----------------------------------------------------------------------------
+   Hit Proxies
+-----------------------------------------------------------------------------*/
+
+struct HPhysicsControlAssetEditorEdBoneProxy : public HHitProxy
+{
+	DECLARE_HIT_PROXY();
+
+	int32							BodyIndex;
+	EAggCollisionShape::Type		PrimType;
+	int32							PrimIndex;
+
+	HPhysicsControlAssetEditorEdBoneProxy(int32 InBodyIndex, EAggCollisionShape::Type InPrimType, int32 InPrimIndex)
+		: HHitProxy(HPP_World)
+		, BodyIndex(InBodyIndex)
+		, PrimType(InPrimType)
+		, PrimIndex(InPrimIndex) {}
+
+	virtual EMouseCursor::Type GetMouseCursor()
+	{
+		return EMouseCursor::Crosshairs;
+	}
+
+	virtual bool AlwaysAllowsTranslucentPrimitives() const override
+	{
+		return true;
+	}
+};
+
+struct HPhysicsControlAssetEditorEdBoneNameProxy : public HHitProxy
+{
+	DECLARE_HIT_PROXY();
+
+	int32			BoneIndex;
+
+	HPhysicsControlAssetEditorEdBoneNameProxy(int32 InBoneIndex)
+		: HHitProxy(HPP_Foreground)
+		, BoneIndex(InBoneIndex) {}
+
+	virtual bool AlwaysAllowsTranslucentPrimitives() const override
+	{
+		return true;
+	}
+};
diff --git a/Engine/Plugins/Experimental/PhysicsControl/Source/PhysicsControlEditor/Private/PhysicsControlAssetEditorMode.cpp b/Engine/Plugins/Experimental/PhysicsControl/Source/PhysicsControlEditor/Private/PhysicsControlAssetEditorMode.cpp
deleted file mode 100644
diff --git a/Engine/Plugins/Experimental/PhysicsControl/Source/PhysicsControlEditor/Private/PhysicsControlAssetEditorPhysicsHandleComponent.cpp b/Engine/Plugins/Experimental/PhysicsControl/Source/PhysicsControlEditor/Private/PhysicsControlAssetEditorPhysicsHandleComponent.cpp
new file mode 100644
--- /dev/null
+++ b/Engine/Plugins/Experimental/PhysicsControl/Source/PhysicsControlEditor/Private/PhysicsControlAssetEditorPhysicsHandleComponent.cpp
@@ -0,0 +1,108 @@
+// Copyright Epic Games, Inc. All Rights Reserved.
+#include "PhysicsControlAssetEditorPhysicsHandleComponent.h"
+
+#include "PhysicsControlAssetEditorSkeletalMeshComponent.h"
+
+//======================================================================================================================
+UPhysicsControlAssetEditorPhysicsHandleComponent::UPhysicsControlAssetEditorPhysicsHandleComponent(
+	const FObjectInitializer& ObjectInitializer)
+	: Super(ObjectInitializer)
+	, bAnimInstanceMode(false)
+{
+}
+
+//======================================================================================================================
+void UPhysicsControlAssetEditorPhysicsHandleComponent::UpdateHandleTransform(const FTransform& NewTransform)
+{
+	Super::UpdateHandleTransform(NewTransform);
+	if (bAnimInstanceMode)
+	{
+		UPhysicsControlAssetEditorSkeletalMeshComponent* PhysicsControlAssetEditorSkeletalMeshComponent = 
+			Cast<UPhysicsControlAssetEditorSkeletalMeshComponent>(GrabbedComponent);
+		if (PhysicsControlAssetEditorSkeletalMeshComponent != nullptr)
+		{
+			PhysicsControlAssetEditorSkeletalMeshComponent->UpdateHandleTransform(NewTransform);
+		}
+	}
+}
+
+//======================================================================================================================
+void UPhysicsControlAssetEditorPhysicsHandleComponent::UpdateDriveSettings()
+{
+	Super::UpdateDriveSettings();
+	if (bAnimInstanceMode)
+	{
+		UPhysicsControlAssetEditorSkeletalMeshComponent* PhysicsControlAssetEditorSkeletalMeshComponent = 
+			Cast<UPhysicsControlAssetEditorSkeletalMeshComponent>(GrabbedComponent);
+		if (PhysicsControlAssetEditorSkeletalMeshComponent != nullptr)
+		{
+			PhysicsControlAssetEditorSkeletalMeshComponent->UpdateDriveSettings(
+				bSoftLinearConstraint, LinearStiffness, LinearDamping);
+		}
+	}
+}
+
+//======================================================================================================================
+void UPhysicsControlAssetEditorPhysicsHandleComponent::SetAnimInstanceMode(bool bInAnimInstanceMode)
+{
+	bAnimInstanceMode = bInAnimInstanceMode;
+}
+
+//======================================================================================================================
+void UPhysicsControlAssetEditorPhysicsHandleComponent::GrabComponentImp(
+	class UPrimitiveComponent* Component, FName InBoneName, 
+	const FVector& Location, const FRotator& Rotation, bool InbRotationConstrained)
+{
+	Super::GrabComponentImp(Component, InBoneName, Location, Rotation, bRotationConstrained);
+	if (bAnimInstanceMode)
+	{
+		TargetTransform = CurrentTransform = FTransform(Rotation, Location);
+
+		UPhysicsControlAssetEditorSkeletalMeshComponent* PhysicsControlAssetEditorSkeletalMeshComponent = 
+			Cast<UPhysicsControlAssetEditorSkeletalMeshComponent>(Component);
+		if (PhysicsControlAssetEditorSkeletalMeshComponent != nullptr)
+		{
+			PhysicsControlAssetEditorSkeletalMeshComponent->Grab(InBoneName, Location, Rotation, InbRotationConstrained);
+		}
+	}
+}
+
+//======================================================================================================================
+void UPhysicsControlAssetEditorPhysicsHandleComponent::ReleaseComponent()
+{
+	if (bAnimInstanceMode)
+	{
+		UPhysicsControlAssetEditorSkeletalMeshComponent* PhysicsControlAssetEditorSkeletalMeshComponent = 
+			Cast<UPhysicsControlAssetEditorSkeletalMeshComponent>(GrabbedComponent);
+		if (PhysicsControlAssetEditorSkeletalMeshComponent != nullptr)
+		{
+			PhysicsControlAssetEditorSkeletalMeshComponent->Ungrab();
+		}
+	}
+	Super::ReleaseComponent();
+}
+
+//======================================================================================================================
+void UPhysicsControlAssetEditorPhysicsHandleComponent::AddImpulseAtLocation(
+	UPhysicsControlAssetEditorSkeletalMeshComponent* PhysicsControlAssetEditorSkeletalMeshComponent, 
+	FVector Impulse, FVector Location, FName BoneName)
+{
+	if (bAnimInstanceMode)
+	{
+		PhysicsControlAssetEditorSkeletalMeshComponent->AddImpulseAtLocation(Impulse, Location, BoneName);
+	}
+	else
+	{
+		FBodyInstance* BodyInstance = PhysicsControlAssetEditorSkeletalMeshComponent->GetBodyInstance(BoneName);
+		if (!BodyInstance)
+		{
+			return;
+		}
+		// If the grabbed body is welded to another, poke that instead
+		if (BodyInstance->WeldParent != nullptr)
+		{
+			BodyInstance = BodyInstance->WeldParent;
+		}
+		BodyInstance->AddImpulseAtPosition(Impulse, Location);
+	}
+}
diff --git a/Engine/Plugins/Experimental/PhysicsControl/Source/PhysicsControlEditor/Private/PhysicsControlAssetEditorPhysicsHandleComponent.h b/Engine/Plugins/Experimental/PhysicsControl/Source/PhysicsControlEditor/Private/PhysicsControlAssetEditorPhysicsHandleComponent.h
new file mode 100644
--- /dev/null
+++ b/Engine/Plugins/Experimental/PhysicsControl/Source/PhysicsControlEditor/Private/PhysicsControlAssetEditorPhysicsHandleComponent.h
@@ -0,0 +1,36 @@
+// Copyright Epic Games, Inc. All Rights Reserved.
+#pragma once
+
+#include "PhysicsEngine/PhysicsHandleComponent.h"
+
+#include "PhysicsControlAssetEditorPhysicsHandleComponent.generated.h"
+
+class UPhysicsControlAssetEditorSkeletalMeshComponent;
+
+/**
+ * Extend the Physics Handle for PhAt. This adds support for manipulating the physics
+ * if it is running in a RigidBody AnimNode (which is always is with Chaos at the moment).
+ */
+UCLASS()
+class UPhysicsControlAssetEditorPhysicsHandleComponent : public UPhysicsHandleComponent
+{
+	GENERATED_UCLASS_BODY()
+
+	bool bAnimInstanceMode;
+
+public:
+	void SetAnimInstanceMode(bool bInAnimInstanceMode);
+
+	virtual void ReleaseComponent() override;
+
+	void AddImpulseAtLocation(
+		UPhysicsControlAssetEditorSkeletalMeshComponent* PhysicsControlAssetEditorSkeletalMeshComponent,
+		FVector Impulse, FVector Location, FName BoneName);
+
+protected:
+	virtual void UpdateHandleTransform(const FTransform& NewTransform) override;
+	virtual void UpdateDriveSettings() override;
+	virtual void GrabComponentImp(
+		class UPrimitiveComponent* Component, FName InBoneName, const FVector& Location, 
+		const FRotator& Rotation, bool bRotationConstrained) override;
+};
diff --git a/Engine/Plugins/Experimental/PhysicsControl/Source/PhysicsControlEditor/Private/PhysicsControlAssetEditorProfilesTabSummoners.cpp b/Engine/Plugins/Experimental/PhysicsControl/Source/PhysicsControlEditor/Private/PhysicsControlAssetEditorProfilesTabSummoners.cpp
new file mode 100644
--- /dev/null
+++ b/Engine/Plugins/Experimental/PhysicsControl/Source/PhysicsControlEditor/Private/PhysicsControlAssetEditorProfilesTabSummoners.cpp
@@ -0,0 +1,170 @@
+// Copyright Epic Games, Inc. All Rights Reserved.
+
+#include "PhysicsControlAssetEditorProfilesTabSummoners.h"
+#include "IDocumentation.h"
+#include "PhysicsControlAssetEditor.h"
+#include "PhysicsControlAsset.h"
+#include "PropertyEditorModule.h"
+#include "IDetailsView.h"
+#include "PhysicsControlAssetSetupDetailsCustomization.h"
+#include "PhysicsControlAssetProfileDetailsCustomization.h"
+#include "PhysicsControlAssetPreviewDetailsCustomization.h"
+
+#define LOCTEXT_NAMESPACE "PhysicsControlAssetEditorProfilesTabSummoner"
+
+//======================================================================================================================
+FName FPhysicsControlAssetEditorSetupTabSummoner::TabName = FName(
+	"PhysicsControlAssetEditorSetupTab");
+FName FPhysicsControlAssetEditorProfileTabSummoner::TabName = FName(
+	"PhysicsControlAssetEditorProfileTab");
+FName FPhysicsControlAssetEditorPreviewTabSummoner::TabName = FName(
+	"PhysicsControlAssetEditorPreviewTab");
+
+//======================================================================================================================
+FPhysicsControlAssetEditorSetupTabSummoner::FPhysicsControlAssetEditorSetupTabSummoner(
+	TSharedPtr<FAssetEditorToolkit> InHostingApp, UPhysicsControlAsset* InPhysicsControlAsset)
+	: FWorkflowTabFactory(FPhysicsControlAssetEditorSetupTabSummoner::TabName, InHostingApp)
+	, PhysicsControlAsset(InPhysicsControlAsset)
+{
+	TabLabel = LOCTEXT("PhysicsControlAssetEditorSetupTabTitle", "Setup");
+	TabIcon = FSlateIcon(FAppStyle::GetAppStyleSetName(), "PhysicsAssetEditor.Tabs.Profiles");
+
+	bIsSingleton = true;
+
+	ViewMenuDescription = LOCTEXT("PhysicsControlAssetEditorSetup", "Setup");
+	ViewMenuTooltip = LOCTEXT("PhysicsControlAssetEditorSetup_ToolTip", "Shows the Control Asset Setup tab");
+}
+
+//======================================================================================================================
+TSharedPtr<SToolTip> FPhysicsControlAssetEditorSetupTabSummoner::CreateTabToolTipWidget(
+	const FWorkflowTabSpawnInfo& Info) const
+{
+	return IDocumentation::Get()->CreateToolTip(LOCTEXT(
+		"PhysicsControlAssetEditorSetupToolTip", 
+		"The Physics Control Asset Setup tab lets you edit the physics control asset relating to setting up controls."), 
+		NULL, 
+		TEXT("Shared/Editors/PhysicsControlAssetEditor"), TEXT("PhysicsControlAssetProfiles_Window"));
+}
+
+//======================================================================================================================
+TSharedRef<SWidget> FPhysicsControlAssetEditorSetupTabSummoner::CreateTabBody(
+	const FWorkflowTabSpawnInfo& Info) const
+{
+	FDetailsViewArgs DetailsViewArgs;
+	DetailsViewArgs.NameAreaSettings = FDetailsViewArgs::HideNameArea;
+	DetailsViewArgs.bHideSelectionTip = true;
+	DetailsViewArgs.bAllowSearch = false;
+
+	FPropertyEditorModule& PropertyEditorModule = 
+		FModuleManager::LoadModuleChecked<FPropertyEditorModule>("PropertyEditor");
+	TSharedRef<IDetailsView> DetailsView = PropertyEditorModule.CreateDetailView(DetailsViewArgs);
+
+	TWeakPtr<FPhysicsControlAssetEditor> PhysicsControlAssetEditor = 
+		StaticCastSharedPtr<FPhysicsControlAssetEditor>(HostingApp.Pin());
+	DetailsView->RegisterInstancedCustomPropertyLayout(
+		UPhysicsControlAsset::StaticClass(), 
+		FOnGetDetailCustomizationInstance::CreateStatic(
+			&FPhysicsControlAssetSetupDetailsCustomization::MakeInstance, PhysicsControlAssetEditor));
+	DetailsView->SetObject(PhysicsControlAsset.Get());
+	return DetailsView;
+}
+
+//======================================================================================================================
+FPhysicsControlAssetEditorProfileTabSummoner::FPhysicsControlAssetEditorProfileTabSummoner(
+	TSharedPtr<FAssetEditorToolkit> InHostingApp, UPhysicsControlAsset* InPhysicsControlAsset)
+	: FWorkflowTabFactory(FPhysicsControlAssetEditorProfileTabSummoner::TabName, InHostingApp)
+	, PhysicsControlAsset(InPhysicsControlAsset)
+{
+	TabLabel = LOCTEXT("PhysicsControlAssetEditorProfileTabTitle", "Profiles");
+	TabIcon = FSlateIcon(FAppStyle::GetAppStyleSetName(), "PhysicsAssetEditor.Tabs.Profiles");
+
+	bIsSingleton = true;
+
+	ViewMenuDescription = LOCTEXT("PhysicsControlAssetEditorProfile", "Profiles");
+	ViewMenuTooltip = LOCTEXT("PhysicsControlAssetEditorProfile_ToolTip", "Shows the Control Asset Profile Edit tab");
+}
+
+//======================================================================================================================
+TSharedPtr<SToolTip> FPhysicsControlAssetEditorProfileTabSummoner::CreateTabToolTipWidget(
+	const FWorkflowTabSpawnInfo& Info) const
+{
+	return IDocumentation::Get()->CreateToolTip(LOCTEXT(
+		"PhysicsControlAssetEditorProfileToolTip",
+		"The Physics Control Asset Profile Edit tab lets you edit the physics control asset relating to setting up profiles."),
+		NULL,
+		TEXT("Shared/Editors/PhysicsControlAssetEditor"), TEXT("PhysicsControlAssetProfiles_Window"));
+}
+
+//======================================================================================================================
+TSharedRef<SWidget> FPhysicsControlAssetEditorProfileTabSummoner::CreateTabBody(
+	const FWorkflowTabSpawnInfo& Info) const
+{
+	FDetailsViewArgs DetailsViewArgs;
+	DetailsViewArgs.NameAreaSettings = FDetailsViewArgs::HideNameArea;
+	DetailsViewArgs.bHideSelectionTip = true;
+	DetailsViewArgs.bAllowSearch = false;
+
+	FPropertyEditorModule& PropertyEditorModule =
+		FModuleManager::LoadModuleChecked<FPropertyEditorModule>("PropertyEditor");
+	TSharedRef<IDetailsView> DetailsView = PropertyEditorModule.CreateDetailView(DetailsViewArgs);
+
+	TWeakPtr<FPhysicsControlAssetEditor> PhysicsControlAssetEditor =
+		StaticCastSharedPtr<FPhysicsControlAssetEditor>(HostingApp.Pin());
+	DetailsView->RegisterInstancedCustomPropertyLayout(
+		UPhysicsControlAsset::StaticClass(), 
+		FOnGetDetailCustomizationInstance::CreateStatic(
+			&FPhysicsControlAssetProfileDetailsCustomization::MakeInstance, PhysicsControlAssetEditor));
+	DetailsView->SetObject(PhysicsControlAsset.Get());
+	return DetailsView;
+}
+
+//======================================================================================================================
+FPhysicsControlAssetEditorPreviewTabSummoner::FPhysicsControlAssetEditorPreviewTabSummoner(
+	TSharedPtr<FAssetEditorToolkit> InHostingApp, UPhysicsControlAsset* InPhysicsControlAsset)
+	: FWorkflowTabFactory(FPhysicsControlAssetEditorPreviewTabSummoner::TabName, InHostingApp)
+	, PhysicsControlAsset(InPhysicsControlAsset)
+{
+	TabLabel = LOCTEXT("PhysicsControlAssetEditorPreviewTabTitle", "Preview");
+	TabIcon = FSlateIcon(FAppStyle::GetAppStyleSetName(), "PhysicsAssetEditor.Tabs.Profiles");
+
+	bIsSingleton = true;
+
+	ViewMenuDescription = LOCTEXT("PhysicsControlAssetEditorPreview", "Preview");
+	ViewMenuTooltip = LOCTEXT("PhysicsControlAssetEditorPreview_ToolTip", "Shows the Control Asset Preview tab");
+}
+
+//======================================================================================================================
+TSharedPtr<SToolTip> FPhysicsControlAssetEditorPreviewTabSummoner::CreateTabToolTipWidget(
+	const FWorkflowTabSpawnInfo& Info) const
+{
+	return IDocumentation::Get()->CreateToolTip(LOCTEXT(
+		"PhysicsControlAssetEditorPreviewToolTip",
+		"The Physics Control Asset Preview tab lets you preview the physics control asset setup and profiles."),
+		NULL,
+		TEXT("Shared/Editors/PhysicsControlAssetEditor"), TEXT("PhysicsControlAssetProfiles_Window"));
+}
+
+//======================================================================================================================
+TSharedRef<SWidget> FPhysicsControlAssetEditorPreviewTabSummoner::CreateTabBody(
+	const FWorkflowTabSpawnInfo& Info) const
+{
+	FDetailsViewArgs DetailsViewArgs;
+	DetailsViewArgs.NameAreaSettings = FDetailsViewArgs::HideNameArea;
+	DetailsViewArgs.bHideSelectionTip = true;
+	DetailsViewArgs.bAllowSearch = false;
+
+	FPropertyEditorModule& PropertyEditorModule =
+		FModuleManager::LoadModuleChecked<FPropertyEditorModule>("PropertyEditor");
+	TSharedRef<IDetailsView> DetailsView = PropertyEditorModule.CreateDetailView(DetailsViewArgs);
+
+	TWeakPtr<FPhysicsControlAssetEditor> PhysicsControlAssetEditor =
+		StaticCastSharedPtr<FPhysicsControlAssetEditor>(HostingApp.Pin());
+	DetailsView->RegisterInstancedCustomPropertyLayout(
+		UPhysicsControlAsset::StaticClass(), 
+		FOnGetDetailCustomizationInstance::CreateStatic(
+			&FPhysicsControlAssetPreviewDetailsCustomization::MakeInstance, PhysicsControlAssetEditor));
+	DetailsView->SetObject(PhysicsControlAsset.Get());
+	return DetailsView;
+}
+
+#undef LOCTEXT_NAMESPACE
\ No newline at end of file
diff --git a/Engine/Plugins/Experimental/PhysicsControl/Source/PhysicsControlEditor/Private/PhysicsControlAssetEditorProfilesTabSummoners.h b/Engine/Plugins/Experimental/PhysicsControl/Source/PhysicsControlEditor/Private/PhysicsControlAssetEditorProfilesTabSummoners.h
new file mode 100644
--- /dev/null
+++ b/Engine/Plugins/Experimental/PhysicsControl/Source/PhysicsControlEditor/Private/PhysicsControlAssetEditorProfilesTabSummoners.h
@@ -0,0 +1,58 @@
+// Copyright Epic Games, Inc. All Rights Reserved.
+
+#pragma once
+
+#include "WorkflowOrientedApp/WorkflowTabFactory.h"
+
+class UPhysicsControlAsset;
+
+//======================================================================================================================
+struct FPhysicsControlAssetEditorSetupTabSummoner : public FWorkflowTabFactory
+{
+public:
+	static FName TabName;
+
+	FPhysicsControlAssetEditorSetupTabSummoner(
+		TSharedPtr<FAssetEditorToolkit> InHostingApp, UPhysicsControlAsset* InPhysicsControlAsset);
+
+	/** FWorkflowTabFactory interface */
+	virtual TSharedRef<SWidget> CreateTabBody(const FWorkflowTabSpawnInfo& Info) const override;
+	virtual TSharedPtr<SToolTip> CreateTabToolTipWidget(const FWorkflowTabSpawnInfo& Info) const override;
+
+	/** Reference to our Physics Asset */
+	TWeakObjectPtr<UPhysicsControlAsset> PhysicsControlAsset;
+};
+
+//======================================================================================================================
+struct FPhysicsControlAssetEditorProfileTabSummoner : public FWorkflowTabFactory
+{
+public:
+	static FName TabName;
+
+	FPhysicsControlAssetEditorProfileTabSummoner(
+		TSharedPtr<FAssetEditorToolkit> InHostingApp, UPhysicsControlAsset* InPhysicsControlAsset);
+
+	/** FWorkflowTabFactory interface */
+	virtual TSharedRef<SWidget> CreateTabBody(const FWorkflowTabSpawnInfo& Info) const override;
+	virtual TSharedPtr<SToolTip> CreateTabToolTipWidget(const FWorkflowTabSpawnInfo& Info) const override;
+
+	/** Reference to our Physics Asset */
+	TWeakObjectPtr<UPhysicsControlAsset> PhysicsControlAsset;
+};
+
+//======================================================================================================================
+struct FPhysicsControlAssetEditorPreviewTabSummoner : public FWorkflowTabFactory
+{
+public:
+	static FName TabName;
+
+	FPhysicsControlAssetEditorPreviewTabSummoner(
+		TSharedPtr<FAssetEditorToolkit> InHostingApp, UPhysicsControlAsset* InPhysicsControlAsset);
+
+	/** FWorkflowTabFactory interface */
+	virtual TSharedRef<SWidget> CreateTabBody(const FWorkflowTabSpawnInfo& Info) const override;
+	virtual TSharedPtr<SToolTip> CreateTabToolTipWidget(const FWorkflowTabSpawnInfo& Info) const override;
+
+	/** Reference to our Physics Asset */
+	TWeakObjectPtr<UPhysicsControlAsset> PhysicsControlAsset;
+};
diff --git a/Engine/Plugins/Experimental/PhysicsControl/Source/PhysicsControlEditor/Private/PhysicsControlAssetEditorSkeletalMeshComponent.cpp b/Engine/Plugins/Experimental/PhysicsControl/Source/PhysicsControlEditor/Private/PhysicsControlAssetEditorSkeletalMeshComponent.cpp
new file mode 100644
--- /dev/null
+++ b/Engine/Plugins/Experimental/PhysicsControl/Source/PhysicsControlEditor/Private/PhysicsControlAssetEditorSkeletalMeshComponent.cpp
@@ -0,0 +1,445 @@
+// Copyright Epic Games, Inc. All Rights Reserved.
+
+#include "PhysicsControlAssetEditorSkeletalMeshComponent.h"
+#include "PhysicsControlAsset.h"
+#include "PhysicsControlAssetEditorData.h"
+#include "PhysicsControlAssetEditorHitProxies.h"
+#include "PhysicsControlAssetEditorSkeletalMeshComponent.h"
+#include "PhysicsControlAssetEditorAnimInstance.h"
+#include "PhysicsControlLog.h"
+
+#include "AnimPreviewInstance.h"
+#include "Chaos/Core.h"
+#include "Chaos/Levelset.h"
+#include "Chaos/WeightedLatticeImplicitObject.h"
+#include "Components/SkeletalMeshComponent.h"
+#include "MaterialDomain.h"
+#include "Materials/Material.h"
+#include "Materials/MaterialInstanceDynamic.h"
+#include "Materials/MaterialInterface.h"
+#include "PhysicsAssetRenderUtils.h"
+#include "PhysicsEngine/PhysicsAsset.h"
+#include "PhysicsEngine/SkeletalBodySetup.h"
+#include "Preferences/PhysicsAssetEditorOptions.h"
+#include "SceneManagement.h"
+#include "SkeletalMeshTypes.h"
+#include "Styling/AppStyle.h"
+#include "UObject/Package.h"
+
+namespace
+{
+	bool bDebugViewportClicks = false;
+	FAutoConsoleVariableRef CVarChaosImmPhysStepTime(TEXT("p.PhAT.DebugViewportClicks"), bDebugViewportClicks, TEXT("Set to 1 to show mouse click results in PhAT"));
+}
+
+UPhysicsControlAssetEditorSkeletalMeshComponent::UPhysicsControlAssetEditorSkeletalMeshComponent(const FObjectInitializer& ObjectInitializer)
+	: Super(ObjectInitializer)
+	, BoneUnselectedColor(170, 155, 225)
+	, NoCollisionColor(200, 200, 200)
+	, FixedColor(125, 125, 0)
+	, ConstraintBone1Color(255, 166, 0)
+	, ConstraintBone2Color(0, 150, 150)
+	, HierarchyDrawColor(220, 255, 220)
+	, AnimSkelDrawColor(255, 64, 64)
+	, COMRenderSize(5.0f)
+	, InfluenceLineLength(2.0f)
+	, InfluenceLineColor(0, 255, 0)
+{
+	if (!HasAnyFlags(RF_DefaultSubObject | RF_ArchetypeObject | RF_ClassDefaultObject))
+	{
+		// Body materials
+		UMaterialInterface* BaseElemSelectedMaterial = LoadObject<UMaterialInterface>(NULL, TEXT("/Engine/EditorMaterials/PhAT_ElemSelectedMaterial.PhAT_ElemSelectedMaterial"), NULL, LOAD_None, NULL);
+		ElemSelectedMaterial = UMaterialInstanceDynamic::Create(BaseElemSelectedMaterial, GetTransientPackage());
+		check(ElemSelectedMaterial);
+
+		BoneMaterialHit = UMaterial::GetDefaultMaterial(MD_Surface);
+		check(BoneMaterialHit);
+
+		UMaterialInterface* BaseBoneUnselectedMaterial = LoadObject<UMaterialInterface>(NULL, TEXT("/Engine/EditorMaterials/PhAT_UnselectedMaterial.PhAT_UnselectedMaterial"), NULL, LOAD_None, NULL);
+		BoneUnselectedMaterial = UMaterialInstanceDynamic::Create(BaseBoneUnselectedMaterial, GetTransientPackage());
+		check(BoneUnselectedMaterial);
+
+		UMaterialInterface* BaseBoneNoCollisionMaterial = LoadObject<UMaterialInterface>(NULL, TEXT("/Engine/EditorMaterials/PhAT_NoCollisionMaterial.PhAT_NoCollisionMaterial"), NULL, LOAD_None, NULL);
+		BoneNoCollisionMaterial = UMaterialInstanceDynamic::Create(BaseBoneNoCollisionMaterial, GetTransientPackage());
+		check(BoneNoCollisionMaterial);
+
+		// this is because in phat editor, you'd like to see fixed bones to be fixed without animation force update
+		KinematicBonesUpdateType = EKinematicBonesUpdateToPhysics::SkipSimulatingBones;
+		bUpdateJointsFromAnimation = false;
+		SetForcedLOD(1);
+
+		static FName CollisionProfileName(TEXT("PhysicsActor"));
+		SetCollisionProfileName(CollisionProfileName);
+	}
+
+	bSelectable = false;
+}
+
+TObjectPtr<UAnimPreviewInstance> UPhysicsControlAssetEditorSkeletalMeshComponent::CreatePreviewInstance()
+{
+	return NewObject<UPhysicsControlAssetEditorAnimInstance>(this, TEXT("PhysicsAssetEditorPreviewInstance"));
+}
+
+void UPhysicsControlAssetEditorSkeletalMeshComponent::DebugDraw(const FSceneView* View, FPrimitiveDrawInterface* PDI)
+{
+	check(EditorData);
+
+	UPhysicsAsset* const PhysicsAsset = GetPhysicsAsset();
+
+	if (!PhysicsAsset)
+	{
+		// Nothing to draw without an asset, this can happen if the preview scene has no skeletal mesh
+		return;
+	}
+
+	EPhysicsAssetEditorCollisionViewMode CollisionViewMode = EditorData->GetCurrentCollisionViewMode(EditorData->bRunningSimulation);
+
+	// set opacity of our materials
+	static FName OpacityName(TEXT("Opacity"));
+	ElemSelectedMaterial->SetScalarParameterValue(OpacityName, EditorData->EditorOptions->CollisionOpacity);
+	BoneUnselectedMaterial->SetScalarParameterValue(OpacityName, EditorData->EditorOptions->bSolidRenderingForSelectedOnly ? 0.0f : EditorData->EditorOptions->CollisionOpacity);
+	BoneNoCollisionMaterial->SetScalarParameterValue(OpacityName, EditorData->EditorOptions->bSolidRenderingForSelectedOnly ? 0.0f : EditorData->EditorOptions->CollisionOpacity);
+
+	static FName SelectionColorName(TEXT("SelectionColor"));
+	const FSlateColor SelectionColor = FAppStyle::GetSlateColor(SelectionColorName);
+	const FLinearColor LinearSelectionColor(SelectionColor.IsColorSpecified() ? SelectionColor.GetSpecifiedColor() : FLinearColor::White);
+
+	ElemSelectedMaterial->SetVectorParameterValue(SelectionColorName, LinearSelectionColor);
+
+	FPhysicsAssetRenderSettings* const RenderSettings = UPhysicsAssetRenderUtilities::GetSettings(PhysicsAsset);
+	
+	if (RenderSettings)
+	{
+		// Copy render settings from editor viewport. These settings must be applied to the rendering in all editors 
+		// when an asset is open in the Physics Asset Editor but should not persist after the editor has been closed.
+		RenderSettings->CollisionViewMode = EditorData->GetCurrentCollisionViewMode(EditorData->bRunningSimulation);
+		RenderSettings->ConstraintViewMode = EditorData->GetCurrentConstraintViewMode(EditorData->bRunningSimulation);
+		RenderSettings->ConstraintDrawSize = EditorData->EditorOptions->ConstraintDrawSize;
+		RenderSettings->PhysicsBlend = EditorData->EditorOptions->PhysicsBlend;
+		RenderSettings->bHideKinematicBodies = EditorData->EditorOptions->bHideKinematicBodies;
+		RenderSettings->bHideSimulatedBodies = EditorData->EditorOptions->bHideSimulatedBodies;
+		RenderSettings->bDrawViolatedLimits = EditorData->EditorOptions->bDrawViolatedLimits;
+
+		// Draw Bodies.
+		{
+			auto TransformFn = [this](
+				const UPhysicsAsset* PhysicsAsset, const FTransform& BoneTM, const int32 BodyIndex, 
+				const EAggCollisionShape::Type PrimType, const int32 PrimIndex, const float Scale) 
+				{ 
+					return this->GetPrimitiveTransform(BoneTM, BodyIndex, PrimType, PrimIndex, Scale);  
+				};
+			auto ColorFn = [this](
+				const int32 BodyIndex, const EAggCollisionShape::Type PrimitiveType, const int32 PrimitiveIndex, 
+				const FPhysicsAssetRenderSettings& Settings) 
+				{ 
+					return this->GetPrimitiveColor(BodyIndex, PrimitiveType, PrimitiveIndex); 
+				};
+			auto MaterialFn = [this](
+				const int32 BodyIndex, const EAggCollisionShape::Type PrimitiveType, const int32 PrimitiveIndex, 
+				const FPhysicsAssetRenderSettings& Settings) 
+				{ 
+					return this->GetPrimitiveMaterial(BodyIndex, PrimitiveType, PrimitiveIndex); 
+				};
+			auto HitProxyFn = [](
+				const int32 BodyIndex, const EAggCollisionShape::Type PrimitiveType, const int32 PrimitiveIndex)
+				{
+					return new HPhysicsControlAssetEditorEdBoneProxy(BodyIndex, PrimitiveType, PrimitiveIndex); 
+				};
+
+			PhysicsAssetRender::DebugDrawBodies(this, PhysicsAsset, PDI, ColorFn, MaterialFn, TransformFn, HitProxyFn);
+		}
+
+		// Draw Constraints.
+		{
+			auto HitProxyFn = [](const int32 InConstraintIndex) { return nullptr; };
+			auto IsConstraintSelectedFn = [this](const uint32 InConstraintIndex) { return false; };
+
+			PhysicsAssetRender::DebugDrawConstraints(this, PhysicsAsset, PDI, IsConstraintSelectedFn, EditorData->bRunningSimulation, HitProxyFn);
+		}
+	}
+}
+
+FPrimitiveSceneProxy* UPhysicsControlAssetEditorSkeletalMeshComponent::CreateSceneProxy()
+{
+	FPrimitiveSceneProxy* Proxy = NULL;
+	EPhysicsAssetEditorMeshViewMode MeshViewMode = EditorData->GetCurrentMeshViewMode(EditorData->bRunningSimulation);
+	if (MeshViewMode != EPhysicsAssetEditorMeshViewMode::None)
+	{
+		Proxy = UDebugSkelMeshComponent::CreateSceneProxy();
+	}
+
+	return Proxy;
+}
+
+FTransform UPhysicsControlAssetEditorSkeletalMeshComponent::GetPrimitiveTransform(const FTransform& BoneTM, const int32 BodyIndex, const EAggCollisionShape::Type PrimType, const int32 PrimIndex, const float Scale) const
+{
+	UPhysicsAsset* PhysicsAsset = EditorData->PhysicsControlAsset->GetPhysicsAsset();
+	if (!PhysicsAsset)
+	{
+		UE_LOG(LogPhysicsControl, Warning, TEXT("GetPrimitiveTransform - no physics asset"));
+		return FTransform::Identity;
+	}
+
+	UBodySetup* SharedBodySetup = PhysicsAsset->SkeletalBodySetups[BodyIndex];
+	FVector Scale3D(Scale);
+
+	FTransform ManTM = FTransform::Identity;
+
+	if (PrimType == EAggCollisionShape::Sphere)
+	{
+		FTransform PrimTM = ManTM * SharedBodySetup->AggGeom.SphereElems[PrimIndex].GetTransform();
+		PrimTM.ScaleTranslation(Scale3D);
+		return PrimTM * BoneTM;
+	}
+	else if (PrimType == EAggCollisionShape::Box)
+	{
+		FTransform PrimTM = ManTM * SharedBodySetup->AggGeom.BoxElems[PrimIndex].GetTransform();
+		PrimTM.ScaleTranslation(Scale3D);
+		return PrimTM * BoneTM;
+	}
+	else if (PrimType == EAggCollisionShape::Sphyl)
+	{
+		FTransform PrimTM = ManTM * SharedBodySetup->AggGeom.SphylElems[PrimIndex].GetTransform();
+		PrimTM.ScaleTranslation(Scale3D);
+		return PrimTM * BoneTM;
+	}
+	else if (PrimType == EAggCollisionShape::Convex)
+	{
+		FTransform PrimTM = ManTM * SharedBodySetup->AggGeom.ConvexElems[PrimIndex].GetTransform();
+		PrimTM.ScaleTranslation(Scale3D);
+		return PrimTM * BoneTM;
+	}
+	else if (PrimType == EAggCollisionShape::TaperedCapsule)
+	{
+		FTransform PrimTM = ManTM * SharedBodySetup->AggGeom.TaperedCapsuleElems[PrimIndex].GetTransform();
+		PrimTM.ScaleTranslation(Scale3D);
+		return PrimTM * BoneTM;
+	}
+	else if (PrimType == EAggCollisionShape::LevelSet)
+	{
+		FTransform PrimTM = ManTM * SharedBodySetup->AggGeom.LevelSetElems[PrimIndex].GetTransform();
+		PrimTM.ScaleTranslation(Scale3D);
+		return PrimTM * BoneTM;
+	}
+	else if (PrimType == EAggCollisionShape::SkinnedLevelSet)
+	{
+		FTransform PrimTM = ManTM * SharedBodySetup->AggGeom.SkinnedLevelSetElems[PrimIndex].GetTransform();
+		PrimTM.ScaleTranslation(Scale3D);
+		return PrimTM * BoneTM;
+	}
+
+	// Should never reach here
+	check(0);
+	return FTransform::Identity;
+}
+
+FColor UPhysicsControlAssetEditorSkeletalMeshComponent::GetPrimitiveColor(const int32 BodyIndex, const EAggCollisionShape::Type PrimitiveType, const int32 PrimitiveIndex) const
+{
+	UPhysicsAsset* PhysicsAsset = EditorData->PhysicsControlAsset->GetPhysicsAsset();
+	if (!PhysicsAsset)
+	{
+		UE_LOG(LogPhysicsControl, Warning, TEXT("GetPrimitiveColor - no physics asset"));
+		return BoneUnselectedColor;
+	}
+
+	UBodySetup* SharedBodySetup = PhysicsAsset->SkeletalBodySetups[BodyIndex];
+
+	FPhysicsControlAssetEditorData::FSelection Body(BodyIndex, PrimitiveType, PrimitiveIndex);
+
+	static FName SelectionColorName(TEXT("SelectionColor"));
+	const FSlateColor SelectionColor = FAppStyle::GetSlateColor(SelectionColorName);
+	const FLinearColor SelectionColorLinear(SelectionColor.IsColorSpecified() ? SelectionColor.GetSpecifiedColor() : FLinearColor::White);
+	const FColor ElemSelectedColor = SelectionColorLinear.ToFColor(true);
+	const FColor ElemSelectedBodyColor = (SelectionColorLinear* 0.5f).ToFColor(true);
+
+	bool bInBody = false;
+	for (int32 i = 0; i<EditorData->SelectedBodies.Num(); ++i)
+	{
+		if (BodyIndex == EditorData->SelectedBodies[i].Index)
+		{
+			bInBody = true;
+		}
+
+		if (Body == EditorData->SelectedBodies[i] && !EditorData->bRunningSimulation)
+		{
+			return ElemSelectedColor;
+		}
+	}
+
+	if (bInBody && !EditorData->bRunningSimulation)	//this primitive is in a body that's currently selected, but this primitive itself isn't selected
+	{
+		return ElemSelectedBodyColor;
+	}
+	if(PrimitiveType == EAggCollisionShape::TaperedCapsule)
+	{
+		return NoCollisionColor;
+	}
+
+	if (EditorData->bRunningSimulation)
+	{
+		const bool bIsSimulatedAtAll = SharedBodySetup->PhysicsType == PhysType_Simulated || (SharedBodySetup->PhysicsType == PhysType_Default && EditorData->EditorOptions->PhysicsBlend > 0.f);
+		if (!bIsSimulatedAtAll)
+		{
+			return FixedColor;
+		}
+	}
+
+	return BoneUnselectedColor;
+}
+
+UMaterialInterface* UPhysicsControlAssetEditorSkeletalMeshComponent::GetPrimitiveMaterial(const int32 BodyIndex, const EAggCollisionShape::Type PrimitiveType, const int32 PrimitiveIndex) const
+{
+	if (EditorData->bRunningSimulation)
+	{
+		return PrimitiveType == EAggCollisionShape::TaperedCapsule ? BoneNoCollisionMaterial : BoneUnselectedMaterial;
+	}
+
+	FPhysicsControlAssetEditorData::FSelection Body(BodyIndex, PrimitiveType, PrimitiveIndex);
+
+	for (int32 i = 0; i < EditorData->SelectedBodies.Num(); ++i)
+	{
+		if (Body == EditorData->SelectedBodies[i] && !EditorData->bRunningSimulation)
+		{
+			return ElemSelectedMaterial;
+		}
+	}
+
+	if (PrimitiveType == EAggCollisionShape::TaperedCapsule)
+	{
+		return BoneNoCollisionMaterial;
+	}
+
+	return BoneUnselectedMaterial;
+}
+
+void UPhysicsControlAssetEditorSkeletalMeshComponent::RefreshBoneTransforms(FActorComponentTickFunction* TickFunction)
+{
+	Super::RefreshBoneTransforms(TickFunction);
+
+	// Horrible kludge, but we need to flip the buffer back here as we need to wait on the physics tick group.
+	// However UDebugSkelMeshComponent passes NULL to force non-threaded work, which assumes a flip is needed straight away
+	if (ShouldBlendPhysicsBones())
+	{
+		bNeedToFlipSpaceBaseBuffers = true;
+		FinalizeBoneTransform();
+		bNeedToFlipSpaceBaseBuffers = true;
+	}
+	UpdateSkinnedLevelSets();
+}
+
+void UPhysicsControlAssetEditorSkeletalMeshComponent::AddImpulseAtLocation(FVector Impulse, FVector Location, FName BoneName)
+{
+	if (PreviewInstance != nullptr)
+	{
+		PreviewInstance->AddImpulseAtLocation(Impulse, Location, BoneName);
+	}
+}
+
+bool UPhysicsControlAssetEditorSkeletalMeshComponent::ShouldCreatePhysicsState() const
+{
+	// @todo(chaos): the main physics scene is not running (and never runs) in the physics editor,
+	// and currently this means it will accumulate body create/destroy commands every time
+	// we hit "Simulate". Fix this!  However, we still need physics state for mouse ray hit detection 
+	// on the bodies so we can't just avoid creating physics state...
+	return Super::ShouldCreatePhysicsState();
+}
+
+
+void UPhysicsControlAssetEditorSkeletalMeshComponent::Grab(FName InBoneName, const FVector& Location, const FRotator& Rotation, bool bRotationConstrained)
+{
+	UPhysicsControlAssetEditorAnimInstance* PhysicsControlAssetEditorAnimInstance = 
+		Cast<UPhysicsControlAssetEditorAnimInstance>(PreviewInstance);
+	if (PhysicsControlAssetEditorAnimInstance != nullptr)
+	{
+		PhysicsControlAssetEditorAnimInstance->Grab(InBoneName, Location, Rotation, bRotationConstrained);
+	}
+}
+
+void UPhysicsControlAssetEditorSkeletalMeshComponent::Ungrab()
+{
+	UPhysicsControlAssetEditorAnimInstance* PhysicsControlAssetEditorAnimInstance = 
+		Cast<UPhysicsControlAssetEditorAnimInstance>(PreviewInstance);
+	if (PhysicsControlAssetEditorAnimInstance != nullptr)
+	{
+		PhysicsControlAssetEditorAnimInstance->Ungrab();
+	}
+}
+
+void UPhysicsControlAssetEditorSkeletalMeshComponent::UpdateHandleTransform(const FTransform& NewTransform)
+{
+	UPhysicsControlAssetEditorAnimInstance* PhysicsControlAssetEditorAnimInstance = 
+		Cast<UPhysicsControlAssetEditorAnimInstance>(PreviewInstance);
+	if (PhysicsControlAssetEditorAnimInstance != nullptr)
+	{
+		PhysicsControlAssetEditorAnimInstance->UpdateHandleTransform(NewTransform);
+	}
+}
+
+void UPhysicsControlAssetEditorSkeletalMeshComponent::UpdateDriveSettings(
+	bool bLinearSoft, float LinearStiffness, float LinearDamping)
+{
+	UPhysicsControlAssetEditorAnimInstance* PhysicsControlAssetEditorAnimInstance = 
+		Cast<UPhysicsControlAssetEditorAnimInstance>(PreviewInstance);
+	if (PhysicsControlAssetEditorAnimInstance != nullptr)
+	{
+		PhysicsControlAssetEditorAnimInstance->UpdateDriveSettings(bLinearSoft, LinearStiffness, LinearDamping);
+	}
+}
+
+void UPhysicsControlAssetEditorSkeletalMeshComponent::CreateSimulationFloor(
+	FBodyInstance* FloorBodyInstance, const FTransform& Transform)
+{
+	UPhysicsControlAssetEditorAnimInstance* PhysicsControlAssetEditorAnimInstance = 
+		Cast<UPhysicsControlAssetEditorAnimInstance>(PreviewInstance);
+	if (PhysicsControlAssetEditorAnimInstance != nullptr)
+	{
+		PhysicsControlAssetEditorAnimInstance->CreateSimulationFloor(FloorBodyInstance, Transform);
+	}
+}
+
+void UPhysicsControlAssetEditorSkeletalMeshComponent::UpdateSkinnedLevelSets()
+{
+	UPhysicsAsset* const PhysicsAsset = GetPhysicsAsset();
+	if (!PhysicsAsset)
+	{
+		return;
+	}
+	for (int32 i = 0; i < PhysicsAsset->SkeletalBodySetups.Num(); ++i)
+	{
+		const int32 BoneIndex = GetBoneIndex(PhysicsAsset->SkeletalBodySetups[i]->BoneName);
+		if (BoneIndex != INDEX_NONE)
+		{
+			FKAggregateGeom* const AggGeom = &PhysicsAsset->SkeletalBodySetups[i]->AggGeom;
+			if (AggGeom)
+			{
+				for (FKSkinnedLevelSetElem& SkinnedLevelSet : AggGeom->SkinnedLevelSetElems)
+				{
+					if (SkinnedLevelSet.WeightedLevelSet().IsValid())
+					{
+						const TArray<FName>& UsedBoneNames = SkinnedLevelSet.WeightedLevelSet()->GetUsedBones();
+
+						const FTransform RootTransformInv = GetBoneTransform(BoneIndex, FTransform::Identity).Inverse();
+						TArray<FTransform> Transforms;
+						Transforms.SetNum(UsedBoneNames.Num());
+
+						for (int32 LocalIdx = 0; LocalIdx < UsedBoneNames.Num(); ++LocalIdx)
+						{
+							const int32 LocalBoneIndex = GetBoneIndex(UsedBoneNames[LocalIdx]);
+							if (LocalBoneIndex != INDEX_NONE)
+							{
+								const FTransform BoneTransformTimesRootTransformInv = GetBoneTransform(LocalBoneIndex, RootTransformInv);
+								Transforms[LocalIdx] = BoneTransformTimesRootTransformInv;
+							}
+							else
+							{
+								Transforms[LocalIdx] = RootTransformInv;
+							}
+						}
+
+						SkinnedLevelSet.WeightedLevelSet()->DeformPoints(Transforms);
+					}
+				}
+			}
+		}
+	}
+}
\ No newline at end of file
diff --git a/Engine/Plugins/Experimental/PhysicsControl/Source/PhysicsControlEditor/Private/PhysicsControlAssetEditorSkeletalMeshComponent.h b/Engine/Plugins/Experimental/PhysicsControl/Source/PhysicsControlEditor/Private/PhysicsControlAssetEditorSkeletalMeshComponent.h
new file mode 100644
--- /dev/null
+++ b/Engine/Plugins/Experimental/PhysicsControl/Source/PhysicsControlEditor/Private/PhysicsControlAssetEditorSkeletalMeshComponent.h
@@ -0,0 +1,85 @@
+// Copyright Epic Games, Inc. All Rights Reserved.
+
+#pragma once
+
+#include "CoreMinimal.h"
+#include "UObject/ObjectMacros.h"
+#include "PhysicsEngine/ShapeElem.h"
+#include "Animation/DebugSkelMeshComponent.h"
+#include "PhysicsControlAssetEditorSkeletalMeshComponent.generated.h"
+
+class FPrimitiveDrawInterface;
+class UMaterialInterface;
+class UMaterialInstanceDynamic;
+class FPhysicsControlAssetEditorData;
+
+UCLASS()
+class UPhysicsControlAssetEditorSkeletalMeshComponent : public UDebugSkelMeshComponent
+{
+	GENERATED_UCLASS_BODY()
+
+	/** Data and methods shared across multiple classes */
+	TSharedPtr<FPhysicsControlAssetEditorData> EditorData;
+
+	// Draw colors
+
+	FColor BoneUnselectedColor;
+	FColor NoCollisionColor;
+	FColor FixedColor;
+	FColor ConstraintBone1Color;
+	FColor ConstraintBone2Color;
+	FColor HierarchyDrawColor;
+	FColor AnimSkelDrawColor;
+	float COMRenderSize;
+	float InfluenceLineLength;
+	FColor InfluenceLineColor;
+
+	// Materials
+
+	UPROPERTY(transient)
+	TObjectPtr<UMaterialInstanceDynamic> ElemSelectedMaterial;
+	UPROPERTY(transient)
+	TObjectPtr<UMaterialInstanceDynamic> BoneUnselectedMaterial;
+	UPROPERTY(transient)
+	TObjectPtr<UMaterialInstanceDynamic> BoneNoCollisionMaterial;
+
+	UPROPERTY(transient)
+	TObjectPtr<UMaterialInterface> BoneMaterialHit;
+
+	/** Mesh-space matrices showing state of just animation (ie before physics) - useful for debugging! */
+	TArray<FTransform> AnimationSpaceBases;
+
+	/** UDebugSkelMeshComponent interface */
+	virtual TObjectPtr<UAnimPreviewInstance> CreatePreviewInstance() override;
+
+	/** UPrimitiveComponent interface */
+	virtual FPrimitiveSceneProxy* CreateSceneProxy() override;
+	virtual void AddImpulseAtLocation(FVector Impulse, FVector Location, FName BoneName = NAME_None) override;
+	virtual bool ShouldCreatePhysicsState() const override;
+
+	/** USkinnedMeshComponent interface */
+	virtual void RefreshBoneTransforms(FActorComponentTickFunction* TickFunction = nullptr) override;
+	
+	/** Debug drawing */
+	void DebugDraw(const FSceneView* View, FPrimitiveDrawInterface* PDI);
+
+	/** Accessors/helper methods */
+	FTransform GetPrimitiveTransform(const FTransform& BoneTM, const int32 BodyIndex, const EAggCollisionShape::Type PrimType, const int32 PrimIndex, const float Scale) const;
+	FColor GetPrimitiveColor(const int32 BodyIndex, const EAggCollisionShape::Type PrimitiveType, const int32 PrimitiveIndex) const;
+	UMaterialInterface* GetPrimitiveMaterial(const int32 BodyIndex, const EAggCollisionShape::Type PrimitiveType, const int32 PrimitiveIndex) const;
+
+	/** Manipulator methods */
+	void Grab(FName InBoneName, const FVector& Location, const FRotator& Rotation, bool bRotationConstrained);
+	void Ungrab();
+	void UpdateHandleTransform(const FTransform& NewTransform);
+	void UpdateDriveSettings(bool bLinearSoft, float LinearStiffness, float LinearDamping);
+
+	/** Sim setup */
+	void CreateSimulationFloor(FBodyInstance* FloorBodyInstance, const FTransform& Transform);
+
+public:
+	virtual bool CanOverrideCollisionProfile() const override { return false;  }
+
+private:
+	void UpdateSkinnedLevelSets();
+};
diff --git a/Engine/Plugins/Experimental/PhysicsControl/Source/PhysicsControlEditor/Private/PhysicsControlAssetEditorSkeletonTreeBuilder.cpp b/Engine/Plugins/Experimental/PhysicsControl/Source/PhysicsControlEditor/Private/PhysicsControlAssetEditorSkeletonTreeBuilder.cpp
new file mode 100644
--- /dev/null
+++ b/Engine/Plugins/Experimental/PhysicsControl/Source/PhysicsControlEditor/Private/PhysicsControlAssetEditorSkeletonTreeBuilder.cpp
@@ -0,0 +1,186 @@
+// Copyright Epic Games, Inc. All Rights Reserved.
+
+#include "PhysicsControlAssetEditorSkeletonTreeBuilder.h"
+#include "SkeletonTreePhysicsControlBodyItem.h"
+#include "SkeletonTreePhysicsControlShapeItem.h"
+#include "PhysicsEngine/PhysicsAsset.h"
+#include "PhysicsEngine/SkeletalBodySetup.h"
+#include "IPersonaPreviewScene.h"
+#include "Animation/DebugSkelMeshComponent.h"
+#include "Misc/TextFilterExpressionEvaluator.h"
+
+#define LOCTEXT_NAMESPACE "PhysicsControlAssetEditorSkeletonTreeBuilder"
+
+FPhysicsControlAssetEditorSkeletonTreeBuilder::FPhysicsControlAssetEditorSkeletonTreeBuilder(
+	UPhysicsAsset*                  InPhysicsAsset, 
+	const FSkeletonTreeBuilderArgs& InSkeletonTreeBuilderArgs)
+	: FSkeletonTreeBuilder(InSkeletonTreeBuilderArgs)
+	, bShowBodies(true)
+	, bShowKinematicBodies(true)
+	, bShowSimulatedBodies(true)
+	, bShowPrimitives(false)
+	, PhysicsAsset(InPhysicsAsset)
+{
+}
+
+void FPhysicsControlAssetEditorSkeletonTreeBuilder::Build(FSkeletonTreeBuilderOutput& Output)
+{
+	if (BuilderArgs.bShowBones)
+	{
+		AddBones(Output);
+	}
+
+	AddBodies(Output);
+
+	if (BuilderArgs.bShowAttachedAssets)
+	{
+		AddAttachedAssets(Output);
+	}
+}
+
+ESkeletonTreeFilterResult FPhysicsControlAssetEditorSkeletonTreeBuilder::FilterItem(const FSkeletonTreeFilterArgs& InArgs, const TSharedPtr<class ISkeletonTreeItem>& InItem)
+{
+	if (InItem->IsOfType<FSkeletonTreePhysicsControlBodyItem>() || InItem->IsOfType<FSkeletonTreePhysicsControlShapeItem>())
+	{
+		ESkeletonTreeFilterResult Result = ESkeletonTreeFilterResult::Shown;
+
+		if (InArgs.TextFilter.IsValid())
+		{
+			if (InArgs.TextFilter->TestTextFilter(FBasicStringFilterExpressionContext(InItem->GetRowItemName().ToString())))
+			{
+				Result = ESkeletonTreeFilterResult::ShownHighlighted;
+			}
+			else
+			{
+				Result = ESkeletonTreeFilterResult::Hidden;
+			}
+		}
+
+		if (InItem->IsOfType<FSkeletonTreePhysicsControlBodyItem>())
+		{
+			bool bShouldHideBody = false;
+			if (!bShowBodies)
+			{
+				bShouldHideBody = true;
+			}
+			else
+			{
+				if (UBodySetup* BodySetup = Cast<UBodySetup>(InItem->GetObject()))
+				{
+					if (BodySetup->PhysicsType == EPhysicsType::PhysType_Simulated && !bShowSimulatedBodies)
+					{
+						bShouldHideBody = true;
+					}
+					else if (BodySetup->PhysicsType == EPhysicsType::PhysType_Kinematic && !bShowKinematicBodies)
+					{
+						bShouldHideBody = true;
+					}
+				}
+			}
+			if (bShouldHideBody)
+			{
+				Result = ESkeletonTreeFilterResult::Hidden;
+			}
+		}
+		else if (InItem->IsOfType<FSkeletonTreePhysicsControlShapeItem>())
+		{
+			if (!bShowPrimitives)
+			{
+				Result = ESkeletonTreeFilterResult::Hidden;
+			}
+		}
+
+		return Result;
+	}
+
+	return FSkeletonTreeBuilder::FilterItem(InArgs, InItem);
+}
+
+void FPhysicsControlAssetEditorSkeletonTreeBuilder::AddBodies(FSkeletonTreeBuilderOutput& Output)
+{
+	if (!PreviewScenePtr.IsValid())
+	{
+		return;
+	}
+	UDebugSkelMeshComponent* PreviewMeshComponent = PreviewScenePtr.Pin()->GetPreviewMeshComponent();
+	if (!PreviewMeshComponent->GetSkeletalMeshAsset())
+	{
+		return;
+	}
+	FReferenceSkeleton& RefSkeleton = PreviewMeshComponent->GetSkeletalMeshAsset()->GetRefSkeleton();
+
+	for (int32 BoneIndex = 0; BoneIndex < RefSkeleton.GetRawBoneNum(); ++BoneIndex)
+	{
+		const FName BoneName = RefSkeleton.GetBoneName(BoneIndex);
+		int32 ParentIndex = RefSkeleton.GetParentIndex(BoneIndex);
+		const FName ParentName = ParentIndex == INDEX_NONE ? NAME_None : RefSkeleton.GetBoneName(ParentIndex);
+
+		bool bHasBodySetup = false;
+		for (int32 BodySetupIndex = 0; BodySetupIndex < PhysicsAsset->SkeletalBodySetups.Num(); ++BodySetupIndex)
+		{
+			if (!ensure(PhysicsAsset->SkeletalBodySetups[BodySetupIndex]))
+			{
+				continue;
+			}
+			if (BoneName != PhysicsAsset->SkeletalBodySetups[BodySetupIndex]->BoneName)
+			{
+				continue;
+			}
+
+			USkeletalBodySetup* BodySetup = PhysicsAsset->SkeletalBodySetups[BodySetupIndex];
+
+			bHasBodySetup = true;
+
+			const FKAggregateGeom& AggGeom = PhysicsAsset->SkeletalBodySetups[BodySetupIndex]->AggGeom;
+
+			bool bHasShapes = AggGeom.GetElementCount() > 0;
+
+			if (bHasShapes)
+			{
+				Output.Add(MakeShared<FSkeletonTreePhysicsControlBodyItem>(
+					BodySetup, BodySetupIndex, BoneName, true, bHasShapes, PhysicsAsset, 
+					SkeletonTreePtr.Pin().ToSharedRef()), BoneName, "FSkeletonTreeBoneItem", true);
+
+				int32 ShapeIndex;
+				for (ShapeIndex = 0; ShapeIndex < AggGeom.SphereElems.Num(); ++ShapeIndex)
+				{
+					Output.Add(MakeShared<FSkeletonTreePhysicsControlShapeItem>(
+						BodySetup, BoneName, BodySetupIndex, EAggCollisionShape::Sphere, ShapeIndex, 
+						SkeletonTreePtr.Pin().ToSharedRef()), BoneName, FSkeletonTreePhysicsControlBodyItem::GetTypeId());
+				}
+
+				for (ShapeIndex = 0; ShapeIndex < AggGeom.BoxElems.Num(); ++ShapeIndex)
+				{
+					Output.Add(MakeShared<FSkeletonTreePhysicsControlShapeItem>(
+						BodySetup, BoneName, BodySetupIndex, EAggCollisionShape::Box, ShapeIndex, 
+						SkeletonTreePtr.Pin().ToSharedRef()), BoneName, FSkeletonTreePhysicsControlBodyItem::GetTypeId());
+				}
+
+				for (ShapeIndex = 0; ShapeIndex < AggGeom.SphylElems.Num(); ++ShapeIndex)
+				{
+					Output.Add(MakeShared<FSkeletonTreePhysicsControlShapeItem>(
+						BodySetup, BoneName, BodySetupIndex, EAggCollisionShape::Sphyl, ShapeIndex, 
+						SkeletonTreePtr.Pin().ToSharedRef()), BoneName, FSkeletonTreePhysicsControlBodyItem::GetTypeId());
+				}
+
+				for (ShapeIndex = 0; ShapeIndex < AggGeom.ConvexElems.Num(); ++ShapeIndex)
+				{
+					Output.Add(MakeShared<FSkeletonTreePhysicsControlShapeItem>(
+						BodySetup, BoneName, BodySetupIndex, EAggCollisionShape::Convex, ShapeIndex, 
+						SkeletonTreePtr.Pin().ToSharedRef()), BoneName, FSkeletonTreePhysicsControlBodyItem::GetTypeId());
+				}
+
+				for (ShapeIndex = 0; ShapeIndex < AggGeom.TaperedCapsuleElems.Num(); ++ShapeIndex)
+				{
+					Output.Add(MakeShared<FSkeletonTreePhysicsControlShapeItem>(
+						BodySetup, BoneName, BodySetupIndex, EAggCollisionShape::TaperedCapsule, ShapeIndex, 
+						SkeletonTreePtr.Pin().ToSharedRef()), BoneName, FSkeletonTreePhysicsControlBodyItem::GetTypeId());
+				}
+			}
+
+			break;
+		}
+	}
+}
+
+#undef LOCTEXT_NAMESPACE
diff --git a/Engine/Plugins/Experimental/PhysicsControl/Source/PhysicsControlEditor/Private/PhysicsControlAssetEditorSkeletonTreeBuilder.h b/Engine/Plugins/Experimental/PhysicsControl/Source/PhysicsControlEditor/Private/PhysicsControlAssetEditorSkeletonTreeBuilder.h
new file mode 100644
--- /dev/null
+++ b/Engine/Plugins/Experimental/PhysicsControl/Source/PhysicsControlEditor/Private/PhysicsControlAssetEditorSkeletonTreeBuilder.h
@@ -0,0 +1,34 @@
+// Copyright Epic Games, Inc. All Rights Reserved.
+
+#pragma once
+
+#include "CoreMinimal.h"
+#include "SkeletonTreeBuilder.h"
+
+class UPhysicsAsset;
+
+class FPhysicsControlAssetEditorSkeletonTreeBuilder : public FSkeletonTreeBuilder
+{
+public:
+	FPhysicsControlAssetEditorSkeletonTreeBuilder(
+		UPhysicsAsset*                  InPhysicsAsset, 
+		const FSkeletonTreeBuilderArgs& InSkeletonTreeBuilderArgs = FSkeletonTreeBuilderArgs(true, false, true, false));
+
+	/** ISkeletonTreeBuilder interface */
+	virtual void Build(FSkeletonTreeBuilderOutput& Output) override;
+	virtual ESkeletonTreeFilterResult FilterItem(const FSkeletonTreeFilterArgs& InArgs, const TSharedPtr<class ISkeletonTreeItem>& InItem) override;
+
+	/** Flags used for filtering */
+	bool bShowBodies;
+	bool bShowKinematicBodies;
+	bool bShowSimulatedBodies;
+	bool bShowPrimitives;
+
+private:
+	/** Add bodies & their shapes to the tree */
+	void AddBodies(FSkeletonTreeBuilderOutput& Output);
+
+private:
+	/** The physics asset we use to populate the tree */
+	UPhysicsAsset* PhysicsAsset;
+};
diff --git a/Engine/Plugins/Experimental/PhysicsControl/Source/PhysicsControlEditor/Private/PhysicsControlAssetEditorToolMenuContext.h b/Engine/Plugins/Experimental/PhysicsControl/Source/PhysicsControlEditor/Private/PhysicsControlAssetEditorToolMenuContext.h
new file mode 100644
--- /dev/null
+++ b/Engine/Plugins/Experimental/PhysicsControl/Source/PhysicsControlEditor/Private/PhysicsControlAssetEditorToolMenuContext.h
@@ -0,0 +1,17 @@
+// Copyright Epic Games, Inc. All Rights Reserved.
+
+#pragma once
+
+#include "CoreMinimal.h"
+#include "PhysicsControlAssetEditorToolMenuContext.generated.h"
+
+class FPhysicsControlAssetEditor;
+
+UCLASS()
+class UPhysicsControlAssetEditorToolMenuContext : public UObject
+{
+	GENERATED_BODY()
+
+public:
+	TWeakPtr<FPhysicsControlAssetEditor> PhysicsControlAssetEditor;
+};
diff --git a/Engine/Plugins/Experimental/PhysicsControl/Source/PhysicsControlEditor/Private/PhysicsControlAssetEditorToolkit.cpp b/Engine/Plugins/Experimental/PhysicsControl/Source/PhysicsControlEditor/Private/PhysicsControlAssetEditorToolkit.cpp
deleted file mode 100644
diff --git a/Engine/Plugins/Experimental/PhysicsControl/Source/PhysicsControlEditor/Private/PhysicsControlAssetPreviewDetailsCustomization.cpp b/Engine/Plugins/Experimental/PhysicsControl/Source/PhysicsControlEditor/Private/PhysicsControlAssetPreviewDetailsCustomization.cpp
new file mode 100644
--- /dev/null
+++ b/Engine/Plugins/Experimental/PhysicsControl/Source/PhysicsControlEditor/Private/PhysicsControlAssetPreviewDetailsCustomization.cpp
@@ -0,0 +1,99 @@
+// Copyright Epic Games, Inc. All Rights Reserved.
+
+#include "PhysicsControlAssetPreviewDetailsCustomization.h"
+
+#include "PhysicsControlAssetActions.h"
+#include "PhysicsControlAssetEditor.h"
+#include "PhysicsControlAssetEditorData.h"
+#include "PhysicsControlAssetEditorCommands.h"
+#include "PhysicsControlAsset.h"
+#include "PhysicsControlComponent.h"
+
+#include "DetailCategoryBuilder.h"
+#include "DetailLayoutBuilder.h"
+#include "DetailWidgetRow.h"
+#include "EditorFontGlyphs.h"
+#include "Framework/Application/SlateApplication.h"
+#include "Framework/MultiBox/MultiBoxBuilder.h"
+#include "IDetailPropertyRow.h"
+#include "PropertyHandle.h"
+#include "ScopedTransaction.h"
+#include "Styling/StyleColors.h"
+#include "Widgets/Images/SImage.h"
+#include "Widgets/Input/SButton.h"
+#include "Widgets/Input/SComboButton.h"
+#include "Widgets/Input/SEditableTextBox.h"
+#include "Widgets/Layout/SUniformGridPanel.h"
+#include "Widgets/SCompoundWidget.h"
+#include "Widgets/Text/SInlineEditableTextBlock.h"
+
+#define LOCTEXT_NAMESPACE "PhysicsControlAssetPreviewDetailsCustomization"
+
+//======================================================================================================================
+TSharedRef<IDetailCustomization> FPhysicsControlAssetPreviewDetailsCustomization::MakeInstance(
+	TWeakPtr<FPhysicsControlAssetEditor> InPhysicsControlAssetEditor)
+{
+	return MakeShared<FPhysicsControlAssetPreviewDetailsCustomization>(InPhysicsControlAssetEditor);
+}
+
+//======================================================================================================================
+void FPhysicsControlAssetPreviewDetailsCustomization::CustomizeDetails(IDetailLayoutBuilder& DetailLayoutBuilder)
+{
+	BindCommands();
+
+	DetailLayoutBuilder.HideCategory(TEXT("PreviewMesh"));
+	DetailLayoutBuilder.HideCategory(TEXT("Actions"));
+	DetailLayoutBuilder.HideCategory(TEXT("Inheritance"));
+	DetailLayoutBuilder.HideCategory(TEXT("Setup"));
+	DetailLayoutBuilder.HideCategory(TEXT("Profiles"));
+
+	TSharedPtr<FPhysicsControlAssetEditorData> EditorData = PhysicsControlAssetEditor.Pin()->GetEditorData();
+	UPhysicsControlAsset* PCA = EditorData->PhysicsControlAsset.Get();
+	if (PCA)
+	{
+		IDetailCategoryBuilder& DetailCategoryBuilder = DetailLayoutBuilder.EditCategory(TEXT("Preview Profiles"));
+
+		for (const TPair<FName, FPhysicsControlControlAndModifierUpdates>& ProfilePair : PCA->Profiles)
+		{
+			const FName ProfileName = ProfilePair.Key;
+			DetailCategoryBuilder.AddCustomRow(
+				FText::FromName(ProfileName))
+				.NameContent()
+				[
+					SNew(STextBlock)
+						.Font(DetailLayoutBuilder.GetDetailFont())
+						.Text(FText::FromName(ProfileName))
+				]
+				.ValueContent()
+				[
+					SNew(SButton)
+						.ButtonStyle(FAppStyle::Get(), "SimpleButton")
+						.ContentPadding(FMargin(6, 2))
+						.Text(LOCTEXT("Invoke", "Invoke"))
+						.OnClicked(this, &FPhysicsControlAssetPreviewDetailsCustomization::InvokeControlProfile, ProfileName)
+				];
+		}
+	}
+
+}
+
+//======================================================================================================================
+// TODO Is this needed?
+void FPhysicsControlAssetPreviewDetailsCustomization::BindCommands()
+{
+}
+
+//======================================================================================================================
+FReply FPhysicsControlAssetPreviewDetailsCustomization::InvokeControlProfile(FName ProfileName)
+{
+	TSharedPtr<FPhysicsControlAssetEditorData> EditorData = PhysicsControlAssetEditor.Pin()->GetEditorData();
+	UPhysicsControlComponent* PCC = EditorData->PhysicsControlComponent.Get();
+	if (PCC)
+	{
+		PCC->InvokeControlProfile(ProfileName);
+	}
+	return FReply::Handled();
+}
+
+
+#undef LOCTEXT_NAMESPACE
diff --git a/Engine/Plugins/Experimental/PhysicsControl/Source/PhysicsControlEditor/Private/PhysicsControlAssetPreviewDetailsCustomization.h b/Engine/Plugins/Experimental/PhysicsControl/Source/PhysicsControlEditor/Private/PhysicsControlAssetPreviewDetailsCustomization.h
new file mode 100644
--- /dev/null
+++ b/Engine/Plugins/Experimental/PhysicsControl/Source/PhysicsControlEditor/Private/PhysicsControlAssetPreviewDetailsCustomization.h
@@ -0,0 +1,38 @@
+// Copyright Epic Games, Inc. All Rights Reserved.
+
+#pragma once
+
+#include "CoreMinimal.h"
+#include "IDetailCustomization.h"
+#include "Types/SlateEnums.h"
+#include "Input/Reply.h"
+
+class FPhysicsControlAssetEditor;
+class IDetailLayoutBuilder;
+class FUICommandInfo;
+class SWidget;
+class IPropertyHandle;
+class FUICommandList;
+class SEditableTextBox;
+
+class FPhysicsControlAssetPreviewDetailsCustomization : public IDetailCustomization
+{
+public:
+	// Makes a new instance of this detail layout class for a specific detail view requesting it
+	static TSharedRef<IDetailCustomization> MakeInstance(TWeakPtr<FPhysicsControlAssetEditor> InPhysicsControlAssetEditor);
+
+	FPhysicsControlAssetPreviewDetailsCustomization(TWeakPtr<FPhysicsControlAssetEditor> InPhysicsControlAssetEditor)
+		: PhysicsControlAssetEditor(InPhysicsControlAssetEditor)
+	{}
+
+	// IDetailCustomization interface
+	virtual void CustomizeDetails(IDetailLayoutBuilder& DetailLayout) override;
+
+private:
+	void BindCommands();
+
+	FReply InvokeControlProfile(FName ProfileName);
+
+private:
+	TWeakPtr<FPhysicsControlAssetEditor> PhysicsControlAssetEditor;
+};
diff --git a/Engine/Plugins/Experimental/PhysicsControl/Source/PhysicsControlEditor/Private/PhysicsControlAssetProfileDetailsCustomization.cpp b/Engine/Plugins/Experimental/PhysicsControl/Source/PhysicsControlEditor/Private/PhysicsControlAssetProfileDetailsCustomization.cpp
new file mode 100644
--- /dev/null
+++ b/Engine/Plugins/Experimental/PhysicsControl/Source/PhysicsControlEditor/Private/PhysicsControlAssetProfileDetailsCustomization.cpp
@@ -0,0 +1,57 @@
+// Copyright Epic Games, Inc. All Rights Reserved.
+
+#include "PhysicsControlAssetProfileDetailsCustomization.h"
+
+#include "PhysicsControlAssetActions.h"
+#include "PhysicsControlAssetEditor.h"
+#include "PhysicsControlAssetEditorCommands.h"
+
+#include "DetailCategoryBuilder.h"
+#include "DetailLayoutBuilder.h"
+#include "DetailWidgetRow.h"
+#include "EditorFontGlyphs.h"
+#include "Framework/Application/SlateApplication.h"
+#include "Framework/MultiBox/MultiBoxBuilder.h"
+#include "IDetailPropertyRow.h"
+#include "PropertyHandle.h"
+#include "ScopedTransaction.h"
+#include "Styling/StyleColors.h"
+#include "Widgets/Images/SImage.h"
+#include "Widgets/Input/SButton.h"
+#include "Widgets/Input/SComboButton.h"
+#include "Widgets/Input/SEditableTextBox.h"
+#include "Widgets/Layout/SUniformGridPanel.h"
+#include "Widgets/SCompoundWidget.h"
+#include "Widgets/Text/SInlineEditableTextBlock.h"
+
+#define LOCTEXT_NAMESPACE "PhysicsControlAssetProfileDetailsCustomization"
+
+//======================================================================================================================
+TSharedRef<IDetailCustomization> FPhysicsControlAssetProfileDetailsCustomization::MakeInstance(
+	TWeakPtr<FPhysicsControlAssetEditor> InPhysicsControlAssetEditor)
+{
+	return MakeShared<FPhysicsControlAssetProfileDetailsCustomization>(InPhysicsControlAssetEditor);
+}
+
+//======================================================================================================================
+void FPhysicsControlAssetProfileDetailsCustomization::CustomizeDetails(IDetailLayoutBuilder& DetailLayout)
+{
+	BindCommands();
+
+	DetailLayout.HideCategory(TEXT("PreviewMesh"));
+	DetailLayout.HideCategory(TEXT("Actions"));
+	DetailLayout.HideCategory(TEXT("Inheritance"));
+	DetailLayout.HideCategory(TEXT("Setup"));
+
+}
+
+//======================================================================================================================
+void FPhysicsControlAssetProfileDetailsCustomization::BindCommands()
+{
+	const FPhysicsControlAssetEditorCommands& Commands = FPhysicsControlAssetEditorCommands::Get();
+
+	TSharedPtr<FUICommandList> CommandList = PhysicsControlAssetEditor.Pin()->GetToolkitCommands();
+
+}
+
+#undef LOCTEXT_NAMESPACE
diff --git a/Engine/Plugins/Experimental/PhysicsControl/Source/PhysicsControlEditor/Private/PhysicsControlAssetProfileDetailsCustomization.h b/Engine/Plugins/Experimental/PhysicsControl/Source/PhysicsControlEditor/Private/PhysicsControlAssetProfileDetailsCustomization.h
new file mode 100644
--- /dev/null
+++ b/Engine/Plugins/Experimental/PhysicsControl/Source/PhysicsControlEditor/Private/PhysicsControlAssetProfileDetailsCustomization.h
@@ -0,0 +1,35 @@
+// Copyright Epic Games, Inc. All Rights Reserved.
+
+#pragma once
+
+#include "CoreMinimal.h"
+#include "IDetailCustomization.h"
+#include "Types/SlateEnums.h"
+
+class FPhysicsControlAssetEditor;
+class IDetailLayoutBuilder;
+class FUICommandInfo;
+class SWidget;
+class IPropertyHandle;
+class FUICommandList;
+class SEditableTextBox;
+
+class FPhysicsControlAssetProfileDetailsCustomization : public IDetailCustomization
+{
+public:
+	// Makes a new instance of this detail layout class for a specific detail view requesting it
+	static TSharedRef<IDetailCustomization> MakeInstance(TWeakPtr<FPhysicsControlAssetEditor> InPhysicsControlAssetEditor);
+
+	FPhysicsControlAssetProfileDetailsCustomization(TWeakPtr<FPhysicsControlAssetEditor> InPhysicsControlAssetEditor)
+		: PhysicsControlAssetEditor(InPhysicsControlAssetEditor)
+	{}
+
+	// IDetailCustomization interface
+	virtual void CustomizeDetails(IDetailLayoutBuilder& DetailLayout) override;
+
+private:
+	void BindCommands();
+
+private:
+	TWeakPtr<FPhysicsControlAssetEditor> PhysicsControlAssetEditor;
+};
diff --git a/Engine/Plugins/Experimental/PhysicsControl/Source/PhysicsControlEditor/Private/PhysicsControlAssetSetupDetailsCustomization.cpp b/Engine/Plugins/Experimental/PhysicsControl/Source/PhysicsControlEditor/Private/PhysicsControlAssetSetupDetailsCustomization.cpp
new file mode 100644
--- /dev/null
+++ b/Engine/Plugins/Experimental/PhysicsControl/Source/PhysicsControlEditor/Private/PhysicsControlAssetSetupDetailsCustomization.cpp
@@ -0,0 +1,54 @@
+// Copyright Epic Games, Inc. All Rights Reserved.
+
+#include "PhysicsControlAssetSetupDetailsCustomization.h"
+
+#include "PhysicsControlAssetActions.h"
+#include "PhysicsControlAssetEditor.h"
+#include "PhysicsControlAssetEditorCommands.h"
+
+#include "DetailCategoryBuilder.h"
+#include "DetailLayoutBuilder.h"
+#include "DetailWidgetRow.h"
+#include "EditorFontGlyphs.h"
+#include "Framework/Application/SlateApplication.h"
+#include "Framework/MultiBox/MultiBoxBuilder.h"
+#include "IDetailPropertyRow.h"
+#include "PropertyHandle.h"
+#include "ScopedTransaction.h"
+#include "Styling/StyleColors.h"
+#include "Widgets/Images/SImage.h"
+#include "Widgets/Input/SButton.h"
+#include "Widgets/Input/SComboButton.h"
+#include "Widgets/Input/SEditableTextBox.h"
+#include "Widgets/Layout/SUniformGridPanel.h"
+#include "Widgets/SCompoundWidget.h"
+#include "Widgets/Text/SInlineEditableTextBlock.h"
+
+#define LOCTEXT_NAMESPACE "PhysicsControlAssetSetupDetailsCustomization"
+
+//======================================================================================================================
+TSharedRef<IDetailCustomization> FPhysicsControlAssetSetupDetailsCustomization::MakeInstance(
+	TWeakPtr<FPhysicsControlAssetEditor> InPhysicsControlAssetEditor)
+{
+	return MakeShared<FPhysicsControlAssetSetupDetailsCustomization>(InPhysicsControlAssetEditor);
+}
+
+//======================================================================================================================
+void FPhysicsControlAssetSetupDetailsCustomization::CustomizeDetails(IDetailLayoutBuilder& DetailLayout)
+{
+	BindCommands();
+
+	DetailLayout.HideCategory(TEXT("Profiles"));
+
+}
+
+//======================================================================================================================
+void FPhysicsControlAssetSetupDetailsCustomization::BindCommands()
+{
+	const FPhysicsControlAssetEditorCommands& Commands = FPhysicsControlAssetEditorCommands::Get();
+
+	TSharedPtr<FUICommandList> CommandList = PhysicsControlAssetEditor.Pin()->GetToolkitCommands();
+
+}
+
+#undef LOCTEXT_NAMESPACE
diff --git a/Engine/Plugins/Experimental/PhysicsControl/Source/PhysicsControlEditor/Private/PhysicsControlAssetSetupDetailsCustomization.h b/Engine/Plugins/Experimental/PhysicsControl/Source/PhysicsControlEditor/Private/PhysicsControlAssetSetupDetailsCustomization.h
new file mode 100644
--- /dev/null
+++ b/Engine/Plugins/Experimental/PhysicsControl/Source/PhysicsControlEditor/Private/PhysicsControlAssetSetupDetailsCustomization.h
@@ -0,0 +1,35 @@
+// Copyright Epic Games, Inc. All Rights Reserved.
+
+#pragma once
+
+#include "CoreMinimal.h"
+#include "IDetailCustomization.h"
+#include "Types/SlateEnums.h"
+
+class FPhysicsControlAssetEditor;
+class IDetailLayoutBuilder;
+class FUICommandInfo;
+class SWidget;
+class IPropertyHandle;
+class FUICommandList;
+class SEditableTextBox;
+
+class FPhysicsControlAssetSetupDetailsCustomization : public IDetailCustomization
+{
+public:
+	// Makes a new instance of this detail layout class for a specific detail view requesting it
+	static TSharedRef<IDetailCustomization> MakeInstance(TWeakPtr<FPhysicsControlAssetEditor> InPhysicsControlAssetEditor);
+
+	FPhysicsControlAssetSetupDetailsCustomization(TWeakPtr<FPhysicsControlAssetEditor> InPhysicsControlAssetEditor)
+		: PhysicsControlAssetEditor(InPhysicsControlAssetEditor)
+	{}
+
+	// IDetailCustomization interface
+	virtual void CustomizeDetails(IDetailLayoutBuilder& DetailLayout) override;
+
+private:
+	void BindCommands();
+
+private:
+	TWeakPtr<FPhysicsControlAssetEditor> PhysicsControlAssetEditor;
+};
diff --git a/Engine/Plugins/Experimental/PhysicsControl/Source/PhysicsControlEditor/Private/PhysicsControlEditorModule.cpp b/Engine/Plugins/Experimental/PhysicsControl/Source/PhysicsControlEditor/Private/PhysicsControlEditorModule.cpp
--- a/Engine/Plugins/Experimental/PhysicsControl/Source/PhysicsControlEditor/Private/PhysicsControlEditorModule.cpp
+++ b/Engine/Plugins/Experimental/PhysicsControl/Source/PhysicsControlEditor/Private/PhysicsControlEditorModule.cpp
@@ -5,7 +5,7 @@
 #include "PhysicsControlComponentVisualizer.h"
 #include "OperatorEditor/OperatorEditor.h"
 #include "PhysicsControlAssetActions.h"
-#include "PhysicsControlAssetEditorMode.h"
+#include "PhysicsControlAssetEditorEditMode.h"
 #include "UnrealEdGlobals.h"
 #include "Editor/UnrealEdEngine.h"
 #include "Features/IModularFeatures.h"
@@ -20,9 +20,9 @@ void FPhysicsControlEditorModule::StartupModule()
 	IAssetTools& AssetTools = FModuleManager::LoadModuleChecked<FAssetToolsModule>("AssetTools").Get();
 	PhysicsControlAssetActions = MakeShared<FPhysicsControlAssetActions>();
 	AssetTools.RegisterAssetTypeActions(PhysicsControlAssetActions.ToSharedRef());
-
-	FEditorModeRegistry::Get().RegisterMode<FPhysicsControlAssetEditorMode>(
-		FPhysicsControlAssetEditorMode::ModeName, 
+	
+	FEditorModeRegistry::Get().RegisterMode<FPhysicsControlAssetEditorEditMode>(
+		FPhysicsControlAssetEditorEditMode::ModeName, 
 		LOCTEXT("PhysicsControlAssetEditorMode", "PhysicsControlAsset"), 
 		FSlateIcon(), false);
 
@@ -59,8 +59,7 @@ void FPhysicsControlEditorModule::ShutdownModule()
 		EditorInterface = nullptr;
 	}
 	
-	// Physics Control Profile editor/asset is disabled for now
-	FEditorModeRegistry::Get().UnregisterMode(FPhysicsControlAssetEditorMode::ModeName);
+	FEditorModeRegistry::Get().UnregisterMode(FPhysicsControlAssetEditorEditMode::ModeName);
 
 	if (FModuleManager::Get().IsModuleLoaded("AssetTools"))
 	{
diff --git a/Engine/Plugins/Experimental/PhysicsControl/Source/PhysicsControlEditor/Private/SkeletonTreePhysicsControlBodyItem.cpp b/Engine/Plugins/Experimental/PhysicsControl/Source/PhysicsControlEditor/Private/SkeletonTreePhysicsControlBodyItem.cpp
new file mode 100644
--- /dev/null
+++ b/Engine/Plugins/Experimental/PhysicsControl/Source/PhysicsControlEditor/Private/SkeletonTreePhysicsControlBodyItem.cpp
@@ -0,0 +1,65 @@
+// Copyright Epic Games, Inc. All Rights Reserved.
+
+#include "SkeletonTreePhysicsControlBodyItem.h"
+#include "Styling/AppStyle.h"
+#include "PhysicsAssetRenderUtils.h"
+#include "PhysicsEngine/PhysicsAsset.h"
+#include "PhysicsEngine/SkeletalBodySetup.h"
+
+#define LOCTEXT_NAMESPACE "FSkeletonTreePhysicsControlBodyItem"
+
+FSkeletonTreePhysicsControlBodyItem::FSkeletonTreePhysicsControlBodyItem(USkeletalBodySetup* InBodySetup, int32 InBodySetupIndex, const FName& InBoneName, bool bInHasBodySetup, bool bInHasShapes, class UPhysicsAsset* const InPhysicsAsset, const TSharedRef<class ISkeletonTree>& InSkeletonTree)
+	: FSkeletonTreePhysicsControlItem(InPhysicsAsset, InSkeletonTree)
+	, BodySetup(InBodySetup)
+	, BodySetupIndex(InBodySetupIndex)
+	, bHasBodySetup(bInHasBodySetup)
+	, bHasShapes(bInHasShapes)
+{
+	DisplayName = InBoneName;
+}
+
+UObject* FSkeletonTreePhysicsControlBodyItem::GetObject() const
+{
+	return BodySetup;
+}
+
+void FSkeletonTreePhysicsControlBodyItem::OnToggleItemDisplayed(ECheckBoxState InCheckboxState)
+{
+	if (FPhysicsAssetRenderSettings* RenderSettings = GetRenderSettings())
+	{
+		RenderSettings->ToggleShowBody(BodySetupIndex);
+	}
+}
+
+ECheckBoxState FSkeletonTreePhysicsControlBodyItem::IsItemDisplayed() const
+{
+	if (FPhysicsAssetRenderSettings* RenderSettings = GetRenderSettings())
+	{
+		return RenderSettings->IsBodyHidden(BodySetupIndex) ? ECheckBoxState::Unchecked : ECheckBoxState::Checked;
+	}
+
+	return ECheckBoxState::Undetermined;
+}
+
+const FSlateBrush* FSkeletonTreePhysicsControlBodyItem::GetBrush() const
+{
+	return BodySetup->PhysicsType == EPhysicsType::PhysType_Kinematic ? FAppStyle::GetBrush("PhysicsAssetEditor.Tree.KinematicBody") : FAppStyle::GetBrush("PhysicsAssetEditor.Tree.Body");
+}
+
+FSlateColor FSkeletonTreePhysicsControlBodyItem::GetTextColor() const
+{
+	FLinearColor Color(1.0f, 1.0f, 1.0f);
+
+	if (FilterResult == ESkeletonTreeFilterResult::ShownDescendant)
+	{
+		Color = FLinearColor::Gray * 0.5f;
+	}
+	return FSlateColor(Color);
+}
+
+FText FSkeletonTreePhysicsControlBodyItem::GetNameColumnToolTip() const
+{
+	return FText::Format(LOCTEXT("BodyTooltip", "Aggregate physics body for bone '{0}'. Bodies can consist of multiple shapes."), FText::FromName(GetRowItemName()));
+}
+
+#undef LOCTEXT_NAMESPACE
diff --git a/Engine/Plugins/Experimental/PhysicsControl/Source/PhysicsControlEditor/Private/SkeletonTreePhysicsControlBodyItem.h b/Engine/Plugins/Experimental/PhysicsControl/Source/PhysicsControlEditor/Private/SkeletonTreePhysicsControlBodyItem.h
new file mode 100644
--- /dev/null
+++ b/Engine/Plugins/Experimental/PhysicsControl/Source/PhysicsControlEditor/Private/SkeletonTreePhysicsControlBodyItem.h
@@ -0,0 +1,46 @@
+// Copyright Epic Games, Inc. All Rights Reserved.
+
+#pragma once
+
+#include "CoreMinimal.h"
+#include "Styling/SlateColor.h"
+#include "SkeletonTreePhysicsControlItem.h"
+
+class FSkeletonTreePhysicsControlBodyItem : public FSkeletonTreePhysicsControlItem
+{
+public:
+	SKELETON_TREE_ITEM_TYPE(FSkeletonTreePhysicsControlBodyItem, FSkeletonTreePhysicsControlItem);
+
+	FSkeletonTreePhysicsControlBodyItem(class USkeletalBodySetup* InBodySetup, int32 InBodySetupIndex, const FName& InBoneName, bool bInHasBodySetup, bool bInHasShapes, class UPhysicsAsset* const InPhysicsAsset, const TSharedRef<class ISkeletonTree>& InSkeletonTree);
+
+	virtual UObject* GetObject() const override;
+
+	/** Get the index of the body setup in the physics asset */
+	int32 GetBodySetupIndex() const { return BodySetupIndex; }
+
+	/** UI Callbacks */
+	virtual void OnToggleItemDisplayed(ECheckBoxState InCheckboxState) override;
+	virtual ECheckBoxState IsItemDisplayed() const override;
+
+private:
+	/** Gets the icon to display for this body */
+	virtual const FSlateBrush* GetBrush() const override;
+
+	/** Gets the color to display the item's text */
+	virtual FSlateColor GetTextColor() const override;
+
+	/** Gets the tool tip to display on hovering over the item's name */
+	virtual FText GetNameColumnToolTip() const override;
+
+	/** The body setup we are representing */
+	USkeletalBodySetup* BodySetup;
+
+	/** The index of the body setup in the physics asset */
+	int32 BodySetupIndex;
+
+	/** Whether there is a body set up for this bone */
+	bool bHasBodySetup;
+
+	/** Whether this body has shapes */
+	bool bHasShapes;
+};
diff --git a/Engine/Plugins/Experimental/PhysicsControl/Source/PhysicsControlEditor/Private/SkeletonTreePhysicsControlItem.cpp b/Engine/Plugins/Experimental/PhysicsControl/Source/PhysicsControlEditor/Private/SkeletonTreePhysicsControlItem.cpp
new file mode 100644
--- /dev/null
+++ b/Engine/Plugins/Experimental/PhysicsControl/Source/PhysicsControlEditor/Private/SkeletonTreePhysicsControlItem.cpp
@@ -0,0 +1,74 @@
+// Copyright Epic Games, Inc. All Rights Reserved.
+
+#include "SkeletonTreePhysicsControlItem.h"
+#include "Widgets/Text/STextBlock.h"
+#include "Widgets/Images/SImage.h"
+#include "Widgets/Input/SCheckBox.h"
+#include "Styling/AppStyle.h"
+#include "PhysicsAssetRenderUtils.h"
+
+#define LOCTEXT_NAMESPACE "FSkeletonTreePhysicsControlItem"
+
+FSkeletonTreePhysicsControlItem::FSkeletonTreePhysicsControlItem(class UPhysicsAsset* const InPhysicsAsset, const TSharedRef<class ISkeletonTree>& InSkeletonTree)
+	: FSkeletonTreeItem(InSkeletonTree)
+	, PhysicsAssetPathNameHash(UPhysicsAssetRenderUtilities::GetPathNameHash(InPhysicsAsset))
+{}
+
+void FSkeletonTreePhysicsControlItem::GenerateWidgetForNameColumn(TSharedPtr< SHorizontalBox > Box, const TAttribute<FText>& FilterText, FIsSelected InIsSelected)
+{
+	Box->AddSlot()
+		.AutoWidth()
+		.Padding(FMargin(0.0f, 1.0f))
+		[
+			SNew(SImage)
+				.ColorAndOpacity(FSlateColor::UseForeground())
+				.Image(GetBrush())
+		];
+
+	Box->AddSlot()
+		.AutoWidth()
+		.Padding(2, 0, 0, 0)
+		[
+			SNew(STextBlock)
+				.ColorAndOpacity(this, &FSkeletonTreePhysicsControlItem::GetTextColor)
+				.Text(FText::FromName(GetRowItemName()))
+				.HighlightText(FilterText)
+				.Font(FAppStyle::GetFontStyle("PhysicsAssetEditor.Tree.Font"))
+				.ToolTipText(GetNameColumnToolTip())
+		];
+}
+
+
+TSharedRef< SWidget > FSkeletonTreePhysicsControlItem::GenerateWidgetForDataColumn(const FName& DataColumnName, FIsSelected InIsSelected)
+{
+	if (DataColumnName == ISkeletonTree::DebugVisualizationOptionsName())
+	{
+		// Create a widget that displays an open/closed eye icon when the body is shown/hidden in the viewport debug draw.
+		return SNew(SHorizontalBox)
+			+ SHorizontalBox::Slot()
+			.AutoWidth()
+			.HAlign(HAlign_Left)
+			[
+				SNew(SCheckBox)
+					.ToolTipText(LOCTEXT("SkeletonTreePhysicsItemVisibilityCheckBoxToolTip", "Click to toggle visibility of debug draw for this item"))
+					.OnCheckStateChanged(this, &FSkeletonTreePhysicsControlItem::OnToggleItemDisplayed)
+					.IsChecked(this, &FSkeletonTreePhysicsControlItem::IsItemDisplayed)
+					.CheckedHoveredImage(FAppStyle::GetBrush("Level.VisibleIcon16x"))
+					.CheckedImage(FAppStyle::GetBrush("Level.VisibleIcon16x"))
+					.CheckedPressedImage(FAppStyle::GetBrush("Level.VisibleIcon16x"))
+					.UncheckedHoveredImage(FAppStyle::GetBrush("Level.NotVisibleIcon16x"))
+					.UncheckedImage(FAppStyle::GetBrush("Level.NotVisibleIcon16x"))
+					.UncheckedPressedImage(FAppStyle::GetBrush("Level.NotVisibleIcon16x"))
+			];
+	}
+
+	return SNullWidget::NullWidget;
+}
+
+
+FPhysicsAssetRenderSettings* FSkeletonTreePhysicsControlItem::GetRenderSettings() const
+{
+	return UPhysicsAssetRenderUtilities::GetSettings(PhysicsAssetPathNameHash);
+}
+
+#undef LOCTEXT_NAMESPACE
diff --git a/Engine/Plugins/Experimental/PhysicsControl/Source/PhysicsControlEditor/Private/SkeletonTreePhysicsControlItem.h b/Engine/Plugins/Experimental/PhysicsControl/Source/PhysicsControlEditor/Private/SkeletonTreePhysicsControlItem.h
new file mode 100644
--- /dev/null
+++ b/Engine/Plugins/Experimental/PhysicsControl/Source/PhysicsControlEditor/Private/SkeletonTreePhysicsControlItem.h
@@ -0,0 +1,48 @@
+// Copyright Epic Games, Inc. All Rights Reserved.
+
+#pragma once
+
+#include "CoreMinimal.h"
+#include "Styling/SlateColor.h"
+#include "Input/Reply.h"
+#include "Widgets/SWidget.h"
+#include "Widgets/Views/STableViewBase.h"
+#include "Widgets/Views/STableRow.h"
+#include "SkeletonTreeItem.h"
+
+class FSkeletonTreePhysicsControlItem : public FSkeletonTreeItem
+{
+public:
+	SKELETON_TREE_ITEM_TYPE(FSkeletonTreePhysicsControlItem, FSkeletonTreeItem);
+
+	FSkeletonTreePhysicsControlItem(class UPhysicsAsset* const InPhysicsAsset, const TSharedRef<class ISkeletonTree>& InSkeletonTree);
+
+	/** ISkeletonTreeItem interface */
+	virtual void GenerateWidgetForNameColumn(TSharedPtr< SHorizontalBox > Box, const TAttribute<FText>& FilterText, FIsSelected InIsSelected) override;
+	virtual TSharedRef< SWidget > GenerateWidgetForDataColumn(const FName& DataColumnName, FIsSelected InIsSelected) override;
+	virtual FName GetRowItemName() const override { return DisplayName; }
+
+	/** UI Callbacks */
+	virtual void OnToggleItemDisplayed(ECheckBoxState InCheckboxState) = 0;
+	virtual ECheckBoxState IsItemDisplayed() const = 0;
+
+private:
+	/** Gets the icon to display for this body */
+	virtual const FSlateBrush* GetBrush() const = 0;
+
+	/** Gets the color to display the item's text */
+	virtual FSlateColor GetTextColor() const = 0;
+
+	/** Gets the tool tip to display on hovering over the item's name */
+	virtual FText GetNameColumnToolTip() const = 0;
+
+protected:
+	/** Gets the physics asset render settings for this physics item */
+	struct FPhysicsAssetRenderSettings* GetRenderSettings() const;
+
+	/** Unique ID of the physics asset for this physics item */
+	uint32 PhysicsAssetPathNameHash;
+
+	/** The name of the item in the tree */
+	FName DisplayName;
+};
diff --git a/Engine/Plugins/Experimental/PhysicsControl/Source/PhysicsControlEditor/Private/SkeletonTreePhysicsControlShapeItem.cpp b/Engine/Plugins/Experimental/PhysicsControl/Source/PhysicsControlEditor/Private/SkeletonTreePhysicsControlShapeItem.cpp
new file mode 100644
--- /dev/null
+++ b/Engine/Plugins/Experimental/PhysicsControl/Source/PhysicsControlEditor/Private/SkeletonTreePhysicsControlShapeItem.cpp
@@ -0,0 +1,207 @@
+// Copyright Epic Games, Inc. All Rights Reserved.
+
+#include "SkeletonTreePhysicsControlShapeItem.h"
+#include "Widgets/Text/STextBlock.h"
+#include "Widgets/Images/SImage.h"
+#include "Styling/AppStyle.h"
+#include "Widgets/Text/SInlineEditableTextBlock.h"
+#include "ScopedTransaction.h"
+#include "PhysicsEngine/SkeletalBodySetup.h"
+
+#define LOCTEXT_NAMESPACE "FSkeletonTreePhysicsControlShapeItem"
+
+FSkeletonTreePhysicsControlShapeItem::FSkeletonTreePhysicsControlShapeItem(USkeletalBodySetup* InBodySetup, const FName& InBoneName, int32 InBodySetupIndex, EAggCollisionShape::Type InShapeType, int32 InShapeIndex, const TSharedRef<class ISkeletonTree>& InSkeletonTree)
+	: FSkeletonTreeItem(InSkeletonTree)
+	, BodySetup(InBodySetup)
+	, BodySetupIndex(InBodySetupIndex)
+	, ShapeType(InShapeType)
+	, ShapeIndex(InShapeIndex)
+{
+	switch (ShapeType)
+	{
+	case EAggCollisionShape::Sphere:
+		ShapeBrush = FAppStyle::GetBrush("PhysicsAssetEditor.Tree.Sphere");
+		DefaultLabel = *FText::Format(LOCTEXT("SphereLabel", "{0} Sphere {1}"), FText::FromName(InBoneName), FText::AsNumber(ShapeIndex)).ToString();
+		break;
+	case EAggCollisionShape::Box:
+		ShapeBrush = FAppStyle::GetBrush("PhysicsAssetEditor.Tree.Box");
+		DefaultLabel = *FText::Format(LOCTEXT("BoxLabel", "{0} Box {1}"), FText::FromName(InBoneName), FText::AsNumber(ShapeIndex)).ToString();
+		break;
+	case EAggCollisionShape::Sphyl:
+		ShapeBrush = FAppStyle::GetBrush("PhysicsAssetEditor.Tree.Sphyl");
+		DefaultLabel = *FText::Format(LOCTEXT("CapsuleLabel", "{0} Capsule {1}"), FText::FromName(InBoneName), FText::AsNumber(ShapeIndex)).ToString();
+		break;
+	case EAggCollisionShape::Convex:
+		ShapeBrush = FAppStyle::GetBrush("PhysicsAssetEditor.Tree.Convex");
+		DefaultLabel = *FText::Format(LOCTEXT("ConvexLabel", "{0} Convex {1}"), FText::FromName(InBoneName), FText::AsNumber(ShapeIndex)).ToString();
+		break;
+	case EAggCollisionShape::TaperedCapsule:
+		ShapeBrush = FAppStyle::GetBrush("PhysicsAssetEditor.Tree.TaperedCapsule");
+		DefaultLabel = *FText::Format(LOCTEXT("TaperedCapsuleLabel", "{0} Tapered Capsule {1}"), FText::FromName(InBoneName), FText::AsNumber(ShapeIndex)).ToString();
+		break;
+	default:
+		check(false);
+		break;
+	}
+}
+
+void FSkeletonTreePhysicsControlShapeItem::GenerateWidgetForNameColumn(TSharedPtr< SHorizontalBox > Box, const TAttribute<FText>& FilterText, FIsSelected InIsSelected)
+{
+	Box->AddSlot()
+		.AutoWidth()
+		.Padding(FMargin(0.0f, 1.0f))
+		[
+			SNew(SImage)
+				.ColorAndOpacity(FSlateColor::UseForeground())
+				.Image(ShapeBrush)
+		];
+
+	TSharedRef<SInlineEditableTextBlock> InlineWidget = SNew(SInlineEditableTextBlock)
+		.ColorAndOpacity(FSlateColor::UseForeground())
+		.Text(this, &FSkeletonTreePhysicsControlShapeItem::GetNameAsText)
+		.ToolTipText(this, &FSkeletonTreePhysicsControlShapeItem::GetNameAsText)
+		.HighlightText(FilterText)
+		.Font(FAppStyle::GetFontStyle("PhysicsAssetEditor.Tree.Font"))
+		.OnTextCommitted(this, &FSkeletonTreePhysicsControlShapeItem::HandleTextCommitted)
+		.IsSelected(InIsSelected);
+
+	OnRenameRequested.BindSP(&InlineWidget.Get(), &SInlineEditableTextBlock::EnterEditingMode);
+
+	Box->AddSlot()
+		.AutoWidth()
+		.Padding(2, 0, 0, 0)
+		[
+			InlineWidget
+		];
+}
+
+TSharedRef< SWidget > FSkeletonTreePhysicsControlShapeItem::GenerateWidgetForDataColumn(const FName& DataColumnName, FIsSelected InIsSelected)
+{
+	return SNullWidget::NullWidget;
+}
+
+FName FSkeletonTreePhysicsControlShapeItem::GetRowItemName() const
+{
+	FString NameAsString = GetNameAsString();
+	return *NameAsString;
+}
+
+UObject* FSkeletonTreePhysicsControlShapeItem::GetObject() const
+{
+	return BodySetup;
+}
+
+bool FSkeletonTreePhysicsControlShapeItem::CanRenameItem() const
+{
+	return true;
+}
+
+void FSkeletonTreePhysicsControlShapeItem::OnItemDoubleClicked()
+{
+	OnRenameRequested.ExecuteIfBound();
+}
+
+void FSkeletonTreePhysicsControlShapeItem::RequestRename()
+{
+	OnRenameRequested.ExecuteIfBound();
+}
+
+FString FSkeletonTreePhysicsControlShapeItem::GetNameAsString() const
+{
+	FString StringName;
+
+	switch (ShapeType)
+	{
+	case EAggCollisionShape::Sphere:
+		if (BodySetup->AggGeom.SphereElems.IsValidIndex(ShapeIndex))
+		{
+			StringName = BodySetup->AggGeom.SphereElems[ShapeIndex].GetName().GetPlainNameString();
+		}
+		break;
+	case EAggCollisionShape::Box:
+		if (BodySetup->AggGeom.BoxElems.IsValidIndex(ShapeIndex))
+		{
+			StringName = BodySetup->AggGeom.BoxElems[ShapeIndex].GetName().GetPlainNameString();
+		}
+		break;
+	case EAggCollisionShape::Sphyl:
+		if (BodySetup->AggGeom.SphylElems.IsValidIndex(ShapeIndex))
+		{
+			StringName = BodySetup->AggGeom.SphylElems[ShapeIndex].GetName().GetPlainNameString();
+		}
+		break;
+	case EAggCollisionShape::Convex:
+		if (BodySetup->AggGeom.ConvexElems.IsValidIndex(ShapeIndex))
+		{
+			StringName = BodySetup->AggGeom.ConvexElems[ShapeIndex].GetName().GetPlainNameString();
+		}
+		break;
+	case EAggCollisionShape::TaperedCapsule:
+		if (BodySetup->AggGeom.TaperedCapsuleElems.IsValidIndex(ShapeIndex))
+		{
+			StringName = BodySetup->AggGeom.TaperedCapsuleElems[ShapeIndex].GetName().GetPlainNameString();
+		}
+		break;
+	}
+
+	if (StringName.IsEmpty())
+	{
+		StringName = DefaultLabel.ToString();
+	}
+
+	return StringName;
+}
+
+FText FSkeletonTreePhysicsControlShapeItem::GetNameAsText() const
+{
+	return FText::FromString(GetNameAsString());
+}
+
+void FSkeletonTreePhysicsControlShapeItem::HandleTextCommitted(const FText& InText, ETextCommit::Type InCommitType)
+{
+	if (!InText.IsEmpty())
+	{
+		FScopedTransaction Transaction(LOCTEXT("RenameShapeTransaction", "Rename Shape"));
+
+		BodySetup->Modify();
+
+		switch (ShapeType)
+		{
+		case EAggCollisionShape::Sphere:
+			if (BodySetup->AggGeom.SphereElems.IsValidIndex(ShapeIndex))
+			{
+				BodySetup->AggGeom.SphereElems[ShapeIndex].SetName(*InText.ToString());
+			}
+			break;
+		case EAggCollisionShape::Box:
+			if (BodySetup->AggGeom.BoxElems.IsValidIndex(ShapeIndex))
+			{
+				BodySetup->AggGeom.BoxElems[ShapeIndex].SetName(*InText.ToString());
+			}
+			break;
+		case EAggCollisionShape::Sphyl:
+			if (BodySetup->AggGeom.SphylElems.IsValidIndex(ShapeIndex))
+			{
+				BodySetup->AggGeom.SphylElems[ShapeIndex].SetName(*InText.ToString());
+			}
+			break;
+		case EAggCollisionShape::Convex:
+			if (BodySetup->AggGeom.ConvexElems.IsValidIndex(ShapeIndex))
+			{
+				BodySetup->AggGeom.ConvexElems[ShapeIndex].SetName(*InText.ToString());
+			}
+			break;
+		case EAggCollisionShape::TaperedCapsule:
+			if (BodySetup->AggGeom.TaperedCapsuleElems.IsValidIndex(ShapeIndex))
+			{
+				BodySetup->AggGeom.TaperedCapsuleElems[ShapeIndex].SetName(*InText.ToString());
+			}
+			break;
+		default:
+			check(false);
+			break;
+		}
+	}
+}
+
+#undef LOCTEXT_NAMESPACE
diff --git a/Engine/Plugins/Experimental/PhysicsControl/Source/PhysicsControlEditor/Private/SkeletonTreePhysicsControlShapeItem.h b/Engine/Plugins/Experimental/PhysicsControl/Source/PhysicsControlEditor/Private/SkeletonTreePhysicsControlShapeItem.h
new file mode 100644
--- /dev/null
+++ b/Engine/Plugins/Experimental/PhysicsControl/Source/PhysicsControlEditor/Private/SkeletonTreePhysicsControlShapeItem.h
@@ -0,0 +1,73 @@
+// Copyright Epic Games, Inc. All Rights Reserved.
+
+#pragma once
+
+#include "CoreMinimal.h"
+#include "Styling/SlateColor.h"
+#include "Input/Reply.h"
+#include "Widgets/SWidget.h"
+#include "Widgets/Views/STableViewBase.h"
+#include "Widgets/Views/STableRow.h"
+#include "SkeletonTreeItem.h"
+#include "PhysicsEngine/BodySetup.h"
+#include "PhysicsEngine/ShapeElem.h"
+#include "PhysicsEngine/PhysicsAsset.h"
+
+class FSkeletonTreePhysicsControlShapeItem : public FSkeletonTreeItem
+{
+public:
+	SKELETON_TREE_ITEM_TYPE(FSkeletonTreePhysicsControlShapeItem, FSkeletonTreeItem);
+
+	FSkeletonTreePhysicsControlShapeItem(USkeletalBodySetup* InBodySetup, const FName& InBoneName, int32 InBodySetupIndex, EAggCollisionShape::Type InShapeType, int32 InShapeIndex, const TSharedRef<class ISkeletonTree>& InSkeletonTree);
+
+	/** ISkeletonTreeItem interface */
+	virtual void GenerateWidgetForNameColumn(TSharedPtr< SHorizontalBox > Box, const TAttribute<FText>& FilterText, FIsSelected InIsSelected) override;
+	virtual TSharedRef< SWidget > GenerateWidgetForDataColumn(const FName& DataColumnName, FIsSelected InIsSelected) override;
+	virtual FName GetRowItemName() const override;
+	virtual UObject* GetObject() const override;
+	virtual bool CanRenameItem() const override;
+	virtual void RequestRename() override;
+	virtual void OnItemDoubleClicked() override;
+
+	/** Get the index of the body setup in the physics asset */
+	int32 GetBodySetupIndex() const { return BodySetupIndex; }
+
+	/** Get the shape type of this item */
+	EAggCollisionShape::Type GetShapeType() const { return ShapeType; }
+
+	/** Get the index of the shape in the physics assets aggregate geom */
+	int32 GetShapeIndex() const { return ShapeIndex; }
+
+private:
+	/** Get the text to display for this item */
+	FText GetNameAsText() const;
+
+	/** Get the text to display for this item */
+	FString GetNameAsString() const;
+
+	/** Handle the shape being renamed */
+	void HandleTextCommitted(const FText& InText, ETextCommit::Type InCommitType);
+
+private:
+	/** Delegate for when the context menu requests a rename */
+	DECLARE_DELEGATE(FOnRenameRequested);
+	FOnRenameRequested OnRenameRequested;
+
+	/** The body setup we are representing part of */
+	USkeletalBodySetup* BodySetup;
+
+	/** The label we display in the tree */
+	FName DefaultLabel;
+
+	/** The index of the body setup in the physics asset */
+	int32 BodySetupIndex;
+
+	/** The kind of shape we represent */
+	EAggCollisionShape::Type ShapeType;
+
+	/** The index into the relevant body setup array for this shape */
+	int32 ShapeIndex;
+
+	/** The brush to use for this shape */
+	const FSlateBrush* ShapeBrush;
+};
diff --git a/Engine/Plugins/Experimental/PhysicsControl/Source/PhysicsControlEditor/Public/PhysicsControlAssetApplicationMode.h b/Engine/Plugins/Experimental/PhysicsControl/Source/PhysicsControlEditor/Public/PhysicsControlAssetApplicationMode.h
--- a/Engine/Plugins/Experimental/PhysicsControl/Source/PhysicsControlEditor/Public/PhysicsControlAssetApplicationMode.h
+++ b/Engine/Plugins/Experimental/PhysicsControl/Source/PhysicsControlEditor/Public/PhysicsControlAssetApplicationMode.h
@@ -7,11 +7,12 @@
 
 class FWorkflowCentricApplication;
 class IPersonaPreviewScene;
+class ISkeletonTree;
 
-class FPhysicsControlAssetEditorToolkit;
+class FPhysicsControlAssetEditor;
 
 /**
- * The application mode for the Physics Control Profile Editor
+ * The application mode for the Physics Control Asset Editor
  * This defines the layout of the UI. It basically spawns all the tabs, such as the viewport, details panels, etc.
  */
 class PHYSICSCONTROLEDITOR_API FPhysicsControlAssetApplicationMode : public FApplicationMode
@@ -22,6 +23,7 @@ public:
 
 	FPhysicsControlAssetApplicationMode(
 		TSharedRef<FWorkflowCentricApplication> InHostingApp, 
+		TSharedPtr<ISkeletonTree>               SkeletonTree,
 		TSharedRef<IPersonaPreviewScene>        InPreviewScene);
 
 	// FApplicationMode overrides.
@@ -30,7 +32,7 @@ public:
 
 protected:
 	/** The hosting app. */
-	TWeakPtr<FPhysicsControlAssetEditorToolkit> EditorToolkit = nullptr;
+	TWeakPtr<FPhysicsControlAssetEditor> PhysicsControlAssetEditor = nullptr;
 
 	/** The tab factories we support. */
 	FWorkflowAllowedTabSet TabFactories;
diff --git a/Engine/Plugins/Experimental/PhysicsControl/Source/PhysicsControlEditor/Public/PhysicsControlAssetEditor.h b/Engine/Plugins/Experimental/PhysicsControl/Source/PhysicsControlEditor/Public/PhysicsControlAssetEditor.h
new file mode 100644
--- /dev/null
+++ b/Engine/Plugins/Experimental/PhysicsControl/Source/PhysicsControlEditor/Public/PhysicsControlAssetEditor.h
@@ -0,0 +1,166 @@
+// Copyright Epic Games, Inc. All Rights Reserved.
+
+#pragma once
+
+#include "CoreMinimal.h"
+#include "EditorUndoClient.h"
+#include "IHasPersonaToolkit.h"
+#include "IPersonaPreviewScene.h"
+#include "IPersonaViewport.h"
+#include "PersonaAssetEditorToolkit.h"
+#include "BodySetupEnums.h"
+#include "Preferences/PhysicsAssetEditorOptions.h"
+
+class FPhysicsControlAssetEditorEditMode;
+class FPhysicsControlAssetEditorData;
+class UAnimPreviewInstance;
+class UPhysicsControlAsset;
+class ISkeletonTree;
+class FPhysicsControlAssetEditorSkeletonTreeBuilder;
+class FUICommandList_Pinnable;
+
+namespace PhysicsControlAssetEditorModes
+{
+	extern const FName PhysicsControlAssetEditorMode;
+}
+
+/**
+ * The main toolkit/editor for working with Physics Control Assets
+ */
+class PHYSICSCONTROLEDITOR_API FPhysicsControlAssetEditor :
+	public FPersonaAssetEditorToolkit,
+	public IHasPersonaToolkit,
+	public FGCObject,
+	public FEditorUndoClient,
+	public FTickableEditorObject
+{
+public:
+	friend class FPhysicsControlAssetApplicationMode;
+	friend class FPhysicsControlAssetEditorEditMode;
+
+public:
+
+	/** Initialize the asset editor. This will register the application mode, init the preview scene, etc. */
+	void InitAssetEditor(
+		const EToolkitMode::Type        Mode,
+		const TSharedPtr<IToolkitHost>& InitToolkitHost,
+		UPhysicsControlAsset*    InPhysicsControlAsset);
+
+	/** Shared data accessor */
+	TSharedPtr<FPhysicsControlAssetEditorData> GetEditorData() const;
+
+	// FAssetEditorToolkit overrides.
+	virtual void RegisterTabSpawners(const TSharedRef<FTabManager>& InTabManager) override;
+	virtual void UnregisterTabSpawners(const TSharedRef<FTabManager>& InTabManager) override;
+	virtual FName GetToolkitFName() const override;
+	virtual FText GetBaseToolkitName() const override;
+	virtual FText GetToolkitName() const override;
+	virtual FLinearColor GetWorldCentricTabColorScale() const override;
+	virtual FString GetWorldCentricTabPrefix() const override;
+	virtual void InitToolMenuContext(FToolMenuContext& MenuContext) override;
+	// ~END FAssetEditorToolkit overrides.
+
+	// FGCObject overrides.
+	virtual FString GetReferencerName() const override { return TEXT("FPhysicsControlAssetEditor"); }
+	virtual void AddReferencedObjects(FReferenceCollector& Collector) override;
+	// ~END FGCObject overrides.
+
+	// FTickableEditorObject overrides.
+	virtual void Tick(float DeltaTime) override {};
+	virtual ETickableTickType GetTickableTickType() const override { return ETickableTickType::Always; }
+	virtual TStatId GetStatId() const override;
+	// ~END FTickableEditorObject overrides.
+
+	// IHasPersonaToolkit overrides.
+	virtual TSharedRef<IPersonaToolkit> GetPersonaToolkit() const override;
+	// ~END IHasPersonaToolkit overrides.
+
+	IPersonaToolkit* GetPersonaToolkitPointer() const { return PersonaToolkit.Get(); }
+
+	/** Repopulates the hierarchy tree view */
+	void RefreshHierachyTree();
+
+	/** Refreshes the preview viewport */
+	void RefreshPreviewViewport();
+
+private:
+	FText GetSimulationToolTip() const;
+	FSlateIcon GetSimulationIcon() const;
+
+	/** Preview scene setup. */
+	void HandlePreviewSceneCreated(const TSharedRef<IPersonaPreviewScene>& InPersonaPreviewScene);
+	void HandleViewportCreated(const TSharedRef<IPersonaViewport>& InPersonaViewport);
+	void HandleDetailsCreated(const TSharedRef<class IDetailsView>& InDetailsView);
+	void HandleSelectionChanged(const TArrayView<TSharedPtr<ISkeletonTreeItem>>& InSelectedItems, ESelectInfo::Type InSelectInfo);
+	void HandleGetFilterLabel(TArray<FText>& InOutItems) const;
+	void HandleExtendFilterMenu(FMenuBuilder& InMenuBuilder);
+	void HandleExtendContextMenu(FMenuBuilder& InMenuBuilder);
+	void OnFinishedChangingDetails(const FPropertyChangedEvent& PropertyChangedEvent);
+	void ShowEmptyDetails() const;
+
+	void ExtendMenu();
+	void ExtendToolbar();
+	void ExtendViewportMenus();
+	void BindCommands();
+
+	/** Building the context menus */
+	void BuildMenuWidgetSelection(FMenuBuilder& InMenuBuilder);
+	void BuildMenuWidgetBody(FMenuBuilder& InMenuBuilder);
+
+	// Toolbar/menu commands
+	void OnCompile();
+	bool IsCompilationNeeded();
+	void OnToggleSimulation();
+	void OnToggleSimulationNoGravity();
+	bool IsNoGravitySimulationEnabled() const;
+	void OnToggleSimulationFloorCollision();
+	bool IsSimulationFloorCollisionEnabled() const;
+	void OnMeshRenderingMode(EPhysicsAssetEditorMeshViewMode Mode, bool bSimulation);
+	bool IsMeshRenderingMode(EPhysicsAssetEditorMeshViewMode Mode, bool bSimulation) const;
+	void OnCollisionRenderingMode(EPhysicsAssetEditorCollisionViewMode Mode, bool bSimulation);
+	bool IsCollisionRenderingMode(EPhysicsAssetEditorCollisionViewMode Mode, bool bSimulation) const;
+	void OnConstraintRenderingMode(EPhysicsAssetEditorConstraintViewMode Mode, bool bSimulation);
+	bool IsConstraintRenderingMode(EPhysicsAssetEditorConstraintViewMode Mode, bool bSimulation) const;
+	void ToggleDrawViolatedLimits();
+	bool IsDrawingViolatedLimits() const;
+	bool IsRunningSimulation() const;
+	bool IsNotRunningSimulation() const;
+
+private:
+	/** Make the constraint scale widget */
+	TSharedRef<SWidget> MakeConstraintScaleWidget();
+
+	/** Make the collision opacity widget */
+	TSharedRef<SWidget> MakeCollisionOpacityWidget();
+
+private:
+	/** The persona toolkit. */
+	TSharedPtr<IPersonaToolkit> PersonaToolkit = nullptr;
+
+	// Persona viewport.
+	TSharedPtr<IPersonaViewport> PersonaViewport = nullptr;
+
+	/** Data and methods shared across multiple classes */
+	TSharedPtr<FPhysicsControlAssetEditorData> EditorData;
+
+	// Asset properties tab 
+	TSharedPtr<IDetailsView> DetailsView;
+
+	/** The skeleton tree widget */
+	TSharedPtr<ISkeletonTree> SkeletonTree;
+
+	/** The skeleton tree builder */
+	TSharedPtr<FPhysicsControlAssetEditorSkeletonTreeBuilder> SkeletonTreeBuilder;
+
+	/** Command list for skeleton tree operations */
+	TSharedPtr<FUICommandList_Pinnable> SkeletonTreeCommandList;
+
+	/** Command list for viewport operations */
+	TSharedPtr<FUICommandList_Pinnable> ViewportCommandList;
+
+	/** Has the asset editor been initialized? */
+	bool bIsInitialized = false;
+
+	/** True if in OnTreeSelectionChanged()... protects against infinite recursion */
+	bool bSelecting;
+};
diff --git a/Engine/Plugins/Experimental/PhysicsControl/Source/PhysicsControlEditor/Public/PhysicsControlAssetEditorEditMode.h b/Engine/Plugins/Experimental/PhysicsControl/Source/PhysicsControlEditor/Public/PhysicsControlAssetEditorEditMode.h
new file mode 100644
--- /dev/null
+++ b/Engine/Plugins/Experimental/PhysicsControl/Source/PhysicsControlEditor/Public/PhysicsControlAssetEditorEditMode.h
@@ -0,0 +1,101 @@
+// Copyright Epic Games, Inc. All Rights Reserved.
+
+#pragma once
+
+#include "CoreMinimal.h"
+#include "IPersonaEditMode.h"
+
+class FPhysicsControlAssetEditor;
+class FPhysicsControlAssetEditorData;
+class FEditorViewportClient;
+class UFont;
+
+/**
+ * The Physics Control Asset Editor edit mode.
+ */
+class PHYSICSCONTROLEDITOR_API FPhysicsControlAssetEditorEditMode : public IPersonaEditMode
+{
+public:
+	/** The name of the mode. */
+	static FName ModeName;
+
+	FPhysicsControlAssetEditorEditMode();
+
+	void SetEditorData(
+		const TSharedRef<FPhysicsControlAssetEditor> InPhysicsControlAssetEditor,
+		TSharedPtr<FPhysicsControlAssetEditorData>   InEditorData);
+
+	// IPersonaEditMode overrides.
+	virtual bool GetCameraTarget(FSphere& OutTarget) const override;
+	virtual class IPersonaPreviewScene& GetAnimPreviewScene() const override;
+	virtual void GetOnScreenDebugInfo(TArray<FText>& OutDebugInfo) const override {}
+	virtual void Exit() override;
+	// ~END IPersonaEditMode overrides.
+
+	// FEdMode overrides.
+	virtual bool InputKey(FEditorViewportClient* ViewportClient, FViewport* Viewport, FKey Key, EInputEvent Event) override;
+	virtual bool InputAxis(FEditorViewportClient* InViewportClient, FViewport* Viewport, int32 ControllerId, FKey Key, float Delta, float DeltaTime) override;
+	virtual void Tick(FEditorViewportClient* ViewportClient, float DeltaTime) override;
+	virtual void Render(const FSceneView* View, FViewport* Viewport, FPrimitiveDrawInterface* PDI) override;
+	virtual void DrawHUD(FEditorViewportClient* ViewportClient, FViewport* Viewport, const FSceneView* View, FCanvas* Canvas) override;
+	virtual bool IsCompatibleWith(FEditorModeID OtherModeID) const override { return true; }
+	virtual bool AllowWidgetMove() override { return false; }
+	virtual bool ShouldDrawWidget() const override { return false; }
+	virtual bool UsesTransformWidget() const override { return false; }
+	virtual bool UsesTransformWidget(UE::Widget::EWidgetMode CheckMode) const override { return false; }
+	virtual bool HandleClick(FEditorViewportClient* InViewportClient, HHitProxy* HitProxy, const FViewportClick& Click) override;
+	virtual bool ReceivedFocus(FEditorViewportClient* ViewportClient, FViewport* Viewport);
+	virtual bool LostFocus(FEditorViewportClient* ViewportClient, FViewport* Viewport);
+	// ~END FEdMode overrides.
+
+protected:
+	/** Simulation mouse forces */
+	bool SimMousePress(FEditorViewportClient* InViewportClient, FKey Key);
+	void SimMouseMove(FEditorViewportClient* InViewportClient, float DeltaX, float DeltaY);
+	bool SimMouseRelease();
+	bool SimMouseWheelUp(FEditorViewportClient* InViewportClient);
+	bool SimMouseWheelDown(FEditorViewportClient* InViewportClient);
+
+	void OpenBodyMenu(FEditorViewportClient* InViewportClient);
+	void OpenSelectionMenu(FEditorViewportClient* InViewportClient);
+
+	/** Called when no scene proxy is hit, deselects everything */
+	void HitNothing(FEditorViewportClient* InViewportClient);
+
+protected:
+	TWeakPtr<FPhysicsControlAssetEditor> PhysicsControlAssetEditor;
+
+	TSharedPtr<FPhysicsControlAssetEditorData> EditorData;
+
+	/** Font used for drawing debug text to the viewport */
+	UFont* PhysicsControlAssetEditorFont;
+
+	/** Misc consts */
+	const float	SimHoldDistanceChangeDelta;
+	const float	SimMinHoldDistance;
+	const float SimGrabMoveSpeed;
+
+	// For mouse forces TODO Danny surely this doesn't need to be here unless its actually shared!
+	FIntPoint LastClickPos;
+	FVector LastClickOrigin;
+	FVector LastClickDirection;
+	FVector LastClickHitPos;
+	FVector LastClickHitNormal;
+	bool bLastClickHit;
+
+	/** Simulation mouse forces */
+	float SimGrabPush;
+	float SimGrabMinPush;
+	FVector SimGrabLocation;
+	FVector SimGrabX;
+	FVector SimGrabY;
+	FVector SimGrabZ;
+
+	/** For interacting with the asset while the simulation is running */
+	TArray<FTransform> StartManRelConTM;
+	TArray<FTransform> StartManParentConTM;
+	TArray<FTransform> StartManChildConTM;
+	float DragX;
+	float DragY;
+
+};
diff --git a/Engine/Plugins/Experimental/PhysicsControl/Source/PhysicsControlEditor/Public/PhysicsControlAssetEditorMode.h b/Engine/Plugins/Experimental/PhysicsControl/Source/PhysicsControlEditor/Public/PhysicsControlAssetEditorMode.h
deleted file mode 100644
diff --git a/Engine/Plugins/Experimental/PhysicsControl/Source/PhysicsControlEditor/Public/PhysicsControlAssetEditorToolkit.h b/Engine/Plugins/Experimental/PhysicsControl/Source/PhysicsControlEditor/Public/PhysicsControlAssetEditorToolkit.h
deleted file mode 100644
diff --git a/Engine/Plugins/Experimental/PhysicsControl/Source/PhysicsControlEditor/Public/PhysicsControlEditorModule.h b/Engine/Plugins/Experimental/PhysicsControl/Source/PhysicsControlEditor/Public/PhysicsControlEditorModule.h
--- a/Engine/Plugins/Experimental/PhysicsControl/Source/PhysicsControlEditor/Public/PhysicsControlEditorModule.h
+++ b/Engine/Plugins/Experimental/PhysicsControl/Source/PhysicsControlEditor/Public/PhysicsControlEditorModule.h
@@ -11,7 +11,7 @@ class FPhysicsControlOperatorEditor;
 class FPhysicsControlAssetActions;
 
 /**
- * Module handles the Physics Control Profile editor, and also the Physics Control Visualizer
+ * Module handles the Physics Control Asset editor, and also the Physics Control Visualizer
  */
 class FPhysicsControlEditorModule : public IModuleInterface
 {
