Small fix to make sure Physics Control Assets will open without crashing when they don't have a Preview Physics Asset assigned

#rnx

diff --git a/Engine/Plugins/Experimental/PhysicsControl/Source/PhysicsControlEditor/Private/PhysicsControlAssetEditor.cpp b/Engine/Plugins/Experimental/PhysicsControl/Source/PhysicsControlEditor/Private/PhysicsControlAssetEditor.cpp
--- a/Engine/Plugins/Experimental/PhysicsControl/Source/PhysicsControlEditor/Private/PhysicsControlAssetEditor.cpp
+++ b/Engine/Plugins/Experimental/PhysicsControl/Source/PhysicsControlEditor/Private/PhysicsControlAssetEditor.cpp
@@ -25,6 +25,7 @@
 #include "IPersonaToolkit.h"
 #include "IPinnedCommandList.h"
 #include "ISkeletonEditorModule.h"
+#include "Misc/MessageDialog.h"
 #include "Modules/ModuleManager.h"
 #include "PersonaModule.h"
 #include "PersonaToolMenuContext.h"
@@ -46,9 +47,11 @@ const FName PhysicsControlAssetEditorAppName = FName(TEXT("PhysicsControlAssetEd
 //======================================================================================================================
 namespace PhysicsControlAssetEditor
 {
-	static TSharedPtr<FPhysicsControlAssetEditor> GetPhysicsControlAssetEditorFromToolContext(const FToolMenuContext& InMenuContext)
+	static TSharedPtr<FPhysicsControlAssetEditor> GetPhysicsControlAssetEditorFromToolContext(
+		const FToolMenuContext& InMenuContext)
 	{
-		if (UPhysicsControlAssetEditorToolMenuContext* Context = InMenuContext.FindContext<UPhysicsControlAssetEditorToolMenuContext>())
+		if (UPhysicsControlAssetEditorToolMenuContext* Context = 
+			InMenuContext.FindContext<UPhysicsControlAssetEditorToolMenuContext>())
 		{
 			return Context->PhysicsControlAssetEditor.Pin();
 		}
@@ -86,33 +89,50 @@ void FPhysicsControlAssetEditor::InitAssetEditor(
 	// Make the skeleton tree
 	{
 		FSkeletonTreeArgs SkeletonTreeArgs;
-		SkeletonTreeArgs.OnSelectionChanged = FOnSkeletonTreeSelectionChanged::CreateSP(this, &FPhysicsControlAssetEditor::HandleSelectionChanged);
+		SkeletonTreeArgs.OnSelectionChanged = FOnSkeletonTreeSelectionChanged::CreateSP(
+			this, &FPhysicsControlAssetEditor::HandleSelectionChanged);
 		SkeletonTreeArgs.PreviewScene = PersonaToolkit->GetPreviewScene();
 		SkeletonTreeArgs.bShowBlendProfiles = false;
 		SkeletonTreeArgs.bShowDebugVisualizationOptions = true;
 		SkeletonTreeArgs.bAllowMeshOperations = false;
 		SkeletonTreeArgs.bAllowSkeletonOperations = false;
 		SkeletonTreeArgs.bHideBonesByDefault = true;
-		SkeletonTreeArgs.OnGetFilterText = FOnGetFilterText::CreateSP(this, &FPhysicsControlAssetEditor::HandleGetFilterLabel);
+		SkeletonTreeArgs.OnGetFilterText = FOnGetFilterText::CreateSP(
+			this, &FPhysicsControlAssetEditor::HandleGetFilterLabel);
 		SkeletonTreeArgs.Extenders = MakeShared<FExtender>();
-		SkeletonTreeArgs.Extenders->AddMenuExtension("FilterOptions", EExtensionHook::After, GetToolkitCommands(), FMenuExtensionDelegate::CreateSP(this, &FPhysicsControlAssetEditor::HandleExtendFilterMenu));
-		SkeletonTreeArgs.Extenders->AddMenuExtension("SkeletonTreeContextMenu", EExtensionHook::After, GetToolkitCommands(), FMenuExtensionDelegate::CreateSP(this, &FPhysicsControlAssetEditor::HandleExtendContextMenu));
+		SkeletonTreeArgs.Extenders->AddMenuExtension(
+			"FilterOptions", EExtensionHook::After, GetToolkitCommands(), 
+			FMenuExtensionDelegate::CreateSP(this, &FPhysicsControlAssetEditor::HandleExtendFilterMenu));
+		SkeletonTreeArgs.Extenders->AddMenuExtension(
+			"SkeletonTreeContextMenu", EExtensionHook::After, GetToolkitCommands(), 
+			FMenuExtensionDelegate::CreateSP(this, &FPhysicsControlAssetEditor::HandleExtendContextMenu));
 
 		if (SkeletalMesh)
 		{
 			if (UPhysicsAsset* PhysicsAsset = SkeletalMesh->GetPhysicsAsset())
 			{
-				SkeletonTreeArgs.Builder = SkeletonTreeBuilder = MakeShared<FPhysicsControlAssetEditorSkeletonTreeBuilder>(PhysicsAsset);
+				SkeletonTreeArgs.Builder = SkeletonTreeBuilder = 
+					MakeShared<FPhysicsControlAssetEditorSkeletonTreeBuilder>(PhysicsAsset);
 			}
 		}
 		SkeletonTreeArgs.ContextName = GetToolkitFName();
 
-		ISkeletonEditorModule& SkeletonEditorModule = FModuleManager::GetModuleChecked<ISkeletonEditorModule>("SkeletonEditor");
-
 		GetMutableDefault<UPersonaOptions>()->bFlattenSkeletonHierarchyWhenFiltering = false;
 		GetMutableDefault<UPersonaOptions>()->bHideParentsWhenFiltering = true;
-
-		SkeletonTree = SkeletonEditorModule.CreateSkeletonTree(PersonaToolkit->GetSkeleton(), SkeletonTreeArgs);
+		if (PersonaToolkit->GetSkeleton())
+		{
+			ISkeletonEditorModule& SkeletonEditorModule = 
+				FModuleManager::LoadModuleChecked<ISkeletonEditorModule>("SkeletonEditor");
+			SkeletonTree = SkeletonEditorModule.CreateSkeletonTree(PersonaToolkit->GetSkeleton(), SkeletonTreeArgs);
+		}
+		else
+		{
+			FMessageDialog::Open(EAppMsgType::Ok, LOCTEXT(
+				"Error_PhysicsControlAssetHasNoSkeleton",
+				"Warning: Physics Control Asset has no skeleton assigned.\n"
+				"This is likely to be because there is no valid Physics Asset. "
+				"Fix this by assigning a Preview Physics Asset/Mesh in the Physics Control Asset."));
+		}
 	}
 
 	bSelecting = false;
@@ -140,7 +160,7 @@ void FPhysicsControlAssetEditor::InitAssetEditor(
 	AddApplicationMode(
 		PhysicsControlAssetEditorModes::PhysicsControlAssetEditorMode,
 		MakeShareable(new FPhysicsControlAssetApplicationMode(
-			SharedThis(this), SkeletonTree.ToSharedRef(), PersonaToolkit->GetPreviewScene())));
+			SharedThis(this), SkeletonTree, PersonaToolkit->GetPreviewScene())));
 	SetCurrentMode(PhysicsControlAssetEditorModes::PhysicsControlAssetEditorMode);
 
 	// Activate the editor mode.
@@ -591,7 +611,8 @@ void FPhysicsControlAssetEditor::OnToggleSimulationFloorCollision()
 		{
 			TSharedRef<IPersonaPreviewScene> PersonaPreviewScene = PersonaToolkit->GetPreviewScene();
 
-			if (UStaticMeshComponent* FloorMeshComponent = const_cast<UStaticMeshComponent*>(PersonaPreviewScene->GetFloorMeshComponent()))
+			if (UStaticMeshComponent* FloorMeshComponent = 
+				const_cast<UStaticMeshComponent*>(PersonaPreviewScene->GetFloorMeshComponent()))
 			{
 				if (EditorData->EditorOptions->bSimulationFloorCollisionEnabled)
 				{
@@ -626,8 +647,9 @@ void FPhysicsControlAssetEditor::OnMeshRenderingMode(EPhysicsAssetEditorMeshView
 
 	EditorData->EditorOptions->SaveConfig();
 
-	// Changing the mesh rendering mode requires the skeletal mesh component to change its render state, which is an operation
-	// which is deferred until after render. Hence we need to trigger another viewport refresh on the following frame.
+	// Changing the mesh rendering mode requires the skeletal mesh component to change its render
+	// state, which is an operation which is deferred until after render. Hence we need to trigger
+	// another viewport refresh on the following frame.
 	RefreshPreviewViewport();
 }
 
@@ -653,7 +675,8 @@ void FPhysicsControlAssetEditor::OnCollisionRenderingMode(EPhysicsAssetEditorCol
 }
 
 //======================================================================================================================
-bool FPhysicsControlAssetEditor::IsCollisionRenderingMode(EPhysicsAssetEditorCollisionViewMode Mode, bool bSimulation) const
+bool FPhysicsControlAssetEditor::IsCollisionRenderingMode(
+	EPhysicsAssetEditorCollisionViewMode Mode, bool bSimulation) const
 {
 	return Mode == EditorData->GetCurrentCollisionViewMode(bSimulation);
 }
@@ -676,7 +699,8 @@ void FPhysicsControlAssetEditor::OnConstraintRenderingMode(EPhysicsAssetEditorCo
 }
 
 //======================================================================================================================
-bool FPhysicsControlAssetEditor::IsConstraintRenderingMode(EPhysicsAssetEditorConstraintViewMode Mode, bool bSimulation) const
+bool FPhysicsControlAssetEditor::IsConstraintRenderingMode(
+	EPhysicsAssetEditorConstraintViewMode Mode, bool bSimulation) const
 {
 	return Mode == EditorData->GetCurrentConstraintViewMode(bSimulation);
 }
@@ -917,7 +941,8 @@ void FPhysicsControlAssetEditor::BuildMenuWidgetSelection(FMenuBuilder& InMenuBu
 
 //======================================================================================================================
 // Danny TODO selection needs to be implemented/handled/made useful
-void FPhysicsControlAssetEditor::HandleSelectionChanged(const TArrayView<TSharedPtr<ISkeletonTreeItem>>& InSelectedItems, ESelectInfo::Type InSelectInfo)
+void FPhysicsControlAssetEditor::HandleSelectionChanged(
+	const TArrayView<TSharedPtr<ISkeletonTreeItem>>& InSelectedItems, ESelectInfo::Type InSelectInfo)
 {
 	if (!bSelecting)
 	{
@@ -926,7 +951,12 @@ void FPhysicsControlAssetEditor::HandleSelectionChanged(const TArrayView<TShared
 		// Always set the details customization object, regardless of selection type
 		// We do this because the tree may have been rebuilt and objects invalidated
 		TArray<UObject*> Objects;
-		Algo::TransformIf(InSelectedItems, Objects, [](const TSharedPtr<ISkeletonTreeItem>& InItem) { return InItem->GetObject() != nullptr; }, [](const TSharedPtr<ISkeletonTreeItem>& InItem) { return InItem->GetObject(); });
+		Algo::TransformIf(InSelectedItems, Objects, 
+			[](const TSharedPtr<ISkeletonTreeItem>& InItem) {
+			return InItem->GetObject() != nullptr; },
+			[](const TSharedPtr<ISkeletonTreeItem>& InItem) { 
+				return InItem->GetObject(); }
+		);
 
 		if (DetailsView.IsValid())
 		{
@@ -962,8 +992,10 @@ void FPhysicsControlAssetEditor::HandleExtendContextMenu(FMenuBuilder& InMenuBui
 	TArray<TSharedPtr<ISkeletonTreeItem>> SelectedItems = SkeletonTree->GetSelectedItems();
 	FSkeletonTreeSelection Selection(SelectedItems);
 
-	TArray<TSharedPtr<FSkeletonTreePhysicsControlBodyItem>> SelectedBodies = Selection.GetSelectedItems<FSkeletonTreePhysicsControlBodyItem>();
-	TArray<TSharedPtr<FSkeletonTreePhysicsControlShapeItem>> SelectedShapes = Selection.GetSelectedItems<FSkeletonTreePhysicsControlShapeItem>();
+	TArray<TSharedPtr<FSkeletonTreePhysicsControlBodyItem>> SelectedBodies = 
+		Selection.GetSelectedItems<FSkeletonTreePhysicsControlBodyItem>();
+	TArray<TSharedPtr<FSkeletonTreePhysicsControlShapeItem>> SelectedShapes = 
+		Selection.GetSelectedItems<FSkeletonTreePhysicsControlShapeItem>();
 	TArray<TSharedPtr<ISkeletonTreeItem>> SelectedBones = Selection.GetSelectedItemsByTypeId("FSkeletonTreeBoneItem");
 }
 
@@ -1087,10 +1119,14 @@ void FPhysicsControlAssetEditor::HandleViewportCreated(const TSharedRef<IPersona
 {
 	PersonaViewport = InPersonaViewport;
 	InPersonaViewport.Get().GetPinnedCommandList()->BindCommandList(ViewportCommandList.ToSharedRef());
-	InPersonaViewport.Get().GetPinnedCommandList()->RegisterCustomWidget(IPinnedCommandList::FOnGenerateCustomWidget::CreateSP(
-		this, &FPhysicsControlAssetEditor::MakeConstraintScaleWidget), TEXT("ConstraintScaleWidget"), LOCTEXT("ConstraintScaleLabel", "Constraint Scale"));
-	InPersonaViewport.Get().GetPinnedCommandList()->RegisterCustomWidget(IPinnedCommandList::FOnGenerateCustomWidget::CreateSP(
-		this, &FPhysicsControlAssetEditor::MakeCollisionOpacityWidget), TEXT("CollisionOpacityWidget"), LOCTEXT("CollisionOpacityLabel", "Collision Opacity"));
+	InPersonaViewport.Get().GetPinnedCommandList()->RegisterCustomWidget(
+		IPinnedCommandList::FOnGenerateCustomWidget::CreateSP(
+		this, &FPhysicsControlAssetEditor::MakeConstraintScaleWidget), 
+		TEXT("ConstraintScaleWidget"), LOCTEXT("ConstraintScaleLabel", "Constraint Scale"));
+	InPersonaViewport.Get().GetPinnedCommandList()->RegisterCustomWidget(
+		IPinnedCommandList::FOnGenerateCustomWidget::CreateSP(
+		this, &FPhysicsControlAssetEditor::MakeCollisionOpacityWidget), 
+		TEXT("CollisionOpacityWidget"), LOCTEXT("CollisionOpacityLabel", "Collision Opacity"));
 }
 
 //======================================================================================================================
@@ -1139,7 +1175,5 @@ void FPhysicsControlAssetEditor::RefreshPreviewViewport()
 	}
 }
 
-
-
 #undef LOCTEXT_NAMESPACE
 
diff --git a/Engine/Plugins/Experimental/PhysicsControl/Source/PhysicsControlEditor/Private/PhysicsControlAssetEditorData.cpp b/Engine/Plugins/Experimental/PhysicsControl/Source/PhysicsControlEditor/Private/PhysicsControlAssetEditorData.cpp
--- a/Engine/Plugins/Experimental/PhysicsControl/Source/PhysicsControlEditor/Private/PhysicsControlAssetEditorData.cpp
+++ b/Engine/Plugins/Experimental/PhysicsControl/Source/PhysicsControlEditor/Private/PhysicsControlAssetEditorData.cpp
@@ -71,26 +71,36 @@ void FPhysicsControlAssetEditorData::CachePreviewMesh()
 		// Fall back to the default skeletal mesh in the EngineMeshes package.
 		// This is statically loaded as the package is likely not fully loaded
 		// (otherwise, it would have been found in the above iteration).
-		PreviewMesh = (USkeletalMesh*)StaticLoadObject(USkeletalMesh::StaticClass(), NULL, TEXT("/Engine/EngineMeshes/SkeletalCube.SkeletalCube"), NULL, LOAD_None, NULL);
+		PreviewMesh = (USkeletalMesh*)StaticLoadObject(
+			USkeletalMesh::StaticClass(), NULL, 
+			TEXT("/Engine/EngineMeshes/SkeletalCube.SkeletalCube"), NULL, LOAD_None, NULL);
 		check(PreviewMesh);
 
 		PhysicsControlAsset->SetPreviewMesh(PreviewMesh);
 
-		FMessageDialog::Open(EAppMsgType::Ok, FText::Format(
-			LOCTEXT("Error_PhysicsControlAssetHasNoSkelMesh", "Warning: Physics Control Asset has no skeletal mesh assigned.\nFor now, a simple default skeletal mesh ({0}) will be used.\nYou can fix this by opening the asset and choosing another skeletal mesh from the toolbar."),
-			FText::FromString(PreviewMesh->GetFullName())));
+		FMessageDialog::Open(EAppMsgType::Ok, LOCTEXT(
+			"Error_PhysicsControlAssetHasNoSkelMesh",
+			"Warning: Physics Control Asset has no skeletal mesh assigned.\n"
+			"This is likely to be because there is no valid Physics Asset. "
+			"Fix this by assigning a Preview Physics Asset/Mesh in the Physics Control Asset.")
+		);
 	}
 	else if (PreviewMesh->GetSkeleton() == nullptr)
 	{
 		// Fall back in the case of a deleted skeleton
-		PreviewMesh = (USkeletalMesh*)StaticLoadObject(USkeletalMesh::StaticClass(), NULL, TEXT("/Engine/EngineMeshes/SkeletalCube.SkeletalCube"), NULL, LOAD_None, NULL);
+		PreviewMesh = (USkeletalMesh*)StaticLoadObject(
+			USkeletalMesh::StaticClass(), NULL, 
+			TEXT("/Engine/EngineMeshes/SkeletalCube.SkeletalCube"), NULL, LOAD_None, NULL);
 		check(PreviewMesh);
 
 		PhysicsControlAsset->SetPreviewMesh(PreviewMesh);
 
-		FMessageDialog::Open(EAppMsgType::Ok, FText::Format(
-			LOCTEXT("Error_PhysicsControlAssetHasNoSkelMeshSkeleton", "Warning: Physics Control Asset has a skeletal mesh with no skeleton assigned.\nFor now, a simple default skeletal mesh ({0}) will be used.\nYou can fix this by opening the asset and choosing another skeletal mesh from the toolbar, or repairing the skeleton."),
-			FText::FromString(PreviewMesh->GetFullName())));
+		FMessageDialog::Open(EAppMsgType::Ok, LOCTEXT(
+			"Error_PhysicsControlAssetHasNoSkelMeshSkeleton", 
+			"Warning: Physics Control Asset has no skeletal mesh skeleton assigned.\n"
+			"This is likely to be because there is no valid Physics Asset. "
+			"Fix this by assigning a Preview Physics Asset/Mesh in the Physics Control Asset.")
+		);
 	}
 }
 
