Remove a unused duplicate definition of cvar "p.PhAT.DebugViewportClicks" which was causing a warning during startup

#rnx

diff --git a/Engine/Plugins/Experimental/PhysicsControl/Source/PhysicsControlEditor/Private/PhysicsControlAssetEditorSkeletalMeshComponent.cpp b/Engine/Plugins/Experimental/PhysicsControl/Source/PhysicsControlEditor/Private/PhysicsControlAssetEditorSkeletalMeshComponent.cpp
--- a/Engine/Plugins/Experimental/PhysicsControl/Source/PhysicsControlEditor/Private/PhysicsControlAssetEditorSkeletalMeshComponent.cpp
+++ b/Engine/Plugins/Experimental/PhysicsControl/Source/PhysicsControlEditor/Private/PhysicsControlAssetEditorSkeletalMeshComponent.cpp
@@ -26,12 +26,6 @@
 #include "Styling/AppStyle.h"
 #include "UObject/Package.h"
 
-namespace
-{
-	bool bDebugViewportClicks = false;
-	FAutoConsoleVariableRef CVarChaosImmPhysStepTime(TEXT("p.PhAT.DebugViewportClicks"), bDebugViewportClicks, TEXT("Set to 1 to show mouse click results in PhAT"));
-}
-
 UPhysicsControlAssetEditorSkeletalMeshComponent::UPhysicsControlAssetEditorSkeletalMeshComponent(const FObjectInitializer& ObjectInitializer)
 	: Super(ObjectInitializer)
 	, BoneUnselectedColor(170, 155, 225)
