Progress on Physics Control Asset Editor - making it functional for editing and previewing

diff --git a/Engine/Plugins/Experimental/PhysicsControl/Config/BasePhysicsControl.ini b/Engine/Plugins/Experimental/PhysicsControl/Config/BasePhysicsControl.ini
--- a/Engine/Plugins/Experimental/PhysicsControl/Config/BasePhysicsControl.ini
+++ b/Engine/Plugins/Experimental/PhysicsControl/Config/BasePhysicsControl.ini
@@ -12,5 +12,7 @@
 +PropertyRedirects=(OldName="PhysicsControlComponent.CreateBodyModifier.MeshComponent",NewName="Component")
 +PropertyRedirects=(OldName="PhysicsControlComponent.CreateNamedBodyModifier.MeshComponent",NewName="Component")
 
++FunctionRedirects=(OldName="PhysicsControlComponent.CreateControlsAndBodyModifiersFromControlProfileAsset",NewName="CreateControlsAndBodyModifiersFromPhysicsControlAsset")
+
 +PropertyRedirects=(OldName="PhysicsControlProfileAsset",NewName="PhysicsControlAsset")
 +ClassRedirects=(OldName="PhysicsControlProfileAsset",NewName="PhysicsControlAsset")
diff --git a/Engine/Plugins/Experimental/PhysicsControl/Source/PhysicsControl/PhysicsControl.Build.cs b/Engine/Plugins/Experimental/PhysicsControl/Source/PhysicsControl/PhysicsControl.Build.cs
--- a/Engine/Plugins/Experimental/PhysicsControl/Source/PhysicsControl/PhysicsControl.Build.cs
+++ b/Engine/Plugins/Experimental/PhysicsControl/Source/PhysicsControl/PhysicsControl.Build.cs
@@ -29,8 +29,5 @@ public class PhysicsControl : ModuleRules
 				"PhysicsCore"
 			}
 		);
-
-		SetupModulePhysicsSupport(Target);
-		PrivateDefinitions.Add("CHAOS_INCLUDE_LEVEL_1=1");
 	}
 }
diff --git a/Engine/Plugins/Experimental/PhysicsControl/Source/PhysicsControl/Private/PhysicsControlAsset.cpp b/Engine/Plugins/Experimental/PhysicsControl/Source/PhysicsControl/Private/PhysicsControlAsset.cpp
--- a/Engine/Plugins/Experimental/PhysicsControl/Source/PhysicsControl/Private/PhysicsControlAsset.cpp
+++ b/Engine/Plugins/Experimental/PhysicsControl/Source/PhysicsControl/Private/PhysicsControlAsset.cpp
@@ -5,19 +5,9 @@
 #include "Engine/SkeletalMesh.h"
 #include "PhysicsEngine/PhysicsAsset.h"
 
-//#ifdef WITH_EDITOR
-//#include "Editor.h"
-//#endif // WITH_EDITOR
-
-
 //======================================================================================================================
 UPhysicsControlAsset::UPhysicsControlAsset()
 {
-	// This needs to be explored further - a possible way to hook into the BP compilation. However, we would also
-	// need to make sure we don't compile (and get marked as dirty) if there are no changes.
-//#ifdef WITH_EDITOR
-//	GEditor->OnBlueprintCompiled().AddUObject(this, &UPhysicsControlAsset::Compile);
-//#endif
 }
 
 #if WITH_EDITOR
@@ -63,8 +53,19 @@ void UPhysicsControlAsset::Compile()
 	AdditionalControlsAndModifiers = GetAdditionalControlsAndModifiers();
 	AdditionalSets = GetAdditionalSets();
 	InitialControlAndModifierUpdates = GetInitialControlAndModifierUpdates();
+
+	TArray<FName> OrigKeys;
+	Profiles.GetKeys(OrigKeys);
+
 	Profiles = GetProfiles();
 
+	TArray<FName> NewKeys;
+	Profiles.GetKeys(NewKeys);
+
+	bool bProfileListChanged = (OrigKeys != NewKeys);
+
+	OnControlAssetCompiledDelegate.Broadcast(bProfileListChanged);
+
 	Modify();
 }
 
@@ -94,6 +95,48 @@ bool UPhysicsControlAsset::IsCompilationNeeded() const
 	return false;
 }
 
+//======================================================================================================================
+bool UPhysicsControlAsset::IsSetupDirty() const
+{
+	if (CharacterSetupData != GetCharacterSetupData())
+	{
+		return true;
+	}
+	if (AdditionalControlsAndModifiers != GetAdditionalControlsAndModifiers())
+	{
+		return true;
+	}
+	if (AdditionalSets != GetAdditionalSets())
+	{
+		return true;
+	}
+	if (InitialControlAndModifierUpdates != GetInitialControlAndModifierUpdates())
+	{
+		return true;
+	}
+	return false;
+}
+
+//======================================================================================================================
+TArray<FName> UPhysicsControlAsset::GetDirtyProfiles() const
+{
+	TArray<FName> DirtyProfiles;
+
+	TMap<FName, FPhysicsControlControlAndModifierUpdates> CompiledProfiles = GetProfiles();
+
+	for (TMap<FName, FPhysicsControlControlAndModifierUpdates>::ElementType ProfilePair : Profiles)
+	{
+		FPhysicsControlControlAndModifierUpdates* Compiled = CompiledProfiles.Find(ProfilePair.Key);
+		if (Compiled && *Compiled == ProfilePair.Value)
+		{
+			continue;
+		}
+		DirtyProfiles.Push(ProfilePair.Key);
+	}
+	return DirtyProfiles;
+}
+
+
 //======================================================================================================================
 FPhysicsControlCharacterSetupData UPhysicsControlAsset::GetCharacterSetupData() const
 {
diff --git a/Engine/Plugins/Experimental/PhysicsControl/Source/PhysicsControl/Private/PhysicsControlComponent.cpp b/Engine/Plugins/Experimental/PhysicsControl/Source/PhysicsControl/Private/PhysicsControlComponent.cpp
--- a/Engine/Plugins/Experimental/PhysicsControl/Source/PhysicsControl/Private/PhysicsControlComponent.cpp
+++ b/Engine/Plugins/Experimental/PhysicsControl/Source/PhysicsControl/Private/PhysicsControlComponent.cpp
@@ -685,6 +685,13 @@ TMap<FName, FPhysicsControlNames> UPhysicsControlComponent::CreateControlsFromLi
 	return Result;
 }
 
+//======================================================================================================================
+void UPhysicsControlComponent::DestroyAllControlsAndBodyModifiers()
+{
+	DestroyControlsInSet(TEXT("All"));
+	DestroyBodyModifiersInSet(TEXT("All"));
+}
+
 //======================================================================================================================
 bool UPhysicsControlComponent::DestroyControl(const FName Name)
 {
@@ -2050,11 +2057,11 @@ void UPhysicsControlComponent::SetBodyModifiersInSetUpdateKinematicFromSimulatio
 //======================================================================================================================
 const TArray<FName>& UPhysicsControlComponent::GetAllControlNames() const
 {
-	return GetControlNamesInSet("All");
+	return GetControlNamesInSet(TEXT("All"));
 }
 
 //======================================================================================================================
-void UPhysicsControlComponent::CreateControlsAndBodyModifiersFromLimbBones(
+bool UPhysicsControlComponent::CreateControlsAndBodyModifiersFromLimbBones(
 	FPhysicsControlNames&                       AllWorldSpaceControls,
 	TMap<FName, FPhysicsControlNames>&          LimbWorldSpaceControls,
 	FPhysicsControlNames&                       AllParentSpaceControls,
@@ -2069,6 +2076,13 @@ void UPhysicsControlComponent::CreateControlsAndBodyModifiersFromLimbBones(
 	UPrimitiveComponent*                        WorldComponent,
 	FName                                       WorldBoneName)
 {
+	UPhysicsAsset* PhysicsAsset = SkeletalMeshComponent ? SkeletalMeshComponent->GetPhysicsAsset() : nullptr;
+	if (!PhysicsAsset)
+	{
+		UE_LOG(LogPhysicsControl, Warning, TEXT("No physics asset in skeletal mesh"));
+		return false;
+	}
+
 	TMap<FName, FPhysicsControlLimbBones> LimbBones = 
 		GetLimbBonesFromSkeletalMesh(SkeletalMeshComponent, LimbSetupData);
 
@@ -2081,19 +2095,21 @@ void UPhysicsControlComponent::CreateControlsAndBodyModifiersFromLimbBones(
 		ParentSpaceControlData);
 
 	LimbBodyModifiers = CreateBodyModifiersFromLimbBones(AllBodyModifiers, LimbBones, BodyModifierData);
+
+	return true;
 }
 
 //======================================================================================================================
-void UPhysicsControlComponent::CreateControlsAndBodyModifiersFromControlProfileAsset(
+bool UPhysicsControlComponent::CreateControlsAndBodyModifiersFromPhysicsControlAsset(
 	USkeletalMeshComponent* SkeletalMeshComponent,
 	UPrimitiveComponent*    WorldComponent,
 	FName                   WorldBoneName)
 {
-	if (!PhysicsControlAsset.IsValid())
+	if (!PhysicsControlAsset.LoadSynchronous())
 	{
 		UE_LOG(LogPhysicsControl, Warning,
-			TEXT("CreateControlsAndBodyModifiersFromControlProfile - unable to get/load the control profile asset"));
-		return;
+			TEXT("CreateControlsAndBodyModifiersFromPhysicsControlAsset - unable to get/load the control profile asset"));
+		return false;
 	}
 
 	FPhysicsControlNames AllWorldSpaceControls;
@@ -2103,7 +2119,7 @@ void UPhysicsControlComponent::CreateControlsAndBodyModifiersFromControlProfileA
 	FPhysicsControlNames AllBodyModifiers;
 	TMap<FName, FPhysicsControlNames> LimbBodyModifiers;
 
-	CreateControlsAndBodyModifiersFromLimbBones(
+	if (!CreateControlsAndBodyModifiersFromLimbBones(
 		AllWorldSpaceControls, LimbWorldSpaceControls, AllParentSpaceControls, LimbParentSpaceControls, 
 		AllBodyModifiers, LimbBodyModifiers,
 		SkeletalMeshComponent,
@@ -2112,7 +2128,12 @@ void UPhysicsControlComponent::CreateControlsAndBodyModifiersFromControlProfileA
 		PhysicsControlAsset->CharacterSetupData.DefaultParentSpaceControlData,
 		PhysicsControlAsset->CharacterSetupData.DefaultBodyModifierData,
 		WorldComponent,
-		WorldBoneName);
+		WorldBoneName))
+	{
+		// We assume that if this one fails, then everything fails. Also that if we can create the
+		// basic setup, then the rest is OK too.
+		return false;
+	}
 
 	// Create additional controls
 	for (const TPair<FName, FPhysicsControlCreationData>& ControlPair : 
@@ -2158,10 +2179,11 @@ void UPhysicsControlComponent::CreateControlsAndBodyModifiersFromControlProfileA
 	{
 		ApplyControlAndModifierUpdates(Updates);
 	}
+	return true;
 }
 
 //======================================================================================================================
-void UPhysicsControlComponent::InvokeControlProfile(FName ProfileName)
+bool UPhysicsControlComponent::InvokeControlProfile(FName ProfileName)
 {
 	if (!PhysicsControlAsset.IsValid())
 	{
@@ -2170,7 +2192,7 @@ void UPhysicsControlComponent::InvokeControlProfile(FName ProfileName)
 			UE_LOG(LogPhysicsControl, Warning,
 				TEXT("InvokeControlProfile - control profile asset is invalid or missing"));
 		}
-		return;
+		return false;
 	}
 
 	const FPhysicsControlControlAndModifierUpdates* ControlAndModifierUpdates =
@@ -2183,12 +2205,12 @@ void UPhysicsControlComponent::InvokeControlProfile(FName ProfileName)
 			UE_LOG(LogPhysicsControl, Warning,
 				TEXT("InvokeControlProfile - control profile %s not found"), *ProfileName.ToString());
 		}
-		return;
+		return false;
 	}
 
 	ApplyControlAndModifierUpdates(*ControlAndModifierUpdates);
 
-	return;
+	return true;
 }
 
 //======================================================================================================================
@@ -2289,7 +2311,7 @@ const TArray<FName>& UPhysicsControlComponent::GetControlNamesInSet(const FName
 //======================================================================================================================
 const TArray<FName>& UPhysicsControlComponent::GetAllBodyModifierNames() const
 {
-	return GetBodyModifierNamesInSet("All");
+	return GetBodyModifierNamesInSet(TEXT("All"));
 }
 
 //======================================================================================================================
diff --git a/Engine/Plugins/Experimental/PhysicsControl/Source/PhysicsControl/Private/PhysicsControlNameRecords.cpp b/Engine/Plugins/Experimental/PhysicsControl/Source/PhysicsControl/Private/PhysicsControlNameRecords.cpp
--- a/Engine/Plugins/Experimental/PhysicsControl/Source/PhysicsControl/Private/PhysicsControlNameRecords.cpp
+++ b/Engine/Plugins/Experimental/PhysicsControl/Source/PhysicsControl/Private/PhysicsControlNameRecords.cpp
@@ -16,7 +16,7 @@ void AddUniqueNamesToSet(const TArray<FName>& Names, const FName SetName, TMap<F
 void FPhysicsControlNameRecords::AddControl(FName Name, FName SetName)
 {
 	ControlSets.FindOrAdd(SetName).AddUnique(Name);
-	ControlSets.FindOrAdd("All").AddUnique(Name);
+	ControlSets.FindOrAdd(TEXT("All")).AddUnique(Name);
 }
 
 //======================================================================================================================
@@ -26,7 +26,7 @@ void FPhysicsControlNameRecords::AddControl(FName Name, const TArray<FName>& Set
 	{
 		ControlSets.FindOrAdd(SetName).AddUnique(Name);
 	}
-	ControlSets.FindOrAdd("All").AddUnique(Name);
+	ControlSets.FindOrAdd(TEXT("All")).AddUnique(Name);
 }
 
 
@@ -62,7 +62,7 @@ const TArray<FName>& FPhysicsControlNameRecords::GetControlNamesInSet(FName SetN
 void FPhysicsControlNameRecords::AddBodyModifier(FName Name, FName SetName)
 {
 	BodyModifierSets.FindOrAdd(SetName).AddUnique(Name);
-	BodyModifierSets.FindOrAdd("All").AddUnique(Name);
+	BodyModifierSets.FindOrAdd(TEXT("All")).AddUnique(Name);
 }
 
 //======================================================================================================================
@@ -72,14 +72,14 @@ void FPhysicsControlNameRecords::AddBodyModifier(FName Name, const TArray<FName>
 	{
 		BodyModifierSets.FindOrAdd(SetName).AddUnique(Name);
 	}
-	BodyModifierSets.FindOrAdd("All").AddUnique(Name);
+	BodyModifierSets.FindOrAdd(TEXT("All")).AddUnique(Name);
 }
 
 //======================================================================================================================
 void FPhysicsControlNameRecords::AddBodyModifiers(const TArray<FName>& BodyModifierNames, const FName SetName)
 {
 	AddUniqueNamesToSet(BodyModifierNames, SetName, BodyModifierSets);
-	AddUniqueNamesToSet(BodyModifierNames, "All", BodyModifierSets);
+	AddUniqueNamesToSet(BodyModifierNames, TEXT("All"), BodyModifierSets);
 }
 
 //======================================================================================================================
@@ -126,7 +126,7 @@ TArray<FName> ExpandName(const FName InName, const TMap<FName, TArray<FName>>& S
 }
 
 //======================================================================================================================
-TArray<FName> ExpandName(const TArray<FName>& InNames, const TMap<FName, TArray<FName>>& SetNames)
+TArray<FName> ExpandNames(const TArray<FName>& InNames, const TMap<FName, TArray<FName>>& SetNames)
 {
 	TArray<FName> OutputNames;
 
diff --git a/Engine/Plugins/Experimental/PhysicsControl/Source/PhysicsControl/Private/PhysicsControlOperatorNameGeneration.cpp b/Engine/Plugins/Experimental/PhysicsControl/Source/PhysicsControl/Private/PhysicsControlOperatorNameGeneration.cpp
--- a/Engine/Plugins/Experimental/PhysicsControl/Source/PhysicsControl/Private/PhysicsControlOperatorNameGeneration.cpp
+++ b/Engine/Plugins/Experimental/PhysicsControl/Source/PhysicsControl/Private/PhysicsControlOperatorNameGeneration.cpp
@@ -217,7 +217,7 @@ void CreateAdditionalSets_Implementation(
 {
 	for (const FPhysicsControlSetUpdate& Set : AdditionalSets.ControlSetUpdates)
 	{
-		TArray<FName> Names = ExpandName(Set.Names, NameRecords.ControlSets);
+		TArray<FName> Names = ExpandNames(Set.Names, NameRecords.ControlSets);
 
 		for (FName Name : Names)
 		{
@@ -236,7 +236,7 @@ void CreateAdditionalSets_Implementation(
 
 	for (const FPhysicsControlSetUpdate& Set : AdditionalSets.ModifierSetUpdates)
 	{
-		TArray<FName> Names = ExpandName(Set.Names, NameRecords.BodyModifierSets);
+		TArray<FName> Names = ExpandNames(Set.Names, NameRecords.BodyModifierSets);
 
 		for (FName Name : Names)
 		{
diff --git a/Engine/Plugins/Experimental/PhysicsControl/Source/PhysicsControl/Public/PhysicsControlAsset.h b/Engine/Plugins/Experimental/PhysicsControl/Source/PhysicsControl/Public/PhysicsControlAsset.h
--- a/Engine/Plugins/Experimental/PhysicsControl/Source/PhysicsControl/Public/PhysicsControlAsset.h
+++ b/Engine/Plugins/Experimental/PhysicsControl/Source/PhysicsControl/Public/PhysicsControlAsset.h
@@ -25,6 +25,7 @@ UCLASS(BlueprintType)
 class PHYSICSCONTROL_API UPhysicsControlAsset : public UObject, public IInterface_PreviewMeshProvider
 {
 	GENERATED_BODY()
+
 public:
 	UPhysicsControlAsset();
 
@@ -65,6 +66,26 @@ public:
 	// Data that will then be compiled down into the runtime data
 
 #if WITH_EDITORONLY_DATA
+	// Whether editing the profiles will automatically compile.
+	UPROPERTY(EditAnywhere, Category = ProfileEditing)
+	bool bAutoCompileProfiles = true;
+
+	// Whether to automatically invoke profiles that have been edited (and have auto-compiled) when simulating
+	UPROPERTY(EditAnywhere, Category = ProfileEditing, Meta = (EditCondition="AutoCompileProfiles"))
+	bool bAutoInvokeProfiles = true;
+
+	// Whether editing the setup data will automatically compile.
+	UPROPERTY(EditAnywhere, Category = SetupEditing)
+	bool bAutoCompileSetup = true;
+
+	// Whether to automatically reinitialize following editing of setup data (when auto-compiling) when simulating
+	UPROPERTY(EditAnywhere, Category = SetupEditing, Meta = (EditCondition="AutoCompileSetup"))
+	bool bAutoReinitSetup = true;
+
+	// Whether to automatically re-invoke the previously invoked profile after automatically running the setup
+	UPROPERTY(EditAnywhere, Category = SetupEditing, Meta = (EditCondition="AutoCompileSetup"))
+	bool bAutoInvokeProfileAfterSetup = true;
+
 	/** A profile asset to inherit from (can be null). If set, we will just add/modify data in that */
 	UPROPERTY(EditAnywhere, Category = Inheritance)
 	TSoftObjectPtr<UPhysicsControlAsset> ParentAsset;
@@ -120,6 +141,9 @@ public:
 	// Buttons/actions in the editor
 
 #if WITH_EDITOR
+	DECLARE_EVENT_OneParam(FPhysicsControlAssetEditor, FOnControlAssetCompiled, bool);
+	FOnControlAssetCompiled& OnControlAssetCompiled() { return OnControlAssetCompiledDelegate; }
+
 	/** Shows all the controls etc that would be made */
 	UFUNCTION(CallInEditor, Category = Actions)
 	void ShowCompiledData() const;
@@ -138,6 +162,19 @@ public:
 	UFUNCTION(CallInEditor, Category = Actions)
 	bool IsCompilationNeeded() const;
 
+	/**
+	 * Returns a list of all the profiles that need compilation
+	 */
+	UFUNCTION(CallInEditor, Category = Actions)
+	TArray<FName> GetDirtyProfiles() const;
+
+	/**
+	 * Returns true if the setup data need compilation such that the controls etc need to be re-initialized.
+	 */
+	UFUNCTION(CallInEditor, Category = Actions)
+	bool IsSetupDirty() const;
+
+
 	/** Combines and returns data from our parent and ourself */
 	FPhysicsControlCharacterSetupData GetCharacterSetupData() const;
 
@@ -153,6 +190,9 @@ public:
 	/** Combines and returns data from our parent and ourself */
 	TMap<FName, FPhysicsControlControlAndModifierUpdates> GetProfiles() const;
 
+protected:
+	FOnControlAssetCompiled OnControlAssetCompiledDelegate;
+
 #endif
 
 public:
@@ -169,8 +209,4 @@ public:
 
 	static const FName GetPreviewMeshPropertyName();
 #endif
-
-
-
-
 };
diff --git a/Engine/Plugins/Experimental/PhysicsControl/Source/PhysicsControl/Public/PhysicsControlComponent.h b/Engine/Plugins/Experimental/PhysicsControl/Source/PhysicsControl/Public/PhysicsControlComponent.h
--- a/Engine/Plugins/Experimental/PhysicsControl/Source/PhysicsControl/Public/PhysicsControlComponent.h
+++ b/Engine/Plugins/Experimental/PhysicsControl/Source/PhysicsControl/Public/PhysicsControlComponent.h
@@ -123,7 +123,7 @@
 	 *        SetControlEnabled(true) in order to enable it.
 	 * @return True if a new control was created, false if a control of the specified name already exists
 	 */
-	UFUNCTION(BlueprintCallable, Category = PhysicsControl)
+	UFUNCTION(BlueprintCallable, Category = PhysicsControl, Meta = (ReturnDisplayName = "Success"))
 	bool CreateNamedControl(
 		FName                         Name,
 		UPrimitiveComponent*          ParentComponent,
@@ -310,12 +310,18 @@
 		const bool                                   bEnabled = true);
 
 	/**
+	 * Destroys all controls and body modifiers
+	 */
+	UFUNCTION(BlueprintCallable, Category = PhysicsControl)
+	void DestroyAllControlsAndBodyModifiers();
+
+	/**
 	 * Destroys a control
 	 *
 	 * @param Name The name of the control to destroy. 
 	 * @return     Returns true if the control was found and destroyed, false if not
 	 */
-	UFUNCTION(BlueprintCallable, Category = PhysicsControl)
+	UFUNCTION(BlueprintCallable, Category = PhysicsControl, Meta = (ReturnDisplayName = "Success"))
 	bool DestroyControl(const FName Name);
 
 	/**
@@ -340,7 +346,7 @@
 	 * Updates the parent object part of a control. Note that this won't change the name of the control (which may
 	 * subsequently be misleading), or any set it is included in, etc.
 	 */
-	UFUNCTION(BlueprintCallable, Category = PhysicsControl)
+	UFUNCTION(BlueprintCallable, Category = PhysicsControl, Meta = (ReturnDisplayName = "Success"))
 	bool SetControlParent(
 		const FName          Name, 
 		UPrimitiveComponent* ParentComponent,
@@ -369,7 +375,7 @@
 	 * @param ControlData The new control data
 	 * @return true if the control was found and modified, false if not
 	 */
-	UFUNCTION(BlueprintCallable, Category = PhysicsControl)
+	UFUNCTION(BlueprintCallable, Category = PhysicsControl, Meta = (ReturnDisplayName = "Success"))
 	bool SetControlData(const FName Name, const FPhysicsControlData ControlData);
 
 	/**
@@ -399,7 +405,7 @@
 	 * @param ControlData The new control data
 	 * @return true if the control was found and modified, false if not
 	 */
-	UFUNCTION(BlueprintCallable, Category = PhysicsControl)
+	UFUNCTION(BlueprintCallable, Category = PhysicsControl, Meta = (ReturnDisplayName = "Success"))
 	bool SetControlSparseData(const FName Name, const FPhysicsControlSparseData ControlData);
 
 	/**
@@ -430,7 +436,7 @@
 	 * @param bEnableControl Enables the control if it is currently disabled
 	 * @return true if the control was found and modified, false if not
 	 */
-	UFUNCTION(BlueprintCallable, Category = PhysicsControl)
+	UFUNCTION(BlueprintCallable, Category = PhysicsControl, Meta = (ReturnDisplayName = "Success"))
 	bool SetControlMultiplier
 	(const FName Name, const FPhysicsControlMultiplier ControlMultiplier, const bool bEnableControl = true);
 
@@ -470,7 +476,7 @@
 	 * @param bEnableControl Enables the control if it is currently disabled
 	 * @return true if the control was found and modified, false if not
 	 */
-	UFUNCTION(BlueprintCallable, Category = PhysicsControl)
+	UFUNCTION(BlueprintCallable, Category = PhysicsControl, Meta = (ReturnDisplayName = "Success"))
 	bool SetControlSparseMultiplier(
 		const FName Name, const FPhysicsControlSparseMultiplier ControlMultiplier, const bool bEnableControl = true);
 
@@ -513,7 +519,7 @@
 	 * @param bEnableControl Enables the control if it is currently disabled
 	 * @return true if the control was found and modified, false if not
 	 */
-	UFUNCTION(BlueprintCallable, Category = PhysicsControl)
+	UFUNCTION(BlueprintCallable, Category = PhysicsControl, Meta = (ReturnDisplayName = "Success"))
 	bool SetControlLinearData(
 		const FName Name,
 		const float Strength = 1.0f, 
@@ -533,7 +539,7 @@
 	 * @param bEnableControl Enables the control if it is currently disabled
 	 * @return true if the control was found and modified, false if not
 	 */
-	UFUNCTION(BlueprintCallable, Category = PhysicsControl)
+	UFUNCTION(BlueprintCallable, Category = PhysicsControl, Meta = (ReturnDisplayName = "Success"))
 	bool SetControlAngularData(
 		const FName Name,
 		const float Strength = 1.0f, 
@@ -549,7 +555,7 @@
 	 * @param Position The position of the control point on the child mesh object (only relevant if that 
 	 *        object is in use and is being simulated)
 	 */
-	UFUNCTION(BlueprintCallable, Category = PhysicsControl)
+	UFUNCTION(BlueprintCallable, Category = PhysicsControl, Meta = (ReturnDisplayName = "Success"))
 	bool SetControlPoint(const FName Name, const FVector Position);
 
 	/**
@@ -557,7 +563,7 @@
 	 *
 	 * @param Name The name of the control to modify. 
 	 */
-	UFUNCTION(BlueprintCallable, Category = PhysicsControl)
+	UFUNCTION(BlueprintCallable, Category = PhysicsControl, Meta = (ReturnDisplayName = "Success"))
 	bool ResetControlPoint(const FName Name);
 
 	/**
@@ -568,7 +574,7 @@
 	 * @param bEnableControl Enables the control if it is currently disabled
 	 * @return true if the control was found and modified, false if not
 	 */
-	UFUNCTION(BlueprintCallable, Category = PhysicsControl)
+	UFUNCTION(BlueprintCallable, Category = PhysicsControl, Meta = (ReturnDisplayName = "Success"))
 	bool SetControlTarget(const FName Name, const FPhysicsControlTarget ControlTarget, const bool bEnableControl = true);
 
 	/**
@@ -614,9 +620,9 @@
 	 *        want to specify the target transform for the object as a whole. If false, then the 
 	 *        target transform is used as is, and the system drives the control point towards this
 	 *        transform.
-	 * @return                  Returns true if the control was found and modified, false if not
+	 * @return Returns true if the control was found and modified, false if not
 	 */
-	UFUNCTION(BlueprintCallable, Category = PhysicsControl)
+	UFUNCTION(BlueprintCallable, Category = PhysicsControl, Meta = (ReturnDisplayName = "Success"))
 	bool SetControlTargetPositionAndOrientation(
 		const FName    Name,
 		const FVector  Position,
@@ -666,7 +672,7 @@
 	 *        transform.
 	 * @return true if the control was found and modified, false if not
 	 */
-	UFUNCTION(BlueprintCallable, Category = PhysicsControl)
+	UFUNCTION(BlueprintCallable, Category = PhysicsControl, Meta = (ReturnDisplayName = "Success"))
 	bool SetControlTargetPosition(
 		const FName   Name,
 		const FVector Position, 
@@ -713,7 +719,7 @@
 	 *        transform.
 	 * @return true if the control was found and modified, false if not
 	 */
-	UFUNCTION(BlueprintCallable, Category = PhysicsControl)
+	UFUNCTION(BlueprintCallable, Category = PhysicsControl, Meta = (ReturnDisplayName = "Success"))
 	bool SetControlTargetOrientation(
 		const FName    Name,
 		const FRotator Orientation, 
@@ -748,7 +754,7 @@
 	 * in size.
 	 * @return true if the control/position arrays match, false if they don't.
 	 */
-	UFUNCTION(BlueprintCallable, Category = PhysicsControl)
+	UFUNCTION(BlueprintCallable, Category = PhysicsControl, Meta = (ReturnDisplayName = "Success"))
 	bool SetControlTargetPositionsFromArray(
 		const TArray<FName>&   Names,
 		const TArray<FVector>& Positions, 
@@ -761,7 +767,7 @@
 	 * in size.
 	 * @return true if the control/position arrays match, false if they don't.
 	 */
-	UFUNCTION(BlueprintCallable, Category = PhysicsControl)
+	UFUNCTION(BlueprintCallable, Category = PhysicsControl, Meta = (ReturnDisplayName = "Success"))
 	bool SetControlTargetOrientationsFromArray(
 		const TArray<FName>&    Names,
 		const TArray<FRotator>& Orientations, 
@@ -774,7 +780,7 @@
 	 * These array should match in size.
 	 * @return true if the control/position/orientation arrays match, false if they don't.
 	 */
-	UFUNCTION(BlueprintCallable, Category = PhysicsControl)
+	UFUNCTION(BlueprintCallable, Category = PhysicsControl, Meta = (ReturnDisplayName = "Success"))
 	bool SetControlTargetPositionsAndOrientationsFromArray(
 		const TArray<FName>&    Names,
 		const TArray<FVector>&  Positions,
@@ -797,7 +803,7 @@
 	 * @param bEnableControl    Enables the control if it is currently disabled
 	 * @return                  Returns true if the control was found and modified, false if not
 	 */
-	UFUNCTION(BlueprintCallable, Category = PhysicsControl)
+	UFUNCTION(BlueprintCallable, Category = PhysicsControl, Meta = (ReturnDisplayName = "Success"))
 	bool SetControlTargetPoses(
 		const FName    Name,
 		const FVector  ParentPosition, 
@@ -818,7 +824,7 @@
 	 *        the amount of velocity extracted from the animation that is used as targets for the controls
 	 * @return true if the control was found and modified, false if not
 	 */
-	UFUNCTION(BlueprintCallable, Category = PhysicsControl)
+	UFUNCTION(BlueprintCallable, Category = PhysicsControl, Meta = (ReturnDisplayName = "Success"))
 	bool SetControlUseSkeletalAnimation(
 		const FName Name,
 		const bool  bUseSkeletalAnimation = true,
@@ -863,7 +869,7 @@
 	 * @param bEnable  Whether to enable/disable the control
 	 * @return         Returns true if the control was found and modified, false if not
 	 */
-	UFUNCTION(BlueprintCallable, Category = PhysicsControl)
+	UFUNCTION(BlueprintCallable, Category = PhysicsControl, Meta = (ReturnDisplayName = "Success"))
 	bool SetControlEnabled(const FName Name, const bool bEnable = true);
 
 	/**
@@ -891,7 +897,7 @@
 	 * @param bDisableCollision If set then the control will disable collision between the bodies it connects.
 	 * @return true if the control was found and modified, false if not
 	 */
-	UFUNCTION(BlueprintCallable, Category = PhysicsControl)
+	UFUNCTION(BlueprintCallable, Category = PhysicsControl, Meta = (ReturnDisplayName = "Success"))
 	bool SetControlDisableCollision(const FName Name, const bool bDisableCollision);
 
 	/**
@@ -914,7 +920,7 @@
 	 * @param Control  The control data that will be filled in if found
 	 * @return         Returns true if the control was found, false if not
 	 */
-	UFUNCTION(BlueprintCallable, Category = PhysicsControl)
+	UFUNCTION(BlueprintCallable, Category = PhysicsControl, Meta = (ReturnDisplayName = "Success"))
 	bool GetControlData(const FName Name, FPhysicsControlData& ControlData) const;
 
 	/**
@@ -922,7 +928,7 @@
 	 * @param Control  The control multipliers that will be filled in if found
 	 * @return         Returns true if the control was found, false if not
 	 */
-	UFUNCTION(BlueprintCallable, Category = PhysicsControl)
+	UFUNCTION(BlueprintCallable, Category = PhysicsControl, Meta = (ReturnDisplayName = "Success"))
 	bool GetControlMultiplier(const FName Name, FPhysicsControlMultiplier& ControlMultiplier) const;
 
 	/**
@@ -930,14 +936,14 @@
 	 * @param Control  The control target, if found
 	 * @return         Returns true if the control was found, false if not
 	 */
-	UFUNCTION(BlueprintCallable, Category = PhysicsControl)
+	UFUNCTION(BlueprintCallable, Category = PhysicsControl, Meta = (ReturnDisplayName = "Success"))
 	bool GetControlTarget(const FName Name, FPhysicsControlTarget& ControlTarget) const;
 
 	/**
 	 * @param Name        The name of the control to access. 
 	 * @return            Returns true if the control is enabled
 	 */
-	UFUNCTION(BlueprintCallable, Category = PhysicsControl)
+	UFUNCTION(BlueprintCallable, Category = PhysicsControl, Meta = (ReturnDisplayName = "Success"))
 	bool GetControlEnabled(const FName Name) const;
 
 	/**
@@ -966,7 +972,7 @@
 	 *        gets added to the set "All"
 	 * @param BodyModifierData The initial properties of the modifier
 	 */
-	UFUNCTION(BlueprintCallable, Category = PhysicsControl)
+	UFUNCTION(BlueprintCallable, Category = PhysicsControl, Meta = (ReturnDisplayName = "Success"))
 	bool CreateNamedBodyModifier(
 		const FName                       Name,
 		UPrimitiveComponent*              Component,
@@ -1015,7 +1021,7 @@
 	 * @param Name        The name of the body modifier to destroy. 
 	 * @return            Returns true if the body modifier was found and destroyed, false if not
 	 */
-	UFUNCTION(BlueprintCallable, Category = PhysicsControl)
+	UFUNCTION(BlueprintCallable, Category = PhysicsControl, Meta = (ReturnDisplayName = "Success"))
 	bool DestroyBodyModifier(const FName Name);
 
 	/**
@@ -1043,7 +1049,7 @@
 	 * @param ModifierData The new data
 	 * @return true if the modifier was found and modified, false if not
 	 */
-	UFUNCTION(BlueprintCallable, Category = PhysicsControl)
+	UFUNCTION(BlueprintCallable, Category = PhysicsControl, Meta = (ReturnDisplayName = "Success"))
 	bool SetBodyModifierData(const FName Name, const FPhysicsControlModifierData ModifierData);
 
 	/**
@@ -1073,7 +1079,7 @@
 	 * @param ModifierData The new data
 	 * @return true if the modifier was found and modified, false if not
 	 */
-	UFUNCTION(BlueprintCallable, Category = PhysicsControl)
+	UFUNCTION(BlueprintCallable, Category = PhysicsControl, Meta = (ReturnDisplayName = "Success"))
 	bool SetBodyModifierSparseData(const FName Name, const FPhysicsControlModifierSparseData ModifierData);
 
 	/**
@@ -1107,7 +1113,7 @@
 	 * @param bMakeKinematic If set then the body will be made kinematic. If not set, then it won't be changed.
 	 * @return true if the body modifier was found, false if not
 	 */
-	UFUNCTION(BlueprintCallable, Category = PhysicsControl)
+	UFUNCTION(BlueprintCallable, Category = PhysicsControl, Meta = (ReturnDisplayName = "Success"))
 	bool SetBodyModifierKinematicTarget(
 		const FName    Name,
 		const FVector  KinematicTargetPosition, 
@@ -1121,7 +1127,7 @@
 	 * @param MovementType Whether to enable/disable simulation on the body
 	 * @return true if the body modifier was found, false if not
 	 */
-	UFUNCTION(BlueprintCallable, Category = PhysicsControl)
+	UFUNCTION(BlueprintCallable, Category = PhysicsControl, Meta = (ReturnDisplayName = "Success"))
 	bool SetBodyModifierMovementType(
 		const FName                Name,
 		const EPhysicsMovementType MovementType = EPhysicsMovementType::Simulated);
@@ -1157,7 +1163,7 @@
 	 * @param CollisionType Collision type to set on the body
 	 * @return true if the body modifier was found, false if not
 	 */
-	UFUNCTION(BlueprintCallable, Category = PhysicsControl)
+	UFUNCTION(BlueprintCallable, Category = PhysicsControl, Meta = (ReturnDisplayName = "Success"))
 	bool SetBodyModifierCollisionType(
 		const FName                   Name,
 		const ECollisionEnabled::Type CollisionType = ECollisionEnabled::QueryAndPhysics);
@@ -1193,7 +1199,7 @@
 	 * @param GravityMultiplier The amount of gravity to apply when simulating
 	 * @return true if the body modifier was found, false if not
 	 */
-	UFUNCTION(BlueprintCallable, Category = PhysicsControl)
+	UFUNCTION(BlueprintCallable, Category = PhysicsControl, Meta = (ReturnDisplayName = "Success"))
 	bool SetBodyModifierGravityMultiplier(
 		const FName Name,
 		const float GravityMultiplier = 1.0f);
@@ -1230,7 +1236,7 @@
 	 *        animation and that coming from simulation.
 	 * @return true if the body modifier was found, false if not
 	 */
-	UFUNCTION(BlueprintCallable, Category = PhysicsControl)
+	UFUNCTION(BlueprintCallable, Category = PhysicsControl, Meta = (ReturnDisplayName = "Success"))
 	bool SetBodyModifierPhysicsBlendWeight(
 		const FName Name,
 		const float PhysicsBlendWeight = 1.0f);
@@ -1269,7 +1275,7 @@
 	 *        animation, rather than world space. Only relevant if the body is part of a skeletal mesh.
 	 * @return true if the body modifier was found, false if not
 	 */
-	UFUNCTION(BlueprintCallable, Category = PhysicsControl)
+	UFUNCTION(BlueprintCallable, Category = PhysicsControl, Meta = (ReturnDisplayName = "Success"))
 	bool SetBodyModifierUseSkeletalAnimation(
 		const FName Name,
 		const bool  bUseSkeletalAnimation);
@@ -1312,7 +1318,7 @@
 	 *        dynamic ones.
 	 * @return true if the body modifier was found, false if not
 	 */
-	UFUNCTION(BlueprintCallable, Category = PhysicsControl)
+	UFUNCTION(BlueprintCallable, Category = PhysicsControl, Meta = (ReturnDisplayName = "Success"))
 	bool SetBodyModifierUpdateKinematicFromSimulation(
 		const FName Name,
 		const bool  bUpdateKinematicFromSimulation);
@@ -1362,8 +1368,8 @@
 	 * It is also possible to specify a mesh component to use for the "world" object - so that the world controls can 
 	 * be made to work in the space of another object (or a bone if that is a skeletal mesh component)
 	 */
-	UFUNCTION(BlueprintCallable, Category = PhysicsControl)
-	void CreateControlsAndBodyModifiersFromLimbBones(
+	UFUNCTION(BlueprintCallable, Category = PhysicsControl, Meta = (ReturnDisplayName = "Success"))
+	bool CreateControlsAndBodyModifiersFromLimbBones(
 		FPhysicsControlNames&                       AllWorldSpaceControls,
 		TMap<FName, FPhysicsControlNames>&          LimbWorldSpaceControls,
 		FPhysicsControlNames&                       AllParentSpaceControls,
@@ -1379,22 +1385,24 @@
 		FName                                       WorldBoneName = NAME_None);
 
 	/**
-	 * This uses the control profile asset (that should have already been set) to create
-	 * controls and body modifiers
+	 * This uses the control profile asset (that should have already been assigned in our data) to create
+	 * controls and body modifiers.
 	 * It is also possible to specify a mesh component to use for the "world" object - so that the world controls can
 	 * be made to work in the space of another object (or a bone if that is a skeletal mesh component)
 	 */
-	UFUNCTION(BlueprintCallable, Category = PhysicsControl)
-	void CreateControlsAndBodyModifiersFromControlProfileAsset(
+	UFUNCTION(BlueprintCallable, Category = PhysicsControl, Meta = (ReturnDisplayName = "Success"))
+	bool CreateControlsAndBodyModifiersFromPhysicsControlAsset(
 		USkeletalMeshComponent* SkeletalMeshComponent,
 		UPrimitiveComponent*    WorldComponent,
 		FName                   WorldBoneName);
 
 	/**
 	 * Looks up the profile which should exist in the registered control profile asset, and invokes it.
+	 * 
+	 * @return true if successful, and false if the profile cannot be found.
 	 */
-	UFUNCTION(BlueprintCallable, Category = PhysicsControl)
-	void InvokeControlProfile(FName ProfileName);
+	UFUNCTION(BlueprintCallable, Category = PhysicsControl, Meta = (ReturnDisplayName = "Success"))
+	bool InvokeControlProfile(FName ProfileName);
 
 	/**
 	 * Adds a Control to a Set. This will add a new set if necessary. For example, you might
@@ -1552,7 +1560,7 @@
 	 * 
 	 * @return true if successful, and false if no cached target can be found for the bone.
 	 */
-	UFUNCTION(BlueprintCallable, Category = PhysicsControl)
+	UFUNCTION(BlueprintCallable, Category = PhysicsControl, Meta = (ReturnDisplayName = "Success"))
 	bool SetCachedBoneData(
 		const USkeletalMeshComponent* SkeletalMeshComponent,
 		const FName                   BoneName, 
@@ -1574,7 +1582,7 @@
 	 * 
 	 * @return true if the body modifier is found (even if no cached target is found), and false otherwise.
 	 */
-	UFUNCTION(BlueprintCallable, Category = PhysicsControl)
+	UFUNCTION(BlueprintCallable, Category = PhysicsControl, Meta = (ReturnDisplayName = "Success"))
 	bool ResetBodyModifierToCachedBoneTransform(
 		const FName                        Name,
 		const EResetToCachedTargetBehavior Behavior = EResetToCachedTargetBehavior::ResetImmediately);
@@ -1596,11 +1604,11 @@
 		const EResetToCachedTargetBehavior Behavior = EResetToCachedTargetBehavior::ResetImmediately);
 
 	/** Indicates if a control with the name exists (doesn't produce a warning if it doesn't) */
-	UFUNCTION(BlueprintCallable, Category = PhysicsControl)
+	UFUNCTION(BlueprintCallable, Category = PhysicsControl, Meta = (ReturnDisplayName = "Success"))
 	bool GetControlExists(const FName Name) const;
 
 	/** Indicates if a body modifier with the name exists (doesn't produce a warning if it doesn't) */
-	UFUNCTION(BlueprintCallable, Category = PhysicsControl)
+	UFUNCTION(BlueprintCallable, Category = PhysicsControl, Meta = (ReturnDisplayName = "Success"))
 	bool GetBodyModifierExists(const FName Name) const;
 
 public:
diff --git a/Engine/Plugins/Experimental/PhysicsControl/Source/PhysicsControl/Public/PhysicsControlNameRecords.h b/Engine/Plugins/Experimental/PhysicsControl/Source/PhysicsControl/Public/PhysicsControlNameRecords.h
--- a/Engine/Plugins/Experimental/PhysicsControl/Source/PhysicsControl/Public/PhysicsControlNameRecords.h
+++ b/Engine/Plugins/Experimental/PhysicsControl/Source/PhysicsControl/Public/PhysicsControlNameRecords.h
@@ -52,4 +52,4 @@ struct FPhysicsControlNameRecords
 
 /** Returns a copy of the supplied list of names in which any names that match keys in the supplied 'Sets' map are replaced with all the names in the array associated with that key. */
 PHYSICSCONTROL_API TArray<FName> ExpandName(const FName InName, const TMap<FName, TArray<FName>>& Sets);
-PHYSICSCONTROL_API TArray<FName> ExpandName(const TArray<FName>& InNames, const TMap<FName, TArray<FName>>& Sets);
+PHYSICSCONTROL_API TArray<FName> ExpandNames(const TArray<FName>& InNames, const TMap<FName, TArray<FName>>& Sets);
diff --git a/Engine/Plugins/Experimental/PhysicsControl/Source/PhysicsControlEditor/Private/PhysicsControlAssetEditor.cpp b/Engine/Plugins/Experimental/PhysicsControl/Source/PhysicsControlEditor/Private/PhysicsControlAssetEditor.cpp
--- a/Engine/Plugins/Experimental/PhysicsControl/Source/PhysicsControlEditor/Private/PhysicsControlAssetEditor.cpp
+++ b/Engine/Plugins/Experimental/PhysicsControl/Source/PhysicsControlEditor/Private/PhysicsControlAssetEditor.cpp
@@ -587,6 +587,12 @@ void FPhysicsControlAssetEditor::OnToggleSimulation()
 	EditorData->ToggleSimulation();
 }
 
+//======================================================================================================================
+void FPhysicsControlAssetEditor::RecreateControlsAndModifiers()
+{
+	EditorData->RecreateControlsAndModifiers();
+}
+
 //======================================================================================================================
 void FPhysicsControlAssetEditor::OnToggleSimulationNoGravity()
 {
@@ -1052,6 +1058,12 @@ void FPhysicsControlAssetEditor::AddReferencedObjects(FReferenceCollector& Colle
 	EditorData->AddReferencedObjects(Collector);
 }
 
+//======================================================================================================================
+void FPhysicsControlAssetEditor::Tick(float DeltaTime)
+{
+	GetPersonaToolkit()->GetPreviewScene()->InvalidateViews();
+}
+
 //======================================================================================================================
 TStatId FPhysicsControlAssetEditor::GetStatId() const
 {
@@ -1175,5 +1187,23 @@ void FPhysicsControlAssetEditor::RefreshPreviewViewport()
 	}
 }
 
+//======================================================================================================================
+void FPhysicsControlAssetEditor::InvokeControlProfile(FName ProfileName)
+{
+	// Danny TODO handle the RBWC simulation case
+	UPhysicsControlComponent* PCC = GetEditorData()->PhysicsControlComponent.Get();
+	if (PCC)
+	{
+		PCC->InvokeControlProfile(ProfileName);
+	}
+	PreviouslyInvokedControlProfile = ProfileName;
+}
+
+//======================================================================================================================
+void FPhysicsControlAssetEditor::ReinvokeControlProfile()
+{
+	InvokeControlProfile(PreviouslyInvokedControlProfile);
+}
+
 #undef LOCTEXT_NAMESPACE
 
diff --git a/Engine/Plugins/Experimental/PhysicsControl/Source/PhysicsControlEditor/Private/PhysicsControlAssetEditorData.cpp b/Engine/Plugins/Experimental/PhysicsControl/Source/PhysicsControlEditor/Private/PhysicsControlAssetEditorData.cpp
--- a/Engine/Plugins/Experimental/PhysicsControl/Source/PhysicsControlEditor/Private/PhysicsControlAssetEditorData.cpp
+++ b/Engine/Plugins/Experimental/PhysicsControl/Source/PhysicsControlEditor/Private/PhysicsControlAssetEditorData.cpp
@@ -144,7 +144,7 @@ void FPhysicsControlAssetEditorData::EnableSimulation(bool bEnableSimulation)
 			EditorSkelComp->WakeAllRigidBodies();
 
 			PhysicsControlComponent->PhysicsControlAsset = PhysicsControlAsset;
-			PhysicsControlComponent->CreateControlsAndBodyModifiersFromControlProfileAsset(
+			PhysicsControlComponent->CreateControlsAndBodyModifiersFromPhysicsControlAsset(
 				EditorSkelComp, nullptr, FName());
 		}
 		else
@@ -205,8 +205,7 @@ void FPhysicsControlAssetEditorData::EnableSimulation(bool bEnableSimulation)
 			EditorSkelComp->EnablePreview(true, PreviewAnimationAsset);
 		}
 
-		PhysicsControlComponent->DestroyControlsInSet("All");
-		PhysicsControlComponent->DestroyBodyModifiersInSet("All");
+		PhysicsControlComponent->DestroyAllControlsAndBodyModifiers();
 
 		BroadcastPreviewChanged();
 	}
@@ -214,6 +213,19 @@ void FPhysicsControlAssetEditorData::EnableSimulation(bool bEnableSimulation)
 	bRunningSimulation = bEnableSimulation;
 }
 
+//======================================================================================================================
+// Danny TODO handle the RBWC mode
+void FPhysicsControlAssetEditorData::RecreateControlsAndModifiers()
+{
+	// Turn it off...
+	PhysicsControlComponent->DestroyAllControlsAndBodyModifiers();
+
+	// ...and back on again
+	PhysicsControlComponent->PhysicsControlAsset = PhysicsControlAsset;
+	PhysicsControlComponent->CreateControlsAndBodyModifiersFromPhysicsControlAsset(
+		EditorSkelComp, nullptr, FName());
+}
+
 //======================================================================================================================
 EPhysicsAssetEditorMeshViewMode FPhysicsControlAssetEditorData::GetCurrentMeshViewMode(bool bSimulation)
 {
diff --git a/Engine/Plugins/Experimental/PhysicsControl/Source/PhysicsControlEditor/Private/PhysicsControlAssetEditorData.h b/Engine/Plugins/Experimental/PhysicsControl/Source/PhysicsControlEditor/Private/PhysicsControlAssetEditorData.h
--- a/Engine/Plugins/Experimental/PhysicsControl/Source/PhysicsControlEditor/Private/PhysicsControlAssetEditorData.h
+++ b/Engine/Plugins/Experimental/PhysicsControl/Source/PhysicsControlEditor/Private/PhysicsControlAssetEditorData.h
@@ -68,6 +68,9 @@ public:
 	/** Toggle simulation on and off */
 	void ToggleSimulation();
 
+	/** Destroys all existing controls modifiers and then recreates them from the control asset */
+	void RecreateControlsAndModifiers();
+
 	/** broadcast a change in the preview*/
 	void BroadcastPreviewChanged();
 
diff --git a/Engine/Plugins/Experimental/PhysicsControl/Source/PhysicsControlEditor/Private/PhysicsControlAssetPreviewDetailsCustomization.cpp b/Engine/Plugins/Experimental/PhysicsControl/Source/PhysicsControlEditor/Private/PhysicsControlAssetPreviewDetailsCustomization.cpp
--- a/Engine/Plugins/Experimental/PhysicsControl/Source/PhysicsControlEditor/Private/PhysicsControlAssetPreviewDetailsCustomization.cpp
+++ b/Engine/Plugins/Experimental/PhysicsControl/Source/PhysicsControlEditor/Private/PhysicsControlAssetPreviewDetailsCustomization.cpp
@@ -12,20 +12,10 @@
 #include "DetailCategoryBuilder.h"
 #include "DetailLayoutBuilder.h"
 #include "DetailWidgetRow.h"
-#include "EditorFontGlyphs.h"
 #include "Framework/Application/SlateApplication.h"
-#include "Framework/MultiBox/MultiBoxBuilder.h"
 #include "IDetailPropertyRow.h"
 #include "PropertyHandle.h"
-#include "ScopedTransaction.h"
-#include "Styling/StyleColors.h"
-#include "Widgets/Images/SImage.h"
 #include "Widgets/Input/SButton.h"
-#include "Widgets/Input/SComboButton.h"
-#include "Widgets/Input/SEditableTextBox.h"
-#include "Widgets/Layout/SUniformGridPanel.h"
-#include "Widgets/SCompoundWidget.h"
-#include "Widgets/Text/SInlineEditableTextBlock.h"
 
 #define LOCTEXT_NAMESPACE "PhysicsControlAssetPreviewDetailsCustomization"
 
@@ -37,15 +27,39 @@ TSharedRef<IDetailCustomization> FPhysicsControlAssetPreviewDetailsCustomization
 }
 
 //======================================================================================================================
-void FPhysicsControlAssetPreviewDetailsCustomization::CustomizeDetails(IDetailLayoutBuilder& DetailLayoutBuilder)
+FPhysicsControlAssetPreviewDetailsCustomization::FPhysicsControlAssetPreviewDetailsCustomization(
+	TWeakPtr<FPhysicsControlAssetEditor> InPhysicsControlAssetEditor)
+	: PhysicsControlAssetEditor(InPhysicsControlAssetEditor)
+{
+}
+
+//======================================================================================================================
+void FPhysicsControlAssetPreviewDetailsCustomization::CustomizeDetails(
+	const TSharedPtr<IDetailLayoutBuilder>& InDetailBuilder)
 {
-	BindCommands();
+	DetailLayoutBuilderWeak = InDetailBuilder;
+	FPhysicsControlAssetPreviewDetailsCustomization::CustomizeDetails(*InDetailBuilder);
 
+	if (TSharedPtr<FPhysicsControlAssetEditor> PCAE = PhysicsControlAssetEditor.Pin())
+	{
+		if (UPhysicsControlAsset* PhysicsControlAsset = PCAE->GetEditorData()->PhysicsControlAsset.Get())
+		{
+			PhysicsControlAsset->OnControlAssetCompiled().AddSP(
+				this, &FPhysicsControlAssetPreviewDetailsCustomization::OnControlAssetCompiled);
+		}
+	}
+}
+
+//======================================================================================================================
+void FPhysicsControlAssetPreviewDetailsCustomization::CustomizeDetails(IDetailLayoutBuilder& DetailLayoutBuilder)
+{
 	DetailLayoutBuilder.HideCategory(TEXT("PreviewMesh"));
 	DetailLayoutBuilder.HideCategory(TEXT("Actions"));
 	DetailLayoutBuilder.HideCategory(TEXT("Inheritance"));
 	DetailLayoutBuilder.HideCategory(TEXT("Setup"));
 	DetailLayoutBuilder.HideCategory(TEXT("Profiles"));
+	DetailLayoutBuilder.HideCategory(TEXT("ProfileEditing"));
+	DetailLayoutBuilder.HideCategory(TEXT("SetupEditing"));
 
 	TSharedPtr<FPhysicsControlAssetEditorData> EditorData = PhysicsControlAssetEditor.Pin()->GetEditorData();
 	UPhysicsControlAsset* PCA = EditorData->PhysicsControlAsset.Get();
@@ -56,44 +70,40 @@ void FPhysicsControlAssetPreviewDetailsCustomization::CustomizeDetails(IDetailLa
 		for (const TPair<FName, FPhysicsControlControlAndModifierUpdates>& ProfilePair : PCA->Profiles)
 		{
 			const FName ProfileName = ProfilePair.Key;
-			DetailCategoryBuilder.AddCustomRow(
-				FText::FromName(ProfileName))
-				.NameContent()
-				[
-					SNew(STextBlock)
-						.Font(DetailLayoutBuilder.GetDetailFont())
-						.Text(FText::FromName(ProfileName))
-				]
-				.ValueContent()
+			FDetailWidgetRow& Row = DetailCategoryBuilder.AddCustomRow(FText::FromName(ProfileName));
+			Row.WholeRowContent()
 				[
 					SNew(SButton)
-						.ButtonStyle(FAppStyle::Get(), "SimpleButton")
-						.ContentPadding(FMargin(6, 2))
-						.Text(LOCTEXT("Invoke", "Invoke"))
+						.Text(FText::FromName(ProfileName))
 						.OnClicked(this, &FPhysicsControlAssetPreviewDetailsCustomization::InvokeControlProfile, ProfileName)
 				];
 		}
 	}
-
 }
 
 //======================================================================================================================
-// TODO Is this needed?
-void FPhysicsControlAssetPreviewDetailsCustomization::BindCommands()
+FReply FPhysicsControlAssetPreviewDetailsCustomization::InvokeControlProfile(FName ProfileName)
 {
+	if (TSharedPtr<FPhysicsControlAssetEditor> PCAE = PhysicsControlAssetEditor.Pin())
+	{
+		PCAE->InvokeControlProfile(ProfileName);
+	}
+	return FReply::Handled();
 }
 
 //======================================================================================================================
-FReply FPhysicsControlAssetPreviewDetailsCustomization::InvokeControlProfile(FName ProfileName)
+void FPhysicsControlAssetPreviewDetailsCustomization::OnControlAssetCompiled(bool bProfileListChanged)
 {
-	TSharedPtr<FPhysicsControlAssetEditorData> EditorData = PhysicsControlAssetEditor.Pin()->GetEditorData();
-	UPhysicsControlComponent* PCC = EditorData->PhysicsControlComponent.Get();
-	if (PCC)
+	if (bProfileListChanged)
 	{
-		PCC->InvokeControlProfile(ProfileName);
+		if (DetailLayoutBuilderWeak.IsValid())
+		{
+			if (IDetailLayoutBuilder* DetailLayoutBuilder = DetailLayoutBuilderWeak.Pin().Get())
+			{
+				DetailLayoutBuilder->ForceRefreshDetails();
+			}
+		}
 	}
-	return FReply::Handled();
 }
 
-
 #undef LOCTEXT_NAMESPACE
diff --git a/Engine/Plugins/Experimental/PhysicsControl/Source/PhysicsControlEditor/Private/PhysicsControlAssetPreviewDetailsCustomization.h b/Engine/Plugins/Experimental/PhysicsControl/Source/PhysicsControlEditor/Private/PhysicsControlAssetPreviewDetailsCustomization.h
--- a/Engine/Plugins/Experimental/PhysicsControl/Source/PhysicsControlEditor/Private/PhysicsControlAssetPreviewDetailsCustomization.h
+++ b/Engine/Plugins/Experimental/PhysicsControl/Source/PhysicsControlEditor/Private/PhysicsControlAssetPreviewDetailsCustomization.h
@@ -21,18 +21,20 @@ public:
 	// Makes a new instance of this detail layout class for a specific detail view requesting it
 	static TSharedRef<IDetailCustomization> MakeInstance(TWeakPtr<FPhysicsControlAssetEditor> InPhysicsControlAssetEditor);
 
-	FPhysicsControlAssetPreviewDetailsCustomization(TWeakPtr<FPhysicsControlAssetEditor> InPhysicsControlAssetEditor)
-		: PhysicsControlAssetEditor(InPhysicsControlAssetEditor)
-	{}
+	FPhysicsControlAssetPreviewDetailsCustomization(TWeakPtr<FPhysicsControlAssetEditor> InPhysicsControlAssetEditor);
 
-	// IDetailCustomization interface
+	//~ Begin IDetailCustomization
 	virtual void CustomizeDetails(IDetailLayoutBuilder& DetailLayout) override;
+	virtual void CustomizeDetails(const TSharedPtr<IDetailLayoutBuilder>& InDetailBuilder) override;
+	//~ End IDetailCustomization
 
 private:
-	void BindCommands();
-
 	FReply InvokeControlProfile(FName ProfileName);
 
+	void OnControlAssetCompiled(bool bProfileListChanged);
+
 private:
 	TWeakPtr<FPhysicsControlAssetEditor> PhysicsControlAssetEditor;
+
+	TWeakPtr<IDetailLayoutBuilder> DetailLayoutBuilderWeak;
 };
diff --git a/Engine/Plugins/Experimental/PhysicsControl/Source/PhysicsControlEditor/Private/PhysicsControlAssetProfileDetailsCustomization.cpp b/Engine/Plugins/Experimental/PhysicsControl/Source/PhysicsControlEditor/Private/PhysicsControlAssetProfileDetailsCustomization.cpp
--- a/Engine/Plugins/Experimental/PhysicsControl/Source/PhysicsControlEditor/Private/PhysicsControlAssetProfileDetailsCustomization.cpp
+++ b/Engine/Plugins/Experimental/PhysicsControl/Source/PhysicsControlEditor/Private/PhysicsControlAssetProfileDetailsCustomization.cpp
@@ -2,27 +2,13 @@
 
 #include "PhysicsControlAssetProfileDetailsCustomization.h"
 
-#include "PhysicsControlAssetActions.h"
 #include "PhysicsControlAssetEditor.h"
-#include "PhysicsControlAssetEditorCommands.h"
+#include "PhysicsControlAsset.h"
+#include "PhysicsControlComponent.h"
+#include "PhysicsControlAssetEditorData.h"
 
-#include "DetailCategoryBuilder.h"
 #include "DetailLayoutBuilder.h"
-#include "DetailWidgetRow.h"
-#include "EditorFontGlyphs.h"
-#include "Framework/Application/SlateApplication.h"
-#include "Framework/MultiBox/MultiBoxBuilder.h"
-#include "IDetailPropertyRow.h"
 #include "PropertyHandle.h"
-#include "ScopedTransaction.h"
-#include "Styling/StyleColors.h"
-#include "Widgets/Images/SImage.h"
-#include "Widgets/Input/SButton.h"
-#include "Widgets/Input/SComboButton.h"
-#include "Widgets/Input/SEditableTextBox.h"
-#include "Widgets/Layout/SUniformGridPanel.h"
-#include "Widgets/SCompoundWidget.h"
-#include "Widgets/Text/SInlineEditableTextBlock.h"
 
 #define LOCTEXT_NAMESPACE "PhysicsControlAssetProfileDetailsCustomization"
 
@@ -34,24 +20,84 @@ TSharedRef<IDetailCustomization> FPhysicsControlAssetProfileDetailsCustomization
 }
 
 //======================================================================================================================
-void FPhysicsControlAssetProfileDetailsCustomization::CustomizeDetails(IDetailLayoutBuilder& DetailLayout)
+void FPhysicsControlAssetProfileDetailsCustomization::CustomizeDetails(
+	const TSharedPtr<IDetailLayoutBuilder>& InDetailBuilder)
 {
-	BindCommands();
+	DetailLayoutBuilderWeak = InDetailBuilder;
+	FPhysicsControlAssetProfileDetailsCustomization::CustomizeDetails(*InDetailBuilder);
+}
+
+//======================================================================================================================
+void FPhysicsControlAssetProfileDetailsCustomization::CustomizeDetails(IDetailLayoutBuilder& DetailLayoutBuilder)
+{
+	DetailLayoutBuilder.HideCategory(TEXT("PreviewMesh"));
+	DetailLayoutBuilder.HideCategory(TEXT("Actions"));
+	DetailLayoutBuilder.HideCategory(TEXT("Inheritance"));
+	DetailLayoutBuilder.HideCategory(TEXT("Setup"));
+	DetailLayoutBuilder.HideCategory(TEXT("SetupEditing"));
 
-	DetailLayout.HideCategory(TEXT("PreviewMesh"));
-	DetailLayout.HideCategory(TEXT("Actions"));
-	DetailLayout.HideCategory(TEXT("Inheritance"));
-	DetailLayout.HideCategory(TEXT("Setup"));
+	const TSharedRef<IPropertyHandle> ProfilesProperty = DetailLayoutBuilder.GetProperty(
+		GET_MEMBER_NAME_CHECKED(UPhysicsControlAsset, MyProfiles));
 
+	ProfilesProperty->SetOnPropertyValueChanged(FSimpleDelegate::CreateSP(
+		this, &FPhysicsControlAssetProfileDetailsCustomization::OnProfilesChanged));
+
+	ProfilesProperty->SetOnChildPropertyValueChanged(FSimpleDelegate::CreateSP(
+			this, &FPhysicsControlAssetProfileDetailsCustomization::OnProfileDetailsChanged));
 }
 
 //======================================================================================================================
-void FPhysicsControlAssetProfileDetailsCustomization::BindCommands()
+// This is called when a parameter in one of the profiles changes
+void FPhysicsControlAssetProfileDetailsCustomization::OnProfileDetailsChanged()
 {
-	const FPhysicsControlAssetEditorCommands& Commands = FPhysicsControlAssetEditorCommands::Get();
+	if (TSharedPtr<FPhysicsControlAssetEditor> PCAE = PhysicsControlAssetEditor.Pin())
+	{
+		UPhysicsControlAsset* PhysicsControlAsset = PCAE->GetEditorData()->PhysicsControlAsset.Get();
+		if (PhysicsControlAsset->bAutoCompileProfiles)
+		{
+			TArray<FName> DirtyProfiles;
+			if (PhysicsControlAsset->bAutoInvokeProfiles && PCAE->IsRunningSimulation())
+			{
+				DirtyProfiles = PhysicsControlAsset->GetDirtyProfiles();
+			}
+
+			PhysicsControlAsset->Compile();
 
-	TSharedPtr<FUICommandList> CommandList = PhysicsControlAssetEditor.Pin()->GetToolkitCommands();
+			for (FName DirtyProfile : DirtyProfiles)
+			{
+				InvokeControlProfile(DirtyProfile);
+			}
+		}
+	}
+}
 
+//======================================================================================================================
+// This is called when the list of profiles changes (i.e. profile added/removed)
+void FPhysicsControlAssetProfileDetailsCustomization::OnProfilesChanged()
+{
+	if (TSharedPtr<FPhysicsControlAssetEditor> PCAE = PhysicsControlAssetEditor.Pin())
+	{
+		UPhysicsControlAsset* PhysicsControlAsset = PCAE->GetEditorData()->PhysicsControlAsset.Get();
+		if (PhysicsControlAsset->bAutoCompileProfiles)
+		{
+			PhysicsControlAsset->Compile();
+		}
+	}
+}
+
+//======================================================================================================================
+FReply FPhysicsControlAssetProfileDetailsCustomization::InvokeControlProfile(FName ProfileName)
+{
+	if (TSharedPtr<FPhysicsControlAssetEditor> PCAE = PhysicsControlAssetEditor.Pin())
+	{
+		// Danny TODO handle the RBWC mode
+		UPhysicsControlComponent* PCC = PCAE->GetEditorData()->PhysicsControlComponent.Get();
+		if (PCC)
+		{
+			PCC->InvokeControlProfile(ProfileName);
+		}
+	}
+	return FReply::Handled();
 }
 
 #undef LOCTEXT_NAMESPACE
diff --git a/Engine/Plugins/Experimental/PhysicsControl/Source/PhysicsControlEditor/Private/PhysicsControlAssetProfileDetailsCustomization.h b/Engine/Plugins/Experimental/PhysicsControl/Source/PhysicsControlEditor/Private/PhysicsControlAssetProfileDetailsCustomization.h
--- a/Engine/Plugins/Experimental/PhysicsControl/Source/PhysicsControlEditor/Private/PhysicsControlAssetProfileDetailsCustomization.h
+++ b/Engine/Plugins/Experimental/PhysicsControl/Source/PhysicsControlEditor/Private/PhysicsControlAssetProfileDetailsCustomization.h
@@ -5,8 +5,11 @@
 #include "CoreMinimal.h"
 #include "IDetailCustomization.h"
 #include "Types/SlateEnums.h"
+#include "Input/Reply.h"
 
+class UPhysicsControlAsset;
 class FPhysicsControlAssetEditor;
+
 class IDetailLayoutBuilder;
 class FUICommandInfo;
 class SWidget;
@@ -14,6 +17,8 @@ class IPropertyHandle;
 class FUICommandList;
 class SEditableTextBox;
 
+struct FPhysicsControlControlAndModifierUpdates;
+
 class FPhysicsControlAssetProfileDetailsCustomization : public IDetailCustomization
 {
 public:
@@ -24,12 +29,19 @@ public:
 		: PhysicsControlAssetEditor(InPhysicsControlAssetEditor)
 	{}
 
-	// IDetailCustomization interface
+	//~ Begin IDetailCustomization
 	virtual void CustomizeDetails(IDetailLayoutBuilder& DetailLayout) override;
+	virtual void CustomizeDetails(const TSharedPtr<IDetailLayoutBuilder>& InDetailBuilder) override;
+	//~ End IDetailCustomization
 
-private:
-	void BindCommands();
+protected:
+	void OnProfilesChanged();
+	void OnProfileDetailsChanged();
+
+protected:
+	FReply InvokeControlProfile(FName ProfileName);
 
-private:
 	TWeakPtr<FPhysicsControlAssetEditor> PhysicsControlAssetEditor;
+
+	TWeakPtr<IDetailLayoutBuilder> DetailLayoutBuilderWeak;
 };
diff --git a/Engine/Plugins/Experimental/PhysicsControl/Source/PhysicsControlEditor/Private/PhysicsControlAssetSetupDetailsCustomization.cpp b/Engine/Plugins/Experimental/PhysicsControl/Source/PhysicsControlEditor/Private/PhysicsControlAssetSetupDetailsCustomization.cpp
--- a/Engine/Plugins/Experimental/PhysicsControl/Source/PhysicsControlEditor/Private/PhysicsControlAssetSetupDetailsCustomization.cpp
+++ b/Engine/Plugins/Experimental/PhysicsControl/Source/PhysicsControlEditor/Private/PhysicsControlAssetSetupDetailsCustomization.cpp
@@ -4,25 +4,11 @@
 
 #include "PhysicsControlAssetActions.h"
 #include "PhysicsControlAssetEditor.h"
-#include "PhysicsControlAssetEditorCommands.h"
+#include "PhysicsControlAsset.h"
+#include "PhysicsControlAssetEditorData.h"
 
-#include "DetailCategoryBuilder.h"
 #include "DetailLayoutBuilder.h"
-#include "DetailWidgetRow.h"
-#include "EditorFontGlyphs.h"
-#include "Framework/Application/SlateApplication.h"
-#include "Framework/MultiBox/MultiBoxBuilder.h"
-#include "IDetailPropertyRow.h"
 #include "PropertyHandle.h"
-#include "ScopedTransaction.h"
-#include "Styling/StyleColors.h"
-#include "Widgets/Images/SImage.h"
-#include "Widgets/Input/SButton.h"
-#include "Widgets/Input/SComboButton.h"
-#include "Widgets/Input/SEditableTextBox.h"
-#include "Widgets/Layout/SUniformGridPanel.h"
-#include "Widgets/SCompoundWidget.h"
-#include "Widgets/Text/SInlineEditableTextBlock.h"
 
 #define LOCTEXT_NAMESPACE "PhysicsControlAssetSetupDetailsCustomization"
 
@@ -34,21 +20,68 @@ TSharedRef<IDetailCustomization> FPhysicsControlAssetSetupDetailsCustomization::
 }
 
 //======================================================================================================================
-void FPhysicsControlAssetSetupDetailsCustomization::CustomizeDetails(IDetailLayoutBuilder& DetailLayout)
+void FPhysicsControlAssetSetupDetailsCustomization::CustomizeDetails(IDetailLayoutBuilder& DetailLayoutBuilder)
 {
-	BindCommands();
+	DetailLayoutBuilder.HideCategory(TEXT("Profiles"));
+	DetailLayoutBuilder.HideCategory(TEXT("ProfileEditing"));
 
-	DetailLayout.HideCategory(TEXT("Profiles"));
+	TArray<TSharedRef<IPropertyHandle>> Properties;
+	Properties.Push(DetailLayoutBuilder.GetProperty(
+		GET_MEMBER_NAME_CHECKED(UPhysicsControlAsset, MyCharacterSetupData)));
 
+	Properties.Push(DetailLayoutBuilder.GetProperty(
+		GET_MEMBER_NAME_CHECKED(UPhysicsControlAsset, MyAdditionalControlsAndModifiers)));
+
+	Properties.Push(DetailLayoutBuilder.GetProperty(
+		GET_MEMBER_NAME_CHECKED(UPhysicsControlAsset, MyAdditionalSets)));
+
+	Properties.Push(DetailLayoutBuilder.GetProperty(
+		GET_MEMBER_NAME_CHECKED(UPhysicsControlAsset, MyInitialControlAndModifierUpdates)));
+
+	for (TSharedRef<IPropertyHandle> Property : Properties)
+	{
+		Property->SetOnPropertyValueChanged(FSimpleDelegate::CreateSP(
+			this, &FPhysicsControlAssetSetupDetailsCustomization::OnSetupChanged));
+
+		Property->SetOnChildPropertyValueChanged(FSimpleDelegate::CreateSP(
+			this, &FPhysicsControlAssetSetupDetailsCustomization::OnSetupDetailsChanged));
+	}
 }
 
 //======================================================================================================================
-void FPhysicsControlAssetSetupDetailsCustomization::BindCommands()
+// This is called when a parameter in one of the profiles changes
+void FPhysicsControlAssetSetupDetailsCustomization::OnSetupDetailsChanged()
 {
-	const FPhysicsControlAssetEditorCommands& Commands = FPhysicsControlAssetEditorCommands::Get();
+	if (TSharedPtr<FPhysicsControlAssetEditor> PCAE = PhysicsControlAssetEditor.Pin())
+	{
+		UPhysicsControlAsset* PhysicsControlAsset = PCAE->GetEditorData()->PhysicsControlAsset.Get();
 
-	TSharedPtr<FUICommandList> CommandList = PhysicsControlAssetEditor.Pin()->GetToolkitCommands();
+		bool bNeedToReinitialize = false;
+		if (PhysicsControlAsset->bAutoReinitSetup && PCAE->IsRunningSimulation())
+		{
+			bNeedToReinitialize = PhysicsControlAsset->IsSetupDirty();
+		}
 
+		if (PhysicsControlAsset->bAutoCompileSetup)
+		{
+			PhysicsControlAsset->Compile();
+		}
+
+		if (bNeedToReinitialize)
+		{
+			PCAE->RecreateControlsAndModifiers();
+			if (PhysicsControlAsset->bAutoInvokeProfileAfterSetup)
+			{
+				PCAE->ReinvokeControlProfile();
+			}
+		}
+	}
 }
 
+//======================================================================================================================
+// This is called when the list of profiles changes (i.e. profile added/removed)
+void FPhysicsControlAssetSetupDetailsCustomization::OnSetupChanged()
+{
+	OnSetupDetailsChanged();
+}
 #undef LOCTEXT_NAMESPACE
diff --git a/Engine/Plugins/Experimental/PhysicsControl/Source/PhysicsControlEditor/Private/PhysicsControlAssetSetupDetailsCustomization.h b/Engine/Plugins/Experimental/PhysicsControl/Source/PhysicsControlEditor/Private/PhysicsControlAssetSetupDetailsCustomization.h
--- a/Engine/Plugins/Experimental/PhysicsControl/Source/PhysicsControlEditor/Private/PhysicsControlAssetSetupDetailsCustomization.h
+++ b/Engine/Plugins/Experimental/PhysicsControl/Source/PhysicsControlEditor/Private/PhysicsControlAssetSetupDetailsCustomization.h
@@ -27,9 +27,10 @@ public:
 	// IDetailCustomization interface
 	virtual void CustomizeDetails(IDetailLayoutBuilder& DetailLayout) override;
 
-private:
-	void BindCommands();
+protected:
+	void OnSetupChanged();
+	void OnSetupDetailsChanged();
 
-private:
+protected:
 	TWeakPtr<FPhysicsControlAssetEditor> PhysicsControlAssetEditor;
 };
diff --git a/Engine/Plugins/Experimental/PhysicsControl/Source/PhysicsControlEditor/Public/PhysicsControlAssetEditor.h b/Engine/Plugins/Experimental/PhysicsControl/Source/PhysicsControlEditor/Public/PhysicsControlAssetEditor.h
--- a/Engine/Plugins/Experimental/PhysicsControl/Source/PhysicsControlEditor/Public/PhysicsControlAssetEditor.h
+++ b/Engine/Plugins/Experimental/PhysicsControl/Source/PhysicsControlEditor/Public/PhysicsControlAssetEditor.h
@@ -37,9 +37,11 @@ class PHYSICSCONTROLEDITOR_API FPhysicsControlAssetEditor :
 public:
 	friend class FPhysicsControlAssetApplicationMode;
 	friend class FPhysicsControlAssetEditorEditMode;
+	friend class FPhysicsControlAssetProfileDetailsCustomization;
+	friend class FPhysicsControlAssetPreviewDetailsCustomization;
+	friend class FPhysicsControlAssetSetupDetailsCustomization;
 
 public:
-
 	/** Initialize the asset editor. This will register the application mode, init the preview scene, etc. */
 	void InitAssetEditor(
 		const EToolkitMode::Type        Mode,
@@ -66,7 +68,7 @@ public:
 	// ~END FGCObject overrides.
 
 	// FTickableEditorObject overrides.
-	virtual void Tick(float DeltaTime) override {};
+	virtual void Tick(float DeltaTime) override;
 	virtual ETickableTickType GetTickableTickType() const override { return ETickableTickType::Always; }
 	virtual TStatId GetStatId() const override;
 	// ~END FTickableEditorObject overrides.
@@ -83,7 +85,16 @@ public:
 	/** Refreshes the preview viewport */
 	void RefreshPreviewViewport();
 
-private:
+	/** Invokes the control profile with the name, assuming simulation is running */
+	void InvokeControlProfile(FName ProfileName);
+
+	/** Invokes the most recently invoked control profile */
+	void ReinvokeControlProfile();
+
+	/** Destroys all existing controls modifiers and then recreates them from the control asset */
+	void RecreateControlsAndModifiers();
+
+protected:
 	FText GetSimulationToolTip() const;
 	FSlateIcon GetSimulationIcon() const;
 
@@ -126,14 +137,13 @@ private:
 	bool IsRunningSimulation() const;
 	bool IsNotRunningSimulation() const;
 
-private:
 	/** Make the constraint scale widget */
 	TSharedRef<SWidget> MakeConstraintScaleWidget();
 
 	/** Make the collision opacity widget */
 	TSharedRef<SWidget> MakeCollisionOpacityWidget();
 
-private:
+protected:
 	/** The persona toolkit. */
 	TSharedPtr<IPersonaToolkit> PersonaToolkit = nullptr;
 
@@ -163,4 +173,7 @@ private:
 
 	/** True if in OnTreeSelectionChanged()... protects against infinite recursion */
 	bool bSelecting;
+
+	/** Stored when a control profile is invoked */
+	FName PreviouslyInvokedControlProfile;
 };
