Adds panels showing the control and modifier sets in the Physics Control Editor. Various other cosmetic tweaks relating to that.

diff --git a/Engine/Plugins/Experimental/PhysicsControl/Source/PhysicsControl/Private/PhysicsControlComponent.cpp b/Engine/Plugins/Experimental/PhysicsControl/Source/PhysicsControl/Private/PhysicsControlComponent.cpp
--- a/Engine/Plugins/Experimental/PhysicsControl/Source/PhysicsControl/Private/PhysicsControlComponent.cpp
+++ b/Engine/Plugins/Experimental/PhysicsControl/Source/PhysicsControl/Private/PhysicsControlComponent.cpp
@@ -18,6 +18,7 @@
 #include "Components/BillboardComponent.h"
 
 #include "Engine/Engine.h"
+#include "Engine/SkeletalMesh.h"
 #include "Engine/Texture2D.h"
 
 #include "SceneManagement.h"
@@ -148,49 +149,13 @@ TMap<FName, FPhysicsControlLimbBones> UPhysicsControlComponent::GetLimbBonesFrom
 		return Result;
 	}
 
-	TSet<FName> AllBones;
-
-	// Now walk through each limb, picking up bones, ignoring any that we have already encountered.
-	// This requires the setup data to have been ordered properly.
-	for (const FPhysicsControlLimbSetupData& LimbSetupData : LimbSetupDatas)
+	const FReferenceSkeleton& RefSkeleton = SkeletalMeshComponent->GetSkeletalMeshAsset()->GetRefSkeleton();
+	Result = UE::PhysicsControl::GetLimbBones(LimbSetupDatas, RefSkeleton, PhysicsAsset);
+	for (TPair<FName, FPhysicsControlLimbBones>& Pair : Result)
 	{
-		FPhysicsControlLimbBones& LimbBones = Result.Add(LimbSetupData.LimbName);
-		LimbBones.SkeletalMeshComponent = SkeletalMeshComponent;
-		LimbBones.bCreateBodyModifiers = LimbSetupData.bCreateBodyModifiers;
-		LimbBones.bCreateWorldSpaceControls = LimbSetupData.bCreateWorldSpaceControls;
-		LimbBones.bCreateParentSpaceControls = LimbSetupData.bCreateParentSpaceControls;
-
-		if (LimbSetupData.bIncludeParentBone)
-		{
-			LimbBones.bFirstBoneIsAdditional = true;
-			const FName ParentBoneName = UE::PhysicsControl::GetPhysicalParentBone(
-				SkeletalMeshComponent, LimbSetupData.StartBone);
-			if (!ParentBoneName.IsNone())
-			{
-				LimbBones.BoneNames.Add(ParentBoneName);
-				AllBones.Add(ParentBoneName);
-			}
-		}
-		else
-		{
-			LimbBones.bFirstBoneIsAdditional = false;
-		}
-
-		SkeletalMeshComponent->ForEachBodyBelow(
-			LimbSetupData.StartBone, true, /*bSkipCustomType=*/false,
-			[PhysicsAsset, &AllBones, &LimbBones](const FBodyInstance* BI)
-			{
-				if (USkeletalBodySetup* BodySetup = Cast<USkeletalBodySetup>(BI->BodySetup.Get()))
-				{
-					const FName BoneName = PhysicsAsset->SkeletalBodySetups[BI->InstanceBodyIndex]->BoneName;
-					if (!AllBones.Find(BoneName))
-					{
-						LimbBones.BoneNames.Add(BoneName);
-						AllBones.Add(BoneName);
-					}
-				}
-			});
+		Pair.Value.SkeletalMeshComponent = SkeletalMeshComponent;
 	}
+
 	return Result;
 }
 
@@ -1574,9 +1539,7 @@ FName UPhysicsControlComponent::CreateBodyModifier(
 	const FName                       Set,
 	const FPhysicsControlModifierData BodyModifierData)
 {
-	TSet<FName> Keys;
-	BodyModifierRecords.GetKeys(Keys);
-	const FName Name = UE::PhysicsControl::GetUniqueBodyModifierName(BoneName, Keys, TEXT(""));
+	const FName Name = UE::PhysicsControl::GetUniqueBodyModifierName(BoneName, BodyModifierRecords, TEXT(""));
 	if (CreateNamedBodyModifier(Name, Component, BoneName, Set, BodyModifierData))
 	{
 		return Name;
diff --git a/Engine/Plugins/Experimental/PhysicsControl/Source/PhysicsControl/Private/PhysicsControlComponentImpl.cpp b/Engine/Plugins/Experimental/PhysicsControl/Source/PhysicsControl/Private/PhysicsControlComponentImpl.cpp
--- a/Engine/Plugins/Experimental/PhysicsControl/Source/PhysicsControl/Private/PhysicsControlComponentImpl.cpp
+++ b/Engine/Plugins/Experimental/PhysicsControl/Source/PhysicsControl/Private/PhysicsControlComponentImpl.cpp
@@ -662,6 +662,9 @@ void UPhysicsControlComponent::ApplyBodyModifier(FPhysicsBodyModifierRecord& Rec
 		BodyInstance->PhysicsBlendWeight = Record.BodyModifier.ModifierData.PhysicsBlendWeight;
 		BodyInstance->SetUpdateKinematicFromSimulation(Record.BodyModifier.ModifierData.bUpdateKinematicFromSimulation);
 
+		// On the shapes, this determines whether there is actually collision. Note that the bodies
+		// need to also have "collision enabled" in order to even be allowed to simulate, which is
+		// normally done via the skeletal mesh.
 		UBodySetup* BodySetup = BodyInstance->GetBodySetup();
 		if (BodySetup)
 		{
diff --git a/Engine/Plugins/Experimental/PhysicsControl/Source/PhysicsControl/Private/PhysicsControlOperatorNameGeneration.cpp b/Engine/Plugins/Experimental/PhysicsControl/Source/PhysicsControl/Private/PhysicsControlOperatorNameGeneration.cpp
--- a/Engine/Plugins/Experimental/PhysicsControl/Source/PhysicsControl/Private/PhysicsControlOperatorNameGeneration.cpp
+++ b/Engine/Plugins/Experimental/PhysicsControl/Source/PhysicsControl/Private/PhysicsControlOperatorNameGeneration.cpp
@@ -16,7 +16,8 @@ namespace PhysicsControl
 {
 
 //======================================================================================================================
-FName FindParentBodyBoneName(const FName BoneName, const FReferenceSkeleton& RefSkeleton, UPhysicsAsset* const PhysicsAsset)
+FName FindParentBodyBoneName(
+	const FName BoneName, const FReferenceSkeleton& RefSkeleton, const UPhysicsAsset* PhysicsAsset)
 {
 	FName ParentBoneName;
 
@@ -28,7 +29,8 @@ FName FindParentBodyBoneName(const FName BoneName, const FReferenceSkeleton& Ref
 		{
 			const int32 ParentBodyIndex = PhysicsAsset->FindParentBodyIndex(RefSkeleton, StartBoneIndex);
 
-			if (PhysicsAsset->SkeletalBodySetups.IsValidIndex(ParentBodyIndex) && PhysicsAsset->SkeletalBodySetups[ParentBodyIndex])
+			if (PhysicsAsset->SkeletalBodySetups.IsValidIndex(ParentBodyIndex) && 
+				PhysicsAsset->SkeletalBodySetups[ParentBodyIndex])
 			{
 				ParentBoneName = PhysicsAsset->SkeletalBodySetups[ParentBodyIndex]->BoneName;
 			}
@@ -51,7 +53,8 @@ template<typename TOperatorFunctorType> void CreateAdditionalBodyModifiers(
 		if (Name.IsNone())
 		{
 			UE_LOG(LogPhysicsControl, Warning,
-				TEXT("CreateAdditionalControls: Failed to make body modifier for %s"), *Specifier.Modifier.BoneName.ToString());
+				TEXT("CreateAdditionalBodyModifiers: Failed to make body modifier for %s"), 
+				*Specifier.Modifier.BoneName.ToString());
 		}
 		else
 		{
@@ -99,7 +102,7 @@ template<typename TFunctor> void CreateControlsFromLimbBones(
 	const FPhysicsControlLimbBones& LimbBones,
 	const EPhysicsControlType       ControlType,
 	const FReferenceSkeleton&       RefSkeleton,
-	UPhysicsAsset*                  PhysicsAsset,
+	const UPhysicsAsset*            PhysicsAsset,
 	const FPhysicsControlData&      ControlData,
 	FPhysicsControlNameRecords&     NameRecords,
 	TFunctor&                       CreateOperationLambda)
@@ -120,48 +123,61 @@ template<typename TFunctor> void CreateControlsFromLimbBones(
 			}
 		}
 
-		const FName ControlName = CreateOperationLambda(ParentBoneName, ChildBoneName, ControlData);
+		const FName ControlName = UE::PhysicsControl::GetUniqueControlName(
+			ParentBoneName, ChildBoneName, NameRecords.GetControlNamesInSet(TEXT("All")), TEXT(""));
+		if (ControlName.IsNone())
+		{
+			UE_LOG(LogPhysicsControl, Warning, TEXT("Unable to find a suitable Control name for bones %s and %s"),
+				*ParentBoneName.ToString(), *ChildBoneName.ToString());
+			return;
+		}
 
-		if (!ControlName.IsNone())
+		bool result = CreateOperationLambda(ControlName, ParentBoneName, ChildBoneName, ControlData);
+
+		if (result)
 		{
 			NameRecords.AddControl(ControlName, LimbName);
 			NameRecords.AddControl(ControlName, GetPhysicsControlTypeName(ControlType));
 			NameRecords.AddControl(ControlName, FName(
 				GetPhysicsControlTypeName(ControlType).ToString().Append("_").Append(LimbName.ToString())));
 		}
-		else if (!ParentBoneName.IsNone())
-		{
-			UE_LOG(LogPhysicsControl, Warning,
-				TEXT("Failed to find unique control name for %s and %s"), *ParentBoneName.ToString(), *ChildBoneName.ToString());
-		}
 		else
 		{
 			UE_LOG(LogPhysicsControl, Warning,
-				TEXT("Failed to find unique control name for %s"), *ChildBoneName.ToString());
+				TEXT("Failed to create control for %s"), *ChildBoneName.ToString());
 		}
 	}
 }
 
 //======================================================================================================================
 template<typename TFunctor> void CreateBodyModifiersFromLimbBones(
-	const FName                   LimbName,
-	const FPhysicsControlLimbBones& LimbBones,
+	const FName                        LimbName,
+	const FPhysicsControlLimbBones&    LimbBones,
 	const FPhysicsControlModifierData& ModifierData,
-	FPhysicsControlNameRecords& NameRecords,
-	TFunctor& CreateOperationLambda)
+	FPhysicsControlNameRecords&        NameRecords,
+	TFunctor&                          CreateOperationLambda)
 {
 	for (const FName BoneName : LimbBones.BoneNames)
 	{
-		const FName BodyModifierName = CreateOperationLambda(BoneName, ModifierData);
+		const FName BodyModifierName = UE::PhysicsControl::GetUniqueBodyModifierName(
+			BoneName, NameRecords.GetBodyModifierNamesInSet(TEXT("All")), TEXT(""));
+		if (BodyModifierName.IsNone())
+		{
+			UE_LOG(LogPhysicsControl, Warning, TEXT("Unable to find a suitable Body Modifier name for bone %s"),
+				*BoneName.ToString());
+			return;
+		}
 
-		if (!BodyModifierName.IsNone())
+		bool result = CreateOperationLambda(BodyModifierName, BoneName, ModifierData);
+
+		if (result)
 		{
 			NameRecords.AddBodyModifier(BodyModifierName, LimbName);
 		}
 		else
 		{
 			UE_LOG(LogPhysicsControl, Warning,
-				TEXT("Failed to find unique body modifier name for %s"), *BoneName.ToString());
+				TEXT("Failed to create body modifier %s"), *BodyModifierName.ToString());
 		}
 	}
 }
@@ -171,9 +187,9 @@ template<typename TControlFunctorType, typename TBodyModifierFunctorType>
 void ForEachPotentialOperator(
 	const FPhysicsControlCharacterSetupData&           CharacterSetupData,
 	const FPhysicsControlAndBodyModifierCreationDatas& AdditionalControlsAndBodyModifiers,
-	TMap<FName, FPhysicsControlLimbBones>              AllLimbBones, 
+	const TMap<FName, FPhysicsControlLimbBones>        AllLimbBones, 
 	const FReferenceSkeleton&                          RefSkeleton, 
-	UPhysicsAsset*                                     PhysicsAsset, 
+	const UPhysicsAsset*                               PhysicsAsset, 
 	FPhysicsControlNameRecords&                        NameRecords, 
 	TControlFunctorType&                               ControlFunctor, 
 	TBodyModifierFunctorType&                          BodyModifierFunctor)
@@ -255,13 +271,12 @@ void CreateAdditionalSets_Implementation(
 }
 
 //======================================================================================================================
+// Note - Output limb bones are not in the order specified in the skeleton - would be better if they were
 TMap<FName, FPhysicsControlLimbBones> GetLimbBones(
-	const TArray<FPhysicsControlLimbSetupData>& LimbSetupData,
-	const FReferenceSkeleton& RefSkeleton,
-	UPhysicsAsset* const PhysicsAsset)
+	const TArray<FPhysicsControlLimbSetupData>& LimbSetupDatas,
+	const FReferenceSkeleton&                   RefSkeleton,
+	const UPhysicsAsset*                        PhysicsAsset)
 {
-	// Note - Output limb bones are not in the order specified in the skeleton - would be better if they were
-
 	TMap<FName, FPhysicsControlLimbBones> Result;
 
 	if (!PhysicsAsset)
@@ -272,7 +287,7 @@ TMap<FName, FPhysicsControlLimbBones> GetLimbBones(
 
 	TSet<FName> AllBones;
 
-	for (const FPhysicsControlLimbSetupData& LimbSetup : LimbSetupData)
+	for (const FPhysicsControlLimbSetupData& LimbSetup : LimbSetupDatas)
 	{
 		FPhysicsControlLimbBones& LimbBones = Result.Add(LimbSetup.LimbName);
 
@@ -313,26 +328,28 @@ TMap<FName, FPhysicsControlLimbBones> GetLimbBones(
 
 //======================================================================================================================
 void CollectOperatorNames(
-	const FAnimNode_RigidBodyWithControl*            Node, 
-	const FPhysicsControlCharacterSetupData&         CharacterSetupData,
+	const FPhysicsControlCharacterSetupData&           CharacterSetupData,
 	const FPhysicsControlAndBodyModifierCreationDatas& AdditionalControlsAndBodyModifiers,
-	TMap<FName, FPhysicsControlLimbBones>            AllLimbBones,
-	const FReferenceSkeleton&                        RefSkeleton, 
-	UPhysicsAsset*                                   PhysicsAsset, 
-	TSet<FName>&                                     BodyModifierNames, 
-	TSet<FName>&                                     ControlNames, 
-	FPhysicsControlNameRecords&                      NameRecords)
+	const TMap<FName, FPhysicsControlLimbBones>        AllLimbBones,
+	const FReferenceSkeleton&                          RefSkeleton, 
+	const UPhysicsAsset*                               PhysicsAsset, 
+	TSet<FName>&                                       BodyModifierNames, 
+	TSet<FName>&                                       ControlNames, 
+	FPhysicsControlNameRecords&                        NameRecords)
 {
-	auto CollectControlName = [&ControlNames](const FName ParentBoneName, const FName ChildBoneName, const FPhysicsControlData& Data) -> FName 
+	auto CollectControlName = [&ControlNames](
+		const FName ControlName, const FName ParentBoneName, 
+		const FName ChildBoneName, const FPhysicsControlData& Data) -> bool 
+		{
+			ControlNames.Add(ControlName); 
+			return true; 
+		};
+	auto CollectBodyModifierName = [&BodyModifierNames](
+		const FName BodyModifierName, const FName BoneName, const FPhysicsControlModifierData& Data) -> bool
 		{
-			const FName Name = GetUniqueControlName(ParentBoneName, ChildBoneName, ControlNames, TEXT(""));
-			ControlNames.Add(Name); 
-			return Name; 
+			BodyModifierNames.Add(BodyModifierName);
+			return true;
 		};
-	auto CollectBodyModifierName = [&BodyModifierNames](const FName BoneName, const FPhysicsControlModifierData& Data) -> FName { 
-		const FName Name = GetUniqueBodyModifierName(BoneName, BodyModifierNames, TEXT(""));
-		BodyModifierNames.Add(Name); 
-		return Name; };
 
 	ForEachPotentialOperator(
 		CharacterSetupData, AdditionalControlsAndBodyModifiers, AllLimbBones, RefSkeleton,
@@ -344,18 +361,21 @@ void CreateOperatorsForNode(
 	FAnimNode_RigidBodyWithControl*                  Node,
 	const FPhysicsControlCharacterSetupData&         CharacterSetupData,
 	const FPhysicsControlAndBodyModifierCreationDatas& AdditionalControlsAndBodyModifiers,
-	TMap<FName, FPhysicsControlLimbBones>            AllLimbBones,
+	const TMap<FName, FPhysicsControlLimbBones>      AllLimbBones,
 	const FReferenceSkeleton&                        RefSkeleton, 
-	UPhysicsAsset*                                   PhysicsAsset, 
+	const UPhysicsAsset*                             PhysicsAsset, 
 	FPhysicsControlNameRecords&                      NameRecords)
 {
-	auto CreateControl = [Node](const FName ParentBoneName, const FName ChildBoneName, const FPhysicsControlData& Data) -> FName 
+	auto CreateControl = [Node](
+		const FName ControlName, const FName ParentBoneName, 
+		const FName ChildBoneName, const FPhysicsControlData& Data) -> bool
 		{ 
-			return Node->CreateControl(ParentBoneName, ChildBoneName, Data); 
+			return Node->CreateNamedControl(ControlName, ParentBoneName, ChildBoneName, Data); 
 		};
-	auto CreateBodyModifier = [Node](const FName BoneName, const FPhysicsControlModifierData& Data) -> FName 
+	auto CreateBodyModifier = [Node](
+		const FName BodyModifierName, const FName BoneName, const FPhysicsControlModifierData& Data) -> bool
 		{ 
-			return Node->CreateBodyModifier(BoneName, Data); 
+			return Node->CreateNamedBodyModifier(BodyModifierName, BoneName, Data); 
 		};
 
 	ForEachPotentialOperator(
diff --git a/Engine/Plugins/Experimental/PhysicsControl/Source/PhysicsControl/Private/RigidBodyWithControl.cpp b/Engine/Plugins/Experimental/PhysicsControl/Source/PhysicsControl/Private/RigidBodyWithControl.cpp
--- a/Engine/Plugins/Experimental/PhysicsControl/Source/PhysicsControl/Private/RigidBodyWithControl.cpp
+++ b/Engine/Plugins/Experimental/PhysicsControl/Source/PhysicsControl/Private/RigidBodyWithControl.cpp
@@ -258,8 +258,11 @@ void FAnimNode_RigidBodyWithControl::UpdateBodyIndicesInControlRecord(FRigidBody
 }
 
 //======================================================================================================================
-FName FAnimNode_RigidBodyWithControl::CreateControl(
-	const FName ParentBoneName, const FName ChildBoneName, const FPhysicsControlData& ControlData)
+bool FAnimNode_RigidBodyWithControl::CreateNamedControl(
+	const FName                ControlName, 
+	const FName                ParentBoneName, 
+	const FName                ChildBoneName, 
+	const FPhysicsControlData& ControlData)
 {
 	FPhysicsControl Control;
 	Control.ParentBoneName = ParentBoneName;
@@ -285,13 +288,24 @@ FName FAnimNode_RigidBodyWithControl::CreateControl(
 	{
 		UE_LOG(LogPhysicsControl, Warning,
 			TEXT("Unable to create world space control constraint for bone %s"), *ChildBoneName.ToString());
-		return FName();
+		return false;
 	}
 
-	FName ControlName = GetUniqueControlName(Control.ParentBoneName, Control.ChildBoneName);
 	ControlRecords.Add(ControlName, FRigidBodyControlRecord(Control, JointHandle));
+	return true;
+}
 
-	return ControlName;
+
+//======================================================================================================================
+FName FAnimNode_RigidBodyWithControl::CreateControl(
+	const FName ParentBoneName, const FName ChildBoneName, const FPhysicsControlData& ControlData)
+{
+	FName ControlName = GetUniqueControlName(ParentBoneName, ChildBoneName);
+	if (CreateNamedControl(ControlName, ParentBoneName, ChildBoneName, ControlData))
+	{
+		return ControlName;
+	}
+	return FName();
 }
 
 //======================================================================================================================
@@ -333,18 +347,29 @@ FName FAnimNode_RigidBodyWithControl::GetUniqueControlName(const FName ParentBon
 }
 
 //======================================================================================================================
-FName FAnimNode_RigidBodyWithControl::CreateBodyModifier(FName BoneName, const FPhysicsControlModifierData& ModifierData)
+bool FAnimNode_RigidBodyWithControl::CreateNamedBodyModifier(
+	const FName ModifierName, const FName BoneName, const FPhysicsControlModifierData& ModifierData)
 {
-	FName Name;
 	ImmediatePhysics::FActorHandle* const ActorHandle = FindBodyFromBoneName(BoneName);
 	if (ActorHandle)
 	{
-		Name = GetUniqueBodyModifierName(BoneName);
 		FPhysicsBodyModifier BodyModifier(BoneName, ModifierData);
-		FRigidBodyModifierRecord& Modifier = ModifierRecords.Add(Name, FRigidBodyModifierRecord(BodyModifier, ActorHandle));
+		FRigidBodyModifierRecord& Modifier = ModifierRecords.Add(
+			ModifierName, FRigidBodyModifierRecord(BodyModifier, ActorHandle));
+		return true;
 	}
+	return false;
+}
 
-	return Name;
+//======================================================================================================================
+FName FAnimNode_RigidBodyWithControl::CreateBodyModifier(FName BoneName, const FPhysicsControlModifierData& ModifierData)
+{
+	FName ModifierName = GetUniqueBodyModifierName(BoneName);
+	if (CreateNamedBodyModifier(ModifierName, BoneName, ModifierData))
+	{
+		return ModifierName;
+	}
+	return FName();
 }
 
 //======================================================================================================================
diff --git a/Engine/Plugins/Experimental/PhysicsControl/Source/PhysicsControl/Public/AnimNode_RigidBodyWithControl.h b/Engine/Plugins/Experimental/PhysicsControl/Source/PhysicsControl/Public/AnimNode_RigidBodyWithControl.h
--- a/Engine/Plugins/Experimental/PhysicsControl/Source/PhysicsControl/Public/AnimNode_RigidBodyWithControl.h
+++ b/Engine/Plugins/Experimental/PhysicsControl/Source/PhysicsControl/Public/AnimNode_RigidBodyWithControl.h
@@ -419,8 +419,12 @@ private:
 public:
 	const FTransform GetBodyTransform(const int32 BodyIndex) const;
 
+	bool CreateNamedControl(
+		const FName ControlName, const FName ParentBoneName, const FName ChildBoneName, const FPhysicsControlData& ControlData);
 	FName CreateControl(const FName ParentBoneName, const FName ChildBoneName, const FPhysicsControlData& ControlData);
 
+	bool CreateNamedBodyModifier(
+		const FName ModifierName, const FName BoneName, const FPhysicsControlModifierData& ModifierData);
 	FName CreateBodyModifier(const FName BoneName, const FPhysicsControlModifierData& ModifierData);
 
 private:
diff --git a/Engine/Plugins/Experimental/PhysicsControl/Source/PhysicsControl/Public/PhysicsControlOperatorNameGeneration.h b/Engine/Plugins/Experimental/PhysicsControl/Source/PhysicsControl/Public/PhysicsControlOperatorNameGeneration.h
--- a/Engine/Plugins/Experimental/PhysicsControl/Source/PhysicsControl/Public/PhysicsControlOperatorNameGeneration.h
+++ b/Engine/Plugins/Experimental/PhysicsControl/Source/PhysicsControl/Public/PhysicsControlOperatorNameGeneration.h
@@ -27,32 +27,31 @@ constexpr int32 MaxNumControlsOrModifiersPerName = 16;
 
 // Parse the skeleton tree to figure out which bones are associated with which limbs. 
 PHYSICSCONTROL_API TMap<FName, FPhysicsControlLimbBones> GetLimbBones(
-	const TArray<FPhysicsControlLimbSetupData>& LimbSetupData, 
+	const TArray<FPhysicsControlLimbSetupData>& LimbSetupDatas, 
 	const FReferenceSkeleton&                   RefSkeleton, 
-	UPhysicsAsset*                              PhysicsAsset);
+	const UPhysicsAsset*                        PhysicsAsset);
 
 // Populates the supplied body modifier names, control names and name records structures with the
-// names and sets that could be created for the supplied node, limb bones, skeleton and physics asset.
+// names and sets that would be created for the supplied limb bones, skeleton and physics asset.
 PHYSICSCONTROL_API void CollectOperatorNames(
-	const FAnimNode_RigidBodyWithControl*            Node, 
-	const FPhysicsControlCharacterSetupData&         CharacterSetupData,
+	const FPhysicsControlCharacterSetupData&           CharacterSetupData,
 	const FPhysicsControlAndBodyModifierCreationDatas& AdditionalControlsAndBodyModifiers,
-	TMap<FName, FPhysicsControlLimbBones>            AllLimbBones,
-	const FReferenceSkeleton&                        RefSkeleton, 
-	UPhysicsAsset*                                   PhysicsAsset, 
-	TSet<FName>&                                     BodyModifierNames, 
-	TSet<FName>&                                     ControlNames, 
-	FPhysicsControlNameRecords&                      NameRecords);
+	const TMap<FName, FPhysicsControlLimbBones>        AllLimbBones,
+	const FReferenceSkeleton&                          RefSkeleton, 
+	const UPhysicsAsset*                               PhysicsAsset, 
+	TSet<FName>&                                       BodyModifierNames, 
+	TSet<FName>&                                       ControlNames, 
+	FPhysicsControlNameRecords&                        NameRecords);
 
 // Creates the body modifiers, controls and sets for the supplied node, limb bones, skeleton and physics asset.
 PHYSICSCONTROL_API void CreateOperatorsForNode(
-	FAnimNode_RigidBodyWithControl*                  Node, 
-	const FPhysicsControlCharacterSetupData&         CharacterSetupData,
+	FAnimNode_RigidBodyWithControl*                    Node, 
+	const FPhysicsControlCharacterSetupData&           CharacterSetupData,
 	const FPhysicsControlAndBodyModifierCreationDatas& AdditionalControlsAndBodyModifiers,
-	TMap<FName, FPhysicsControlLimbBones>            AllLimbBones,
-	const FReferenceSkeleton&                        RefSkeleton, 
-	UPhysicsAsset*                                   PhysicsAsset, 
-	FPhysicsControlNameRecords&                      NameRecords);
+	const TMap<FName, FPhysicsControlLimbBones>        AllLimbBones,
+	const FReferenceSkeleton&                          RefSkeleton, 
+	const UPhysicsAsset*                               PhysicsAsset, 
+	FPhysicsControlNameRecords&                        NameRecords);
 
 // Adds the specified additional sets to the supplied Name Records structure.
 PHYSICSCONTROL_API void CreateAdditionalSets(
@@ -70,10 +69,32 @@ PHYSICSCONTROL_API void CreateAdditionalSets(
 
 // Adds the specified additional sets to the supplied Name Records structure.
 PHYSICSCONTROL_API void CreateAdditionalSets(
-	const FPhysicsControlSetUpdates&             AdditionalSets, 
-	const TMap<FName, FPhysicsBodyModifierRecord>&     BodyModifierRecords,
-	const TMap<FName, FPhysicsControlRecord>&    Controls,
-	FPhysicsControlNameRecords&                  NameRecords);
+	const FPhysicsControlSetUpdates&               AdditionalSets, 
+	const TMap<FName, FPhysicsBodyModifierRecord>& BodyModifierRecords,
+	const TMap<FName, FPhysicsControlRecord>&      Controls,
+	FPhysicsControlNameRecords&                    NameRecords);
+
+//======================================================================================================================
+// Helper for GetUniqueName
+inline bool DoesNameExist(const FName Name, const TArray<FName>& ExistingNames)
+{
+	return ExistingNames.Find(Name) != INDEX_NONE;
+}
+
+//======================================================================================================================
+// Helper for GetUniqueName
+template<typename T>
+bool DoesNameExist(const FName Name, const TMap<FName, T>& ExistingNames)
+{
+	return ExistingNames.Find(Name) != nullptr;
+}
+
+//======================================================================================================================
+// Helper for GetUniqueName
+inline bool DoesNameExist(const FName Name, const TSet<FName>& ExistingNames)
+{
+	return ExistingNames.Find(Name) != nullptr;
+}
 
 //======================================================================================================================
 template<typename CollectionType>
@@ -87,7 +108,7 @@ FName GetUniqueName(const FString& NameBase, const CollectionType& ExistingNames
 	for (int32 Index = 0; Index < MaxNameIndex; ++Index)
 	{
 		const FName Name(NameStr);
-		if (!ExistingNames.Find(Name))
+		if (!DoesNameExist(Name, ExistingNames))
 		{
 			return Name;
 		}
diff --git a/Engine/Plugins/Experimental/PhysicsControl/Source/PhysicsControlEditor/Private/PhysicsControlAssetApplicationMode.cpp b/Engine/Plugins/Experimental/PhysicsControl/Source/PhysicsControlEditor/Private/PhysicsControlAssetApplicationMode.cpp
--- a/Engine/Plugins/Experimental/PhysicsControl/Source/PhysicsControlEditor/Private/PhysicsControlAssetApplicationMode.cpp
+++ b/Engine/Plugins/Experimental/PhysicsControl/Source/PhysicsControlEditor/Private/PhysicsControlAssetApplicationMode.cpp
@@ -2,7 +2,7 @@
 
 #include "PhysicsControlAssetApplicationMode.h"
 #include "PhysicsControlAssetEditor.h"
-#include "PhysicsControlAssetEditorProfilesTabSummoners.h"
+#include "PhysicsControlAssetEditorTabSummoners.h"
 #include "PhysicsControlAsset.h"
 
 #include "PersonaModule.h"
@@ -14,8 +14,6 @@
 
 FName FPhysicsControlAssetApplicationMode::ModeName("PhysicsControlAssetEditMode");
 
-static const FName PhysicsControlAssetEditorTabName("PhysicsControlAssetEditorTab");
-
 //======================================================================================================================
 FPhysicsControlAssetApplicationMode::FPhysicsControlAssetApplicationMode(
 	TSharedRef<FWorkflowCentricApplication> InHostingApp,
@@ -74,50 +72,85 @@ FPhysicsControlAssetApplicationMode::FPhysicsControlAssetApplicationMode(
 			CastChecked<UPhysicsControlAsset>(
 				(*PhysicsControlAssetEditorSharedRef->GetObjectsCurrentlyBeingEdited())[0])));
 
+	TabFactories.RegisterFactory(
+		MakeShared<FPhysicsControlAssetEditorControlSetsTabSummoner>(
+			InHostingApp,
+			CastChecked<UPhysicsControlAsset>(
+				(*PhysicsControlAssetEditorSharedRef->GetObjectsCurrentlyBeingEdited())[0])));
+
+	TabFactories.RegisterFactory(
+		MakeShared<FPhysicsControlAssetEditorBodyModifierSetsTabSummoner>(
+			InHostingApp,
+			CastChecked<UPhysicsControlAsset>(
+				(*PhysicsControlAssetEditorSharedRef->GetObjectsCurrentlyBeingEdited())[0])));
 
 	// For standard tabs, these are provided by Persona. For custom tabs, the contents of the tab is
 	// provided by a summoner. When Summoners are made, they register their name in the constructor.
-	// These names then hook into the tab names below. See MLDeformerApplicationMode.cpp - that puts
-	// the names into a TabID field in the summoner.
+	// These names then hook into the tab names below.
 
 	// Create tab layout.
-	TabLayout = FTabManager::NewLayout("Standalone_PhysicsControlAssetEditor_Layout_v1.0")
+	TabLayout = FTabManager::NewLayout("Standalone_PhysicsControlAssetEditor_Layout_v1.1")
 		->AddArea
 		(
 			FTabManager::NewPrimaryArea()
-			->SetOrientation(Orient_Vertical)
+			->SetOrientation(Orient_Horizontal)
 			->Split
 			(
+				// Control/modifier sets
 				FTabManager::NewSplitter()
-				->SetSizeCoefficient(0.32f)
-				->SetOrientation(Orient_Horizontal)
+				->SetSizeCoefficient(0.3f)
+				->SetOrientation(Orient_Vertical)
 				->Split
 				(
-					FTabManager::NewSplitter()
+					FTabManager::NewStack()
 					->SetSizeCoefficient(1.0f)
-					->SetOrientation(Orient_Vertical)
-					->Split
-					(
-						FTabManager::NewStack()
-						->SetSizeCoefficient(0.7f)
-						->SetHideTabWell(true)
-						->AddTab(FPersonaTabs::PreviewViewportID, ETabState::OpenedTab)
-						->AddTab(FPersonaTabs::SkeletonTreeViewID, ETabState::OpenedTab)
-					)
+					->AddTab(
+						FPhysicsControlAssetEditorBodyModifierSetsTabSummoner::TabName, ETabState::OpenedTab)
+				)
+				->Split
+				(
+					FTabManager::NewStack()
+					->SetSizeCoefficient(2.0f)
+					->AddTab(
+						FPhysicsControlAssetEditorControlSetsTabSummoner::TabName, ETabState::OpenedTab)
 				)
+			)
+			->Split
+			(
+				// Skeleton
+				FTabManager::NewStack()
+				->SetSizeCoefficient(0.3f)
+				->AddTab(FPersonaTabs::SkeletonTreeViewID, ETabState::OpenedTab)
+			)
+			->Split
+			(
+				// Preview window
+				FTabManager::NewStack()
+				->SetSizeCoefficient(1.0f)
+				->AddTab(FPersonaTabs::PreviewViewportID, ETabState::OpenedTab)
+			)
+			->Split
+			(
+				// the profile and detail panels
+				FTabManager::NewSplitter()
+				->SetSizeCoefficient(0.5f)
+				->SetOrientation(Orient_Vertical)
 				->Split
 				(
 					FTabManager::NewStack()
-					->SetSizeCoefficient(0.3f)
-					//->AddTab(FPersonaTabs::DetailsID, ETabState::OpenedTab)
+					->SetSizeCoefficient(1.0f)
 					->AddTab(FPersonaTabs::AdvancedPreviewSceneSettingsID, ETabState::OpenedTab)
 					->AddTab(
 						FPhysicsControlAssetEditorSetupTabSummoner::TabName, ETabState::OpenedTab)
 					->AddTab(
 						FPhysicsControlAssetEditorProfileTabSummoner::TabName, ETabState::OpenedTab)
+					->SetForegroundTab(FPhysicsControlAssetEditorSetupTabSummoner::TabName)
+				)
+				->Split(
+					FTabManager::NewStack()
+					->SetSizeCoefficient(0.6f)
 					->AddTab(
 						FPhysicsControlAssetEditorPreviewTabSummoner::TabName, ETabState::OpenedTab)
-					->SetForegroundTab(FPhysicsControlAssetEditorSetupTabSummoner::TabName)
 				)
 			)
 		);
diff --git a/Engine/Plugins/Experimental/PhysicsControl/Source/PhysicsControlEditor/Private/PhysicsControlAssetEditorProfilesTabSummoners.cpp b/Engine/Plugins/Experimental/PhysicsControl/Source/PhysicsControlEditor/Private/PhysicsControlAssetEditorProfilesTabSummoners.cpp
deleted file mode 100644
diff --git a/Engine/Plugins/Experimental/PhysicsControl/Source/PhysicsControlEditor/Private/PhysicsControlAssetEditorProfilesTabSummoners.h b/Engine/Plugins/Experimental/PhysicsControl/Source/PhysicsControlEditor/Private/PhysicsControlAssetEditorProfilesTabSummoners.h
deleted file mode 100644
diff --git a/Engine/Plugins/Experimental/PhysicsControl/Source/PhysicsControlEditor/Private/PhysicsControlAssetEditorTabSummoners.cpp b/Engine/Plugins/Experimental/PhysicsControl/Source/PhysicsControlEditor/Private/PhysicsControlAssetEditorTabSummoners.cpp
new file mode 100644
--- /dev/null
+++ b/Engine/Plugins/Experimental/PhysicsControl/Source/PhysicsControlEditor/Private/PhysicsControlAssetEditorTabSummoners.cpp
@@ -0,0 +1,278 @@
+// Copyright Epic Games, Inc. All Rights Reserved.
+
+#include "PhysicsControlAssetEditorTabSummoners.h"
+#include "IDocumentation.h"
+#include "PhysicsControlAssetEditor.h"
+#include "PhysicsControlAsset.h"
+#include "PropertyEditorModule.h"
+#include "IDetailsView.h"
+#include "PhysicsControlAssetSetupDetailsCustomization.h"
+#include "PhysicsControlAssetProfileDetailsCustomization.h"
+#include "PhysicsControlAssetPreviewDetailsCustomization.h"
+#include "PhysicsControlAssetInfoDetailsCustomization.h"
+
+#define LOCTEXT_NAMESPACE "PhysicsControlAssetEditorTabSummoner"
+
+//======================================================================================================================
+FName FPhysicsControlAssetEditorSetupTabSummoner::TabName = FName(
+	"PhysicsControlAssetEditorSetupTab");
+FName FPhysicsControlAssetEditorProfileTabSummoner::TabName = FName(
+	"PhysicsControlAssetEditorProfileTab");
+FName FPhysicsControlAssetEditorPreviewTabSummoner::TabName = FName(
+	"PhysicsControlAssetEditorPreviewTab");
+FName FPhysicsControlAssetEditorControlSetsTabSummoner::TabName = FName(
+	"PhysicsControlAssetEditorControlSetsTab");
+FName FPhysicsControlAssetEditorBodyModifierSetsTabSummoner::TabName = FName(
+	"PhysicsControlAssetEditorBodyModifierSetsTab");
+
+//======================================================================================================================
+FPhysicsControlAssetEditorSetupTabSummoner::FPhysicsControlAssetEditorSetupTabSummoner(
+	TSharedPtr<FAssetEditorToolkit> InHostingApp, UPhysicsControlAsset* InPhysicsControlAsset)
+	: FWorkflowTabFactory(FPhysicsControlAssetEditorSetupTabSummoner::TabName, InHostingApp)
+	, PhysicsControlAsset(InPhysicsControlAsset)
+{
+	TabLabel = LOCTEXT("PhysicsControlAssetEditorSetupTabTitle", "Setup");
+	TabIcon = FSlateIcon(FAppStyle::GetAppStyleSetName(), "PhysicsAssetEditor.Tabs.Profiles");
+
+	bIsSingleton = true;
+
+	ViewMenuDescription = LOCTEXT("PhysicsControlAssetEditorSetup", "Setup");
+	ViewMenuTooltip = LOCTEXT("PhysicsControlAssetEditorSetup_ToolTip", "Shows the Control Asset Setup tab");
+}
+
+//======================================================================================================================
+TSharedPtr<SToolTip> FPhysicsControlAssetEditorSetupTabSummoner::CreateTabToolTipWidget(
+	const FWorkflowTabSpawnInfo& Info) const
+{
+	return IDocumentation::Get()->CreateToolTip(LOCTEXT(
+		"PhysicsControlAssetEditorSetupToolTip", 
+		"The Physics Control Asset Setup tab lets you edit the physics control asset relating to setting up controls."), 
+		NULL, 
+		TEXT("Shared/Editors/PhysicsControlAssetEditor"), TEXT("PhysicsControlAssetProfiles_Window"));
+}
+
+//======================================================================================================================
+TSharedRef<SWidget> FPhysicsControlAssetEditorSetupTabSummoner::CreateTabBody(
+	const FWorkflowTabSpawnInfo& Info) const
+{
+	FDetailsViewArgs DetailsViewArgs;
+	DetailsViewArgs.NameAreaSettings = FDetailsViewArgs::HideNameArea;
+	DetailsViewArgs.bHideSelectionTip = true;
+	DetailsViewArgs.bAllowSearch = false;
+
+	FPropertyEditorModule& PropertyEditorModule = 
+		FModuleManager::LoadModuleChecked<FPropertyEditorModule>("PropertyEditor");
+	TSharedRef<IDetailsView> DetailsView = PropertyEditorModule.CreateDetailView(DetailsViewArgs);
+
+	TWeakPtr<FPhysicsControlAssetEditor> PhysicsControlAssetEditor = 
+		StaticCastSharedPtr<FPhysicsControlAssetEditor>(HostingApp.Pin());
+	DetailsView->RegisterInstancedCustomPropertyLayout(
+		UPhysicsControlAsset::StaticClass(), 
+		FOnGetDetailCustomizationInstance::CreateStatic(
+			&FPhysicsControlAssetSetupDetailsCustomization::MakeInstance, PhysicsControlAssetEditor));
+	DetailsView->SetObject(PhysicsControlAsset.Get());
+	return DetailsView;
+}
+
+//======================================================================================================================
+FPhysicsControlAssetEditorProfileTabSummoner::FPhysicsControlAssetEditorProfileTabSummoner(
+	TSharedPtr<FAssetEditorToolkit> InHostingApp, UPhysicsControlAsset* InPhysicsControlAsset)
+	: FWorkflowTabFactory(FPhysicsControlAssetEditorProfileTabSummoner::TabName, InHostingApp)
+	, PhysicsControlAsset(InPhysicsControlAsset)
+{
+	TabLabel = LOCTEXT("PhysicsControlAssetEditorProfileTabTitle", "Profiles");
+	TabIcon = FSlateIcon(FAppStyle::GetAppStyleSetName(), "PhysicsAssetEditor.Tabs.Profiles");
+
+	bIsSingleton = true;
+
+	ViewMenuDescription = LOCTEXT("PhysicsControlAssetEditorProfile", "Profiles");
+	ViewMenuTooltip = LOCTEXT("PhysicsControlAssetEditorProfile_ToolTip", "Shows the Control Asset Profile Edit tab");
+}
+
+//======================================================================================================================
+TSharedPtr<SToolTip> FPhysicsControlAssetEditorProfileTabSummoner::CreateTabToolTipWidget(
+	const FWorkflowTabSpawnInfo& Info) const
+{
+	return IDocumentation::Get()->CreateToolTip(LOCTEXT(
+		"PhysicsControlAssetEditorProfileToolTip",
+		"The Physics Control Asset Profile Edit tab lets you edit the physics control asset relating to setting up profiles."),
+		NULL,
+		TEXT("Shared/Editors/PhysicsControlAssetEditor"), TEXT("PhysicsControlAssetProfiles_Window"));
+}
+
+//======================================================================================================================
+TSharedRef<SWidget> FPhysicsControlAssetEditorProfileTabSummoner::CreateTabBody(
+	const FWorkflowTabSpawnInfo& Info) const
+{
+	FDetailsViewArgs DetailsViewArgs;
+	DetailsViewArgs.NameAreaSettings = FDetailsViewArgs::HideNameArea;
+	DetailsViewArgs.bHideSelectionTip = true;
+	DetailsViewArgs.bAllowSearch = false;
+
+	FPropertyEditorModule& PropertyEditorModule =
+		FModuleManager::LoadModuleChecked<FPropertyEditorModule>("PropertyEditor");
+	TSharedRef<IDetailsView> DetailsView = PropertyEditorModule.CreateDetailView(DetailsViewArgs);
+
+	TWeakPtr<FPhysicsControlAssetEditor> PhysicsControlAssetEditor =
+		StaticCastSharedPtr<FPhysicsControlAssetEditor>(HostingApp.Pin());
+	DetailsView->RegisterInstancedCustomPropertyLayout(
+		UPhysicsControlAsset::StaticClass(), 
+		FOnGetDetailCustomizationInstance::CreateStatic(
+			&FPhysicsControlAssetProfileDetailsCustomization::MakeInstance, PhysicsControlAssetEditor));
+	DetailsView->SetObject(PhysicsControlAsset.Get());
+	return DetailsView;
+}
+
+//======================================================================================================================
+FPhysicsControlAssetEditorPreviewTabSummoner::FPhysicsControlAssetEditorPreviewTabSummoner(
+	TSharedPtr<FAssetEditorToolkit> InHostingApp, UPhysicsControlAsset* InPhysicsControlAsset)
+	: FWorkflowTabFactory(FPhysicsControlAssetEditorPreviewTabSummoner::TabName, InHostingApp)
+	, PhysicsControlAsset(InPhysicsControlAsset)
+{
+	TabLabel = LOCTEXT("PhysicsControlAssetEditorPreviewTabTitle", "Preview");
+	TabIcon = FSlateIcon(FAppStyle::GetAppStyleSetName(), "PhysicsAssetEditor.Tabs.Profiles");
+
+	bIsSingleton = true;
+
+	ViewMenuDescription = LOCTEXT("PhysicsControlAssetEditorPreview", "Preview");
+	ViewMenuTooltip = LOCTEXT("PhysicsControlAssetEditorPreview_ToolTip", "Shows the Control Asset Preview tab");
+}
+
+//======================================================================================================================
+TSharedPtr<SToolTip> FPhysicsControlAssetEditorPreviewTabSummoner::CreateTabToolTipWidget(
+	const FWorkflowTabSpawnInfo& Info) const
+{
+	return IDocumentation::Get()->CreateToolTip(LOCTEXT(
+		"PhysicsControlAssetEditorPreviewToolTip",
+		"The Physics Control Asset Preview tab lets you preview the physics control asset setup and profiles."),
+		NULL,
+		TEXT("Shared/Editors/PhysicsControlAssetEditor"), TEXT("PhysicsControlAssetProfiles_Window"));
+}
+
+//======================================================================================================================
+TSharedRef<SWidget> FPhysicsControlAssetEditorPreviewTabSummoner::CreateTabBody(
+	const FWorkflowTabSpawnInfo& Info) const
+{
+	FDetailsViewArgs DetailsViewArgs;
+	DetailsViewArgs.NameAreaSettings = FDetailsViewArgs::HideNameArea;
+	DetailsViewArgs.bHideSelectionTip = true;
+	DetailsViewArgs.bAllowSearch = false;
+
+	FPropertyEditorModule& PropertyEditorModule =
+		FModuleManager::LoadModuleChecked<FPropertyEditorModule>("PropertyEditor");
+	TSharedRef<IDetailsView> DetailsView = PropertyEditorModule.CreateDetailView(DetailsViewArgs);
+
+	TWeakPtr<FPhysicsControlAssetEditor> PhysicsControlAssetEditor =
+		StaticCastSharedPtr<FPhysicsControlAssetEditor>(HostingApp.Pin());
+	DetailsView->RegisterInstancedCustomPropertyLayout(
+		UPhysicsControlAsset::StaticClass(), 
+		FOnGetDetailCustomizationInstance::CreateStatic(
+			&FPhysicsControlAssetPreviewDetailsCustomization::MakeInstance, PhysicsControlAssetEditor));
+	DetailsView->SetObject(PhysicsControlAsset.Get());
+	return DetailsView;
+}
+
+//======================================================================================================================
+FPhysicsControlAssetEditorControlSetsTabSummoner::FPhysicsControlAssetEditorControlSetsTabSummoner(
+	TSharedPtr<FAssetEditorToolkit> InHostingApp, UPhysicsControlAsset* InPhysicsControlAsset)
+	: FWorkflowTabFactory(FPhysicsControlAssetEditorControlSetsTabSummoner::TabName, InHostingApp)
+	, PhysicsControlAsset(InPhysicsControlAsset)
+{
+	TabLabel = LOCTEXT("PhysicsControlAssetEditorControlSetsTabTitle", "ControlSets");
+	TabIcon = FSlateIcon(FAppStyle::GetAppStyleSetName(), "PhysicsAssetEditor.Tabs.Profiles");
+
+	bIsSingleton = true;
+
+	ViewMenuDescription = LOCTEXT("PhysicsControlAssetEditorControlSets", "ControlSets");
+	ViewMenuTooltip = LOCTEXT("PhysicsControlAssetEditorControlSets_ToolTip", "Shows the Control Asset ControlSets tab");
+}
+
+//======================================================================================================================
+TSharedPtr<SToolTip> FPhysicsControlAssetEditorControlSetsTabSummoner::CreateTabToolTipWidget(
+	const FWorkflowTabSpawnInfo& Info) const
+{
+	return IDocumentation::Get()->CreateToolTip(LOCTEXT(
+		"PhysicsControlAssetEditorControlSetsToolTip",
+		"The Physics Control Asset Control Sets tab lets you see the control sets."),
+		NULL,
+		TEXT("Shared/Editors/PhysicsControlAssetEditor"), TEXT("PhysicsControlAssetProfiles_Window"));
+}
+
+//======================================================================================================================
+TSharedRef<SWidget> FPhysicsControlAssetEditorControlSetsTabSummoner::CreateTabBody(
+	const FWorkflowTabSpawnInfo& Info) const
+{
+	FDetailsViewArgs DetailsViewArgs;
+	DetailsViewArgs.NameAreaSettings = FDetailsViewArgs::HideNameArea;
+	DetailsViewArgs.bHideSelectionTip = true;
+	DetailsViewArgs.bAllowSearch = false;
+
+	FPropertyEditorModule& PropertyEditorModule =
+		FModuleManager::LoadModuleChecked<FPropertyEditorModule>("PropertyEditor");
+	TSharedRef<IDetailsView> DetailsView = PropertyEditorModule.CreateDetailView(DetailsViewArgs);
+
+	TWeakPtr<FPhysicsControlAssetEditor> PhysicsControlAssetEditor =
+		StaticCastSharedPtr<FPhysicsControlAssetEditor>(HostingApp.Pin());
+	DetailsView->RegisterInstancedCustomPropertyLayout(
+		UPhysicsControlAsset::StaticClass(),
+		FOnGetDetailCustomizationInstance::CreateStatic(
+			&FPhysicsControlAssetInfoDetailsCustomization::MakeInstance, 
+			PhysicsControlAssetEditor,
+			FPhysicsControlAssetInfoDetailsCustomization::EInfoType::Controls));
+	DetailsView->SetObject(PhysicsControlAsset.Get());
+	return DetailsView;
+}
+
+//======================================================================================================================
+FPhysicsControlAssetEditorBodyModifierSetsTabSummoner::FPhysicsControlAssetEditorBodyModifierSetsTabSummoner(
+	TSharedPtr<FAssetEditorToolkit> InHostingApp, UPhysicsControlAsset* InPhysicsControlAsset)
+	: FWorkflowTabFactory(FPhysicsControlAssetEditorBodyModifierSetsTabSummoner::TabName, InHostingApp)
+	, PhysicsControlAsset(InPhysicsControlAsset)
+{
+	TabLabel = LOCTEXT("PhysicsControlAssetEditorBodyModifierSetsTabTitle", "BodyModifierSets");
+	TabIcon = FSlateIcon(FAppStyle::GetAppStyleSetName(), "PhysicsAssetEditor.Tabs.Profiles");
+
+	bIsSingleton = true;
+
+	ViewMenuDescription = LOCTEXT("PhysicsControlAssetEditorBodyModifierSets", "BodyModifierSets");
+	ViewMenuTooltip = LOCTEXT("PhysicsControlAssetEditorBodyModifierSets_ToolTip", "Shows the Control Asset BodyModifierSets tab");
+}
+
+//======================================================================================================================
+TSharedPtr<SToolTip> FPhysicsControlAssetEditorBodyModifierSetsTabSummoner::CreateTabToolTipWidget(
+	const FWorkflowTabSpawnInfo& Info) const
+{
+	return IDocumentation::Get()->CreateToolTip(LOCTEXT(
+		"PhysicsControlAssetEditorBodyModifierSetsToolTip",
+		"The Physics Control Asset Control Sets tab lets you see the body modifier sets."),
+		NULL,
+		TEXT("Shared/Editors/PhysicsControlAssetEditor"), TEXT("PhysicsControlAssetProfiles_Window"));
+}
+
+//======================================================================================================================
+TSharedRef<SWidget> FPhysicsControlAssetEditorBodyModifierSetsTabSummoner::CreateTabBody(
+	const FWorkflowTabSpawnInfo& Info) const
+{
+	FDetailsViewArgs DetailsViewArgs;
+	DetailsViewArgs.NameAreaSettings = FDetailsViewArgs::HideNameArea;
+	DetailsViewArgs.bHideSelectionTip = true;
+	DetailsViewArgs.bAllowSearch = false;
+
+	FPropertyEditorModule& PropertyEditorModule =
+		FModuleManager::LoadModuleChecked<FPropertyEditorModule>("PropertyEditor");
+	TSharedRef<IDetailsView> DetailsView = PropertyEditorModule.CreateDetailView(DetailsViewArgs);
+
+	TWeakPtr<FPhysicsControlAssetEditor> PhysicsControlAssetEditor =
+		StaticCastSharedPtr<FPhysicsControlAssetEditor>(HostingApp.Pin());
+	DetailsView->RegisterInstancedCustomPropertyLayout(
+		UPhysicsControlAsset::StaticClass(),
+		FOnGetDetailCustomizationInstance::CreateStatic(
+			&FPhysicsControlAssetInfoDetailsCustomization::MakeInstance,
+			PhysicsControlAssetEditor,
+			FPhysicsControlAssetInfoDetailsCustomization::EInfoType::BodyModifiers));
+	DetailsView->SetObject(PhysicsControlAsset.Get());
+	return DetailsView;
+}
+
+
+#undef LOCTEXT_NAMESPACE
\ No newline at end of file
diff --git a/Engine/Plugins/Experimental/PhysicsControl/Source/PhysicsControlEditor/Private/PhysicsControlAssetEditorTabSummoners.h b/Engine/Plugins/Experimental/PhysicsControl/Source/PhysicsControlEditor/Private/PhysicsControlAssetEditorTabSummoners.h
new file mode 100644
--- /dev/null
+++ b/Engine/Plugins/Experimental/PhysicsControl/Source/PhysicsControlEditor/Private/PhysicsControlAssetEditorTabSummoners.h
@@ -0,0 +1,92 @@
+// Copyright Epic Games, Inc. All Rights Reserved.
+
+#pragma once
+
+#include "WorkflowOrientedApp/WorkflowTabFactory.h"
+
+class UPhysicsControlAsset;
+
+//======================================================================================================================
+struct FPhysicsControlAssetEditorSetupTabSummoner : public FWorkflowTabFactory
+{
+public:
+	static FName TabName;
+
+	FPhysicsControlAssetEditorSetupTabSummoner(
+		TSharedPtr<FAssetEditorToolkit> InHostingApp, UPhysicsControlAsset* InPhysicsControlAsset);
+
+	/** FWorkflowTabFactory interface */
+	virtual TSharedRef<SWidget> CreateTabBody(const FWorkflowTabSpawnInfo& Info) const override;
+	virtual TSharedPtr<SToolTip> CreateTabToolTipWidget(const FWorkflowTabSpawnInfo& Info) const override;
+
+	/** Reference to our Physics Asset */
+	TWeakObjectPtr<UPhysicsControlAsset> PhysicsControlAsset;
+};
+
+//======================================================================================================================
+struct FPhysicsControlAssetEditorProfileTabSummoner : public FWorkflowTabFactory
+{
+public:
+	static FName TabName;
+
+	FPhysicsControlAssetEditorProfileTabSummoner(
+		TSharedPtr<FAssetEditorToolkit> InHostingApp, UPhysicsControlAsset* InPhysicsControlAsset);
+
+	/** FWorkflowTabFactory interface */
+	virtual TSharedRef<SWidget> CreateTabBody(const FWorkflowTabSpawnInfo& Info) const override;
+	virtual TSharedPtr<SToolTip> CreateTabToolTipWidget(const FWorkflowTabSpawnInfo& Info) const override;
+
+	/** Reference to our Physics Asset */
+	TWeakObjectPtr<UPhysicsControlAsset> PhysicsControlAsset;
+};
+
+//======================================================================================================================
+struct FPhysicsControlAssetEditorPreviewTabSummoner : public FWorkflowTabFactory
+{
+public:
+	static FName TabName;
+
+	FPhysicsControlAssetEditorPreviewTabSummoner(
+		TSharedPtr<FAssetEditorToolkit> InHostingApp, UPhysicsControlAsset* InPhysicsControlAsset);
+
+	/** FWorkflowTabFactory interface */
+	virtual TSharedRef<SWidget> CreateTabBody(const FWorkflowTabSpawnInfo& Info) const override;
+	virtual TSharedPtr<SToolTip> CreateTabToolTipWidget(const FWorkflowTabSpawnInfo& Info) const override;
+
+	/** Reference to our Physics Asset */
+	TWeakObjectPtr<UPhysicsControlAsset> PhysicsControlAsset;
+};
+
+//======================================================================================================================
+struct FPhysicsControlAssetEditorControlSetsTabSummoner : public FWorkflowTabFactory
+{
+public:
+	static FName TabName;
+
+	FPhysicsControlAssetEditorControlSetsTabSummoner(
+		TSharedPtr<FAssetEditorToolkit> InHostingApp, UPhysicsControlAsset* InPhysicsControlAsset);
+
+	/** FWorkflowTabFactory interface */
+	virtual TSharedRef<SWidget> CreateTabBody(const FWorkflowTabSpawnInfo& Info) const override;
+	virtual TSharedPtr<SToolTip> CreateTabToolTipWidget(const FWorkflowTabSpawnInfo& Info) const override;
+
+	/** Reference to our Physics Asset */
+	TWeakObjectPtr<UPhysicsControlAsset> PhysicsControlAsset;
+};
+
+//======================================================================================================================
+struct FPhysicsControlAssetEditorBodyModifierSetsTabSummoner : public FWorkflowTabFactory
+{
+public:
+	static FName TabName;
+
+	FPhysicsControlAssetEditorBodyModifierSetsTabSummoner(
+		TSharedPtr<FAssetEditorToolkit> InHostingApp, UPhysicsControlAsset* InPhysicsControlAsset);
+
+	/** FWorkflowTabFactory interface */
+	virtual TSharedRef<SWidget> CreateTabBody(const FWorkflowTabSpawnInfo& Info) const override;
+	virtual TSharedPtr<SToolTip> CreateTabToolTipWidget(const FWorkflowTabSpawnInfo& Info) const override;
+
+	/** Reference to our Physics Asset */
+	TWeakObjectPtr<UPhysicsControlAsset> PhysicsControlAsset;
+};
diff --git a/Engine/Plugins/Experimental/PhysicsControl/Source/PhysicsControlEditor/Private/PhysicsControlAssetInfoDetailsCustomization.cpp b/Engine/Plugins/Experimental/PhysicsControl/Source/PhysicsControlEditor/Private/PhysicsControlAssetInfoDetailsCustomization.cpp
new file mode 100644
--- /dev/null
+++ b/Engine/Plugins/Experimental/PhysicsControl/Source/PhysicsControlEditor/Private/PhysicsControlAssetInfoDetailsCustomization.cpp
@@ -0,0 +1,131 @@
+// Copyright Epic Games, Inc. All Rights Reserved.
+
+#include "PhysicsControlAssetInfoDetailsCustomization.h"
+
+#include "PhysicsControlAssetActions.h"
+#include "PhysicsControlAssetEditor.h"
+#include "PhysicsControlAsset.h"
+#include "PhysicsControlAssetEditorData.h"
+#include "PhysicsControlAssetEditorSkeletalMeshComponent.h"
+#include "PhysicsControlOperatorNameGeneration.h"
+#include "PhysicsControlNameRecords.h"
+
+#include "DetailLayoutBuilder.h"
+#include "DetailWidgetRow.h"
+#include "Framework/Application/SlateApplication.h"
+#include "IDetailPropertyRow.h"
+#include "PropertyHandle.h"
+
+#define LOCTEXT_NAMESPACE "PhysicsControlAssetInfoDetailsCustomization"
+
+//======================================================================================================================
+TSharedRef<IDetailCustomization> FPhysicsControlAssetInfoDetailsCustomization::MakeInstance(
+	TWeakPtr<FPhysicsControlAssetEditor> InPhysicsControlAssetEditor, EInfoType InfoType)
+{
+	return MakeShared<FPhysicsControlAssetInfoDetailsCustomization>(InPhysicsControlAssetEditor, InfoType);
+}
+
+//======================================================================================================================
+void FPhysicsControlAssetInfoDetailsCustomization::CustomizeDetails(
+	const TSharedPtr<IDetailLayoutBuilder>& InDetailBuilder)
+{
+	DetailLayoutBuilderWeak = InDetailBuilder;
+	FPhysicsControlAssetInfoDetailsCustomization::CustomizeDetails(*InDetailBuilder);
+
+	if (TSharedPtr<FPhysicsControlAssetEditor> PCAE = PhysicsControlAssetEditor.Pin())
+	{
+		if (UPhysicsControlAsset* PhysicsControlAsset = PCAE->GetEditorData()->PhysicsControlAsset.Get())
+		{
+			PhysicsControlAsset->OnControlAssetCompiled().AddSP(
+				this, &FPhysicsControlAssetInfoDetailsCustomization::OnControlAssetCompiled);
+		}
+	}
+}
+
+//======================================================================================================================
+void FPhysicsControlAssetInfoDetailsCustomization::CustomizeDetails(IDetailLayoutBuilder& DetailLayoutBuilder)
+{
+	TArray<FName> CategoryNames;
+	DetailLayoutBuilder.GetCategoryNames(CategoryNames);
+	for (FName Category : CategoryNames)
+	{
+		DetailLayoutBuilder.HideCategory(Category);
+	}
+	DetailLayoutBuilder.HideCategory(TEXT("Actions"));
+
+	TSharedPtr<const FPhysicsControlAssetEditor> PCAE = PhysicsControlAssetEditor.Pin();
+	TSharedPtr<const FPhysicsControlAssetEditorData> EditorData = PCAE->GetEditorData();
+	const UPhysicsControlAsset* PCA = EditorData->PhysicsControlAsset.Get();
+	if (!PCA)
+	{
+		return;
+	}
+	const UPhysicsAsset* PA = PCA->GetPhysicsAsset();
+	const USkeletalMeshComponent* SKMC = EditorData->EditorSkelComp.Get();
+	if (!PA || !SKMC)
+	{
+		return;
+	}
+	const USkeletalMesh* SKM = SKMC->GetSkeletalMeshAsset();
+	if (!SKM)
+	{
+		return;
+	}
+	const FReferenceSkeleton& RefSkeleton = SKM->GetRefSkeleton();
+
+	// Process the asset to get all the controls/body modifiers and sets
+	TMap<FName, FPhysicsControlLimbBones> LimbBones = 
+		UE::PhysicsControl::GetLimbBones(PCA->CharacterSetupData.LimbSetupData,RefSkeleton, PA);
+
+	TSet<FName> BodyModifierNames;
+	TSet<FName> ControlNames;
+	FPhysicsControlNameRecords NameRecords;
+
+	UE::PhysicsControl::CollectOperatorNames(
+		PCA->CharacterSetupData, PCA->AdditionalControlsAndModifiers,
+		LimbBones, RefSkeleton, PA, BodyModifierNames, ControlNames, NameRecords);
+
+	UE::PhysicsControl::CreateAdditionalSets(
+		PCA->AdditionalSets, BodyModifierNames, ControlNames, NameRecords);
+
+	TMap<FName, TArray<FName>>& SetsToShow = 
+		(InfoType == EInfoType::Controls) ? NameRecords.ControlSets : NameRecords.BodyModifierSets;
+
+	SetsToShow.KeySort([](FName A, FName B) { return A.ToString() < B.ToString(); });
+
+	for (const TPair<FName, TArray<FName>>& Set : SetsToShow)
+	{
+		const FName SetName = Set.Key;
+		const TArray<FName>& Names = Set.Value;
+
+		FText SectionName = FText::FromName(SetName);
+		IDetailCategoryBuilder& DetailCategoryBuilder = 
+			DetailLayoutBuilder.EditCategory(SetName, FText::FromName(SetName));
+		DetailCategoryBuilder.InitiallyCollapsed(true);
+
+		for (FName Name : Names)
+		{
+			FDetailWidgetRow& Row = DetailCategoryBuilder.AddCustomRow(
+				FText::Format(LOCTEXT("NameRow", "Name_{0}"), FText::FromName(Name)));
+			Row.WholeRowContent()
+				[
+					SNew(STextBlock).Text(FText::FromName(Name))
+				];
+		}
+	}
+}
+
+//======================================================================================================================
+void FPhysicsControlAssetInfoDetailsCustomization::OnControlAssetCompiled(bool bProfileListChanged)
+{
+	if (DetailLayoutBuilderWeak.IsValid())
+	{
+		if (IDetailLayoutBuilder* DetailLayoutBuilder = DetailLayoutBuilderWeak.Pin().Get())
+		{
+			DetailLayoutBuilder->ForceRefreshDetails();
+		}
+	}
+}
+
+
+#undef LOCTEXT_NAMESPACE
diff --git a/Engine/Plugins/Experimental/PhysicsControl/Source/PhysicsControlEditor/Private/PhysicsControlAssetInfoDetailsCustomization.h b/Engine/Plugins/Experimental/PhysicsControl/Source/PhysicsControlEditor/Private/PhysicsControlAssetInfoDetailsCustomization.h
new file mode 100644
--- /dev/null
+++ b/Engine/Plugins/Experimental/PhysicsControl/Source/PhysicsControlEditor/Private/PhysicsControlAssetInfoDetailsCustomization.h
@@ -0,0 +1,51 @@
+// Copyright Epic Games, Inc. All Rights Reserved.
+
+#pragma once
+
+#include "CoreMinimal.h"
+#include "IDetailCustomization.h"
+#include "Types/SlateEnums.h"
+
+class FPhysicsControlAssetEditor;
+class IDetailLayoutBuilder;
+class FUICommandInfo;
+class SWidget;
+class IPropertyHandle;
+class FUICommandList;
+class SEditableTextBox;
+
+class FPhysicsControlAssetInfoDetailsCustomization : public IDetailCustomization
+{
+public:
+	enum class EInfoType : uint8
+	{
+		Controls,
+		BodyModifiers
+	};
+
+	// Makes a new instance of this detail layout class for a specific detail view requesting it
+	static TSharedRef<IDetailCustomization> MakeInstance(
+		TWeakPtr<FPhysicsControlAssetEditor> InPhysicsControlAssetEditor, EInfoType InfoType);
+
+	FPhysicsControlAssetInfoDetailsCustomization(
+		TWeakPtr<FPhysicsControlAssetEditor> InPhysicsControlAssetEditor,
+		EInfoType                            InInfoType)
+		: PhysicsControlAssetEditor(InPhysicsControlAssetEditor)
+		, InfoType(InInfoType)
+	{}
+
+	//~ Begin IDetailCustomization
+	virtual void CustomizeDetails(IDetailLayoutBuilder& DetailLayout) override;
+	virtual void CustomizeDetails(const TSharedPtr<IDetailLayoutBuilder>& InDetailBuilder) override;
+	//~ End IDetailCustomization
+
+protected:
+	void OnControlAssetCompiled(bool bProfileListChanged);
+
+protected:
+	TWeakPtr<FPhysicsControlAssetEditor> PhysicsControlAssetEditor;
+
+	TWeakPtr<IDetailLayoutBuilder> DetailLayoutBuilderWeak;
+
+	EInfoType InfoType;
+};
diff --git a/Engine/Plugins/Experimental/PhysicsControl/Source/PhysicsControlEditor/Private/PhysicsControlAssetPreviewDetailsCustomization.cpp b/Engine/Plugins/Experimental/PhysicsControl/Source/PhysicsControlEditor/Private/PhysicsControlAssetPreviewDetailsCustomization.cpp
--- a/Engine/Plugins/Experimental/PhysicsControl/Source/PhysicsControlEditor/Private/PhysicsControlAssetPreviewDetailsCustomization.cpp
+++ b/Engine/Plugins/Experimental/PhysicsControl/Source/PhysicsControlEditor/Private/PhysicsControlAssetPreviewDetailsCustomization.cpp
@@ -53,13 +53,13 @@ void FPhysicsControlAssetPreviewDetailsCustomization::CustomizeDetails(
 //======================================================================================================================
 void FPhysicsControlAssetPreviewDetailsCustomization::CustomizeDetails(IDetailLayoutBuilder& DetailLayoutBuilder)
 {
-	DetailLayoutBuilder.HideCategory(TEXT("PreviewMesh"));
+	TArray<FName> CategoryNames;
+	DetailLayoutBuilder.GetCategoryNames(CategoryNames);
+	for (FName Category : CategoryNames)
+	{
+		DetailLayoutBuilder.HideCategory(Category);
+	}
 	DetailLayoutBuilder.HideCategory(TEXT("Actions"));
-	DetailLayoutBuilder.HideCategory(TEXT("Inheritance"));
-	DetailLayoutBuilder.HideCategory(TEXT("Setup"));
-	DetailLayoutBuilder.HideCategory(TEXT("Profiles"));
-	DetailLayoutBuilder.HideCategory(TEXT("ProfileEditing"));
-	DetailLayoutBuilder.HideCategory(TEXT("SetupEditing"));
 
 	TSharedPtr<FPhysicsControlAssetEditorData> EditorData = PhysicsControlAssetEditor.Pin()->GetEditorData();
 	UPhysicsControlAsset* PCA = EditorData->PhysicsControlAsset.Get();
diff --git a/Engine/Plugins/Experimental/PhysicsControl/Source/PhysicsControlEditor/Private/PhysicsControlAssetPreviewDetailsCustomization.h b/Engine/Plugins/Experimental/PhysicsControl/Source/PhysicsControlEditor/Private/PhysicsControlAssetPreviewDetailsCustomization.h
--- a/Engine/Plugins/Experimental/PhysicsControl/Source/PhysicsControlEditor/Private/PhysicsControlAssetPreviewDetailsCustomization.h
+++ b/Engine/Plugins/Experimental/PhysicsControl/Source/PhysicsControlEditor/Private/PhysicsControlAssetPreviewDetailsCustomization.h
@@ -28,7 +28,7 @@ public:
 	virtual void CustomizeDetails(const TSharedPtr<IDetailLayoutBuilder>& InDetailBuilder) override;
 	//~ End IDetailCustomization
 
-private:
+protected:
 	FReply InvokeControlProfile(FName ProfileName);
 
 	void OnControlAssetCompiled(bool bProfileListChanged);
diff --git a/Engine/Plugins/Experimental/PhysicsControl/Source/PhysicsControlEditor/Public/PhysicsControlAssetEditor.h b/Engine/Plugins/Experimental/PhysicsControl/Source/PhysicsControlEditor/Public/PhysicsControlAssetEditor.h
--- a/Engine/Plugins/Experimental/PhysicsControl/Source/PhysicsControlEditor/Public/PhysicsControlAssetEditor.h
+++ b/Engine/Plugins/Experimental/PhysicsControl/Source/PhysicsControlEditor/Public/PhysicsControlAssetEditor.h
@@ -40,6 +40,7 @@ public:
 	friend class FPhysicsControlAssetProfileDetailsCustomization;
 	friend class FPhysicsControlAssetPreviewDetailsCustomization;
 	friend class FPhysicsControlAssetSetupDetailsCustomization;
+	friend class FPhysicsControlAssetInfoDetailsCustomization;
 
 public:
 	/** Initialize the asset editor. This will register the application mode, init the preview scene, etc. */
diff --git a/Engine/Plugins/Experimental/PhysicsControl/Source/PhysicsControlUncookedOnly/Private/AnimGraphNode_RigidBodyWithControl.cpp b/Engine/Plugins/Experimental/PhysicsControl/Source/PhysicsControlUncookedOnly/Private/AnimGraphNode_RigidBodyWithControl.cpp
--- a/Engine/Plugins/Experimental/PhysicsControl/Source/PhysicsControlUncookedOnly/Private/AnimGraphNode_RigidBodyWithControl.cpp
+++ b/Engine/Plugins/Experimental/PhysicsControl/Source/PhysicsControlUncookedOnly/Private/AnimGraphNode_RigidBodyWithControl.cpp
@@ -360,7 +360,7 @@ TArray<TPair<FName, TArray<FName>>> UAnimGraphNode_RigidBodyWithControl::Generat
 
 		// Get the list of modifier and control names, based on the setup data
 		UE::PhysicsControl::CollectOperatorNames(
-			&Node, SetupData, AdditionalControlAndBodyModifierCreationDatas,
+			SetupData, AdditionalControlAndBodyModifierCreationDatas,
 			AllLimbBones, RefSkeleton, Node.OverridePhysicsAsset.Get(), BodyModifierNames, ControlNames, NameRecords);
 
 		// Create any additional sets that have been requested
