Ensure the RigidBodyWithControl pose cache is big enough

diff --git a/Engine/Plugins/Experimental/PhysicsControl/Source/PhysicsControl/Private/RigidBodyPoseData.cpp b/Engine/Plugins/Experimental/PhysicsControl/Source/PhysicsControl/Private/RigidBodyPoseData.cpp
--- a/Engine/Plugins/Experimental/PhysicsControl/Source/PhysicsControl/Private/RigidBodyPoseData.cpp
+++ b/Engine/Plugins/Experimental/PhysicsControl/Source/PhysicsControl/Private/RigidBodyPoseData.cpp
@@ -28,11 +28,24 @@ void FRigidBodyPoseData::Update(
 
 	for (const FOutputBoneData& OutputData : OutputBoneData)
 	{
+		// It is very unusual, but possible that BodyIndex is invalid - in particular that it is too
+		// big. This can happen when OutputBoneData has changed in size and we haven't been
+		// reinitialized. In this edge case, we could simply refuse to calculate TMs, but there's no
+		// harm in simply expanding our cache array and continuing to function. See UE-214162
 		const int32 BodyIndex = OutputData.BodyIndex;
-		if (ensure(BoneTMs.IsValidIndex(BodyIndex)))
+		if (BodyIndex >= 0)
 		{
-			const FTransform& ComponentSpaceTM = ComponentSpacePoseContext.Pose.GetComponentSpaceTransform(OutputData.CompactPoseBoneIndex);
-			const FTransform BodyTM = ConvertCSTransformToSimSpace(SimulationSpace, ComponentSpaceTM, CompWorldSpaceTM, BaseBoneTM);
+			if (BodyIndex >= BoneTMs.Num())
+			{
+				// This could cause multiple re-allocations if we keep finding a body index that is
+				// too big, but the situation will be so rare that it's not a significant problem
+				// (and would only happen for one frame).
+				BoneTMs.SetNumUninitialized(BodyIndex + 1);
+			}
+			const FTransform& ComponentSpaceTM = 
+				ComponentSpacePoseContext.Pose.GetComponentSpaceTransform(OutputData.CompactPoseBoneIndex);
+			const FTransform BodyTM = ConvertCSTransformToSimSpace(
+				SimulationSpace, ComponentSpaceTM, CompWorldSpaceTM, BaseBoneTM);
 			BoneTMs[BodyIndex] = BodyTM;
 		}
 	}
