Hide the PhysicsControlViewer behind a CVar p.RigidBodyWithControl.EnableControlSetViewer which is disable by default. This also renames the relevant classes to clarify that this is a viewer, not an editor, and also to avoid duplicating relevant modular feature names through the code.

diff --git a/Engine/Plugins/Experimental/PhysicsControl/Source/PhysicsControlEditor/Private/OperatorEditor/OperatorEditor.cpp b/Engine/Plugins/Experimental/PhysicsControl/Source/PhysicsControlEditor/Private/OperatorEditor/OperatorEditor.cpp
deleted file mode 100644
diff --git a/Engine/Plugins/Experimental/PhysicsControl/Source/PhysicsControlEditor/Private/OperatorEditor/OperatorTreeElements.cpp b/Engine/Plugins/Experimental/PhysicsControl/Source/PhysicsControlEditor/Private/OperatorEditor/OperatorTreeElements.cpp
deleted file mode 100644
diff --git a/Engine/Plugins/Experimental/PhysicsControl/Source/PhysicsControlEditor/Private/OperatorEditor/SOperatorEditorTabWidget.cpp b/Engine/Plugins/Experimental/PhysicsControl/Source/PhysicsControlEditor/Private/OperatorEditor/SOperatorEditorTabWidget.cpp
deleted file mode 100644
diff --git a/Engine/Plugins/Experimental/PhysicsControl/Source/PhysicsControlEditor/Private/OperatorEditor/SOperatorTreeWidget.cpp b/Engine/Plugins/Experimental/PhysicsControl/Source/PhysicsControlEditor/Private/OperatorEditor/SOperatorTreeWidget.cpp
deleted file mode 100644
diff --git a/Engine/Plugins/Experimental/PhysicsControl/Source/PhysicsControlEditor/Private/OperatorViewer/OperatorTreeElements.cpp b/Engine/Plugins/Experimental/PhysicsControl/Source/PhysicsControlEditor/Private/OperatorViewer/OperatorTreeElements.cpp
new file mode 100644
--- /dev/null
+++ b/Engine/Plugins/Experimental/PhysicsControl/Source/PhysicsControlEditor/Private/OperatorViewer/OperatorTreeElements.cpp
@@ -0,0 +1,349 @@
+// Copyright Epic Games, Inc. All Rights Reserved.
+
+#include "OperatorViewer/OperatorTreeElements.h"
+
+#include "AnimGraphNode_RigidBodyWithControl.h"
+#include "BlueprintEditorModule.h"
+#include "Input/Reply.h"
+#include "Kismet2/BlueprintEditorUtils.h"
+#include "Styling/SlateColor.h"
+#include "Types/SlateEnums.h"
+#include "Widgets/Layout/SWrapBox.h"
+#include "Widgets/Text/STextBlock.h"
+#include "Widgets/Views/STableRow.h"
+#include "Widgets/Views/STreeView.h"
+
+// UE_DISABLE_OPTIMIZATION;
+
+#define LOCTEXT_NAMESPACE "PhysicsControlOperatorViewer"
+
+// class OperatorTreeItem //
+const EHorizontalAlignment OperatorTreeItem::HorizontalAlignment = HAlign_Left;
+const EVerticalAlignment OperatorTreeItem::VerticalAlignment = VAlign_Top;
+const float OperatorTreeItem::HorizontalPadding = 5.0f;
+const float OperatorTreeItem::VerticalPadding = 1.0f;
+
+static const FName OperatorTagControl("Control");
+static const FName OperatorTagModifier("Modifier");
+
+bool MatchSearchText(const FString& TextToSearch, const TArray<TArray<FString>>& StructuredSearchCriteria)
+{
+	for (const TArray<FString>& AndCriteria : StructuredSearchCriteria)
+	{
+		bool bMatchFound = true;
+
+		for (const FString& Token : AndCriteria)
+		{
+			bMatchFound &= TextToSearch.Contains(Token);
+		}
+
+		if (bMatchFound)
+		{
+			return true;
+		}
+	}
+
+	return false;
+}
+
+// class OperatorTreeItem //
+FReply OperatorTreeItem::OnDoubleClick()
+{ 
+	return FReply::Unhandled(); 
+}
+
+// class OperatorTreeControlItem //
+const TArray<FName> OperatorTreeControlItem::HiddenTags = { OperatorTagControl, OperatorTagModifier };
+
+FText OperatorTreeControlItem::Name() const
+{
+	return FText::FromName(NameAndTags.Key);
+}
+
+FText OperatorTreeControlItem::Description() const
+{
+	FString DescriptionStr;
+	
+	DescriptionStr += OperatorTypeName().ToString();
+	DescriptionStr += FString(" ");
+	DescriptionStr += NameAndTags.Key.ToString();
+
+	for (const FName TagName : NameAndTags.Value)
+	{
+		if (!HiddenTags.Contains(TagName))
+		{
+			DescriptionStr += " " + TagName.ToString();
+		}
+	}
+
+	DescriptionStr.TrimEndInline();
+
+	return FText::FromString(DescriptionStr);
+}
+
+FName OperatorTreeControlItem::OperatorTypeName() const
+{
+	if (HasTag(OperatorTagControl))
+	{
+		return OperatorTagControl;
+	}
+
+	if (HasTag(OperatorTagModifier))
+	{
+		return OperatorTagModifier;
+	}
+
+	return NAME_None;
+}
+
+TSharedRef<ITableRow> OperatorTreeControlItem::GenerateRow(const TSharedRef<STableViewBase>& OwnerTable) const
+{
+	TSharedPtr<SHorizontalBox> HorizontalBox = SNew(SHorizontalBox);
+
+	// Display Operator Type Indicator
+	{
+		FText OperatorTypeLabel;
+		FLinearColor OperatorTypeColor;
+
+		if (HasTag(OperatorTagControl))
+		{
+			OperatorTypeLabel = LOCTEXT("Control", "Control");
+			OperatorTypeColor = FLinearColor::Yellow;
+		}
+
+		if (HasTag(OperatorTagModifier))
+		{
+			OperatorTypeLabel = LOCTEXT("Modifier", "Modifier");
+			OperatorTypeColor = FLinearColor::Red;
+		}
+
+		HorizontalBox->AddSlot()
+			.HAlign(HorizontalAlignment)
+			.Padding(HorizontalPadding, VerticalPadding)
+			.AutoWidth()
+			[
+				SNew(STextBlock)
+				.Text(OperatorTypeLabel)
+				.ColorAndOpacity(OperatorTypeColor)
+			];
+	}
+
+	// Display Operator Name
+	{
+		HorizontalBox->AddSlot()
+			.HAlign(HorizontalAlignment)
+			.VAlign(VerticalAlignment)
+			.Padding(HorizontalPadding, VerticalPadding)
+			.FillWidth(0.4)
+			[
+				SNew(STextBlock)
+				.Text(FText::FromName(NameAndTags.Key))
+			];
+	}
+
+	// Display Sets
+	{
+		TSharedPtr<SWrapBox> TagBox = SNew(SWrapBox)
+			.UseAllottedSize(true)
+			.InnerSlotPadding(FVector2D(4, 4));
+
+		for (const FName TagName : NameAndTags.Value)
+		{
+			if (!HiddenTags.Contains(TagName))
+			{
+				TagBox->AddSlot()
+					[
+						SNew(STextBlock)
+						.Text(FText::FromName(TagName))
+						.ColorAndOpacity(FLinearColor::White)
+					];
+			}
+		}
+
+		TagBox->SetVisibility(EVisibility::Visible);
+
+		HorizontalBox->AddSlot()
+			.HAlign(HAlign_Fill)
+			.Padding(HorizontalPadding, VerticalPadding)
+			[
+				TagBox.ToSharedRef()
+			];
+	}
+
+	return SNew(STableRow<TSharedPtr<OperatorTreeItem>>, OwnerTable)
+		[
+			HorizontalBox.ToSharedRef()
+		];
+}
+
+bool OperatorTreeControlItem::MatchesFilterText(const TArray<TArray<FString>>& FilterStructuredCriteria) const
+{
+	if (FilterStructuredCriteria.IsEmpty())
+	{
+		return true;
+	}
+
+	FString TextToSearch = NameAndTags.Key.ToString();
+
+	for (const FName Tag : NameAndTags.Value)
+	{
+		TextToSearch += Tag.ToString();
+	}
+
+	return MatchSearchText(TextToSearch, FilterStructuredCriteria);
+}
+
+bool OperatorTreeControlItem::HasTag(const FName Tag) const
+{
+	return NameAndTags.Value.Contains(Tag);
+}
+
+// class OperatorTreeNodeItem //
+OperatorTreeNodeItem::OperatorTreeNodeItem(const UEdGraphNode* InNode, const FName& InName, const FString& InBlueprintName)
+: Node(InNode)
+{
+	const FString Name = Node->GetNodeTitle(ENodeTitleType::FullTitle).ToString() + " [" + InBlueprintName + "]";
+	NodeDisplayName = FText::FromString(Name);
+}
+
+FText OperatorTreeNodeItem::Name() const
+{
+	return NodeDisplayName;
+}
+
+TSharedRef<ITableRow> OperatorTreeNodeItem::GenerateRow(const TSharedRef<STableViewBase>& OwnerTable) const
+{
+	return SNew(STableRow<TSharedPtr<OperatorTreeItem>>, OwnerTable)
+		[
+			SNew(SHorizontalBox)
+			+ SHorizontalBox::Slot()
+			[
+				SNew(STextBlock)
+				.Text(NodeDisplayName) // TODO - update on changes
+				.ColorAndOpacity_Lambda([this]() { return this->ColorAndOpacity(); })
+			]
+		];
+}
+
+bool OperatorTreeNodeItem::IsSelectedInGraph() const
+{
+	return Node && Node->IsSelected();
+}
+
+FSlateColor OperatorTreeNodeItem::ColorAndOpacity() const
+{
+	return IsSelectedInGraph() ? FAppStyle::Get().GetSlateColor("Colors.AccentOrange") : FAppStyle::Get().GetSlateColor("Colors.White");
+}
+
+FReply OperatorTreeNodeItem::OnDoubleClick()
+{
+	if (UBlueprint* const Blueprint = FBlueprintEditorUtils::FindBlueprintForNode(Node))
+	{
+		if (IBlueprintEditor* BlueprintEditor = static_cast<IBlueprintEditor*>(GEditor->GetEditorSubsystem<UAssetEditorSubsystem>()->FindEditorForAsset(Blueprint, true)))
+		{
+			BlueprintEditor->JumpToHyperlink(Node, false);
+			return FReply::Handled();
+		}
+	}
+
+	return FReply::Unhandled();
+}
+
+bool operator==(const OperatorTreeNodeItem& Lhs, const OperatorTreeNodeItem& Rhs)
+{
+	return Lhs.GraphNode()->GetFName().IsEqual(Rhs.GraphNode()->GetFName()); // note: must use IsEqual as operator == doesn't account for the FName's number postfix.
+}
+
+// class OperatorTreeMessageItem //
+OperatorTreeMessageItem::OperatorTreeMessageItem(const UEdGraphNode* InNode, const MessageType InType, const FText& InText)
+	: Text(InText)
+	, Node(InNode)
+	, Type(InType)
+{}
+
+FText OperatorTreeMessageItem::Name() const
+{
+	return Text;
+}
+
+TSharedRef<ITableRow> OperatorTreeMessageItem::GenerateRow(const TSharedRef<STableViewBase>& OwnerTable) const
+{
+	FText MessageTypeLabel;
+	FLinearColor MessageTypeColor;
+
+	if (Type == MessageType::Log)
+	{
+		MessageTypeLabel = LOCTEXT("MessageItemLog", "Log");
+		MessageTypeColor = FLinearColor::Gray;
+	}
+
+	if (Type == MessageType::Warning)
+	{
+		MessageTypeLabel = LOCTEXT("MessageItemWarning", "Warning");
+		MessageTypeColor = FLinearColor::Yellow;
+	}
+
+	if (Type == MessageType::Error)
+	{
+		MessageTypeLabel = LOCTEXT("MessageItemError", "Error");
+		MessageTypeColor = FLinearColor::Red;
+	}
+
+	return SNew(STableRow<TSharedPtr<OperatorTreeItem>>, OwnerTable)
+		[
+			SNew(SHorizontalBox)
+			+ SHorizontalBox::Slot()
+			.HAlign(HorizontalAlignment)
+			.Padding(HorizontalPadding, VerticalPadding)
+			.AutoWidth()
+			[
+				SNew(STextBlock)
+				.Text(MessageTypeLabel)
+				.ColorAndOpacity(MessageTypeColor)
+			]
+			+ SHorizontalBox::Slot()
+			.HAlign(HorizontalAlignment)
+			.Padding(HorizontalPadding, VerticalPadding)
+			.AutoWidth()
+			[
+				SNew(STextBlock)
+				.Text(Text)
+				.ColorAndOpacity(MessageTypeColor)
+			]
+		];
+}
+
+template<typename TDerivedType> bool IsEqualAfterCast(const TSharedPtr<OperatorTreeItem> Lhs, const TSharedPtr<OperatorTreeItem> Rhs)
+{
+	TSharedPtr<TDerivedType> LhsDerived = StaticCastSharedPtr<TDerivedType>(Lhs);
+	TSharedPtr<TDerivedType> RhsDerived = StaticCastSharedPtr<TDerivedType>(Rhs);
+
+	return LhsDerived && RhsDerived && (*LhsDerived == *RhsDerived);
+}
+
+bool IsEqual(const TSharedPtr<OperatorTreeItem> Lhs, const TSharedPtr<OperatorTreeItem> Rhs)
+{
+	if (Lhs->GetType() == Rhs->GetType())
+	{
+		const OperatorTreeItem::OperatorTreeItemType ItemType = Lhs->GetType();
+
+		if (ItemType == OperatorTreeControlItem::StaticType())
+		{
+			return IsEqualAfterCast<OperatorTreeControlItem>(Lhs, Rhs);
+		}
+
+		if (ItemType == OperatorTreeMessageItem::StaticType())
+		{
+			return IsEqualAfterCast<OperatorTreeMessageItem>(Lhs, Rhs);
+		}
+
+		if (ItemType == OperatorTreeNodeItem::StaticType())
+		{
+			return IsEqualAfterCast<OperatorTreeNodeItem>(Lhs, Rhs);
+		}
+	}
+
+	return false;
+}
+
+#undef LOCTEXT_NAMESPACE
diff --git a/Engine/Plugins/Experimental/PhysicsControl/Source/PhysicsControlEditor/Private/OperatorViewer/OperatorViewer.cpp b/Engine/Plugins/Experimental/PhysicsControl/Source/PhysicsControlEditor/Private/OperatorViewer/OperatorViewer.cpp
new file mode 100644
--- /dev/null
+++ b/Engine/Plugins/Experimental/PhysicsControl/Source/PhysicsControlEditor/Private/OperatorViewer/OperatorViewer.cpp
@@ -0,0 +1,102 @@
+// Copyright Epic Games, Inc. All Rights Reserved.
+
+#include "OperatorViewer/OperatorViewer.h"
+#include "OperatorViewer/SOperatorViewerTabWidget.h"
+#include "OperatorViewer/OperatorViewer.h"
+
+#include "Widgets/Docking/SDockTab.h"
+#include "Widgets/Layout/SBox.h"
+
+
+// UE_DISABLE_OPTIMIZATION;
+
+static const FName PhysicsControlEditorModule_OperatorNamesTabWidget("PhysicsControlEditorModule_OperatorNamesTabWidget");
+
+#define LOCTEXT_NAMESPACE "PhysicsControlOperatorViewer"
+
+void FPhysicsControlOperatorViewer::Startup()
+{
+	// Physics Operator Names Tool
+	FGlobalTabmanager::Get()->RegisterNomadTabSpawner(PhysicsControlEditorModule_OperatorNamesTabWidget, 
+		FOnSpawnTab::CreateRaw(this, &FPhysicsControlOperatorViewer::OnCreateTab))
+		.SetDisplayName(LOCTEXT("PhysicsAnimationEditor_OperatorNamesTabTitle", "Rigid Body With Control"))
+		.SetMenuType(ETabSpawnerMenuType::Hidden);
+}
+
+void FPhysicsControlOperatorViewer::Shutdown()
+{
+	FGlobalTabmanager::Get()->UnregisterNomadTabSpawner(PhysicsControlEditorModule_OperatorNamesTabWidget);
+}
+
+void FPhysicsControlOperatorViewer::OpenOperatorNamesTab()
+{
+	OperatorNamesTab = FGlobalTabmanager::Get()->FindExistingLiveTab(PhysicsControlEditorModule_OperatorNamesTabWidget);
+
+	if (!OperatorNamesTab)
+	{
+		OperatorNamesTab = FGlobalTabmanager::Get()->TryInvokeTab(PhysicsControlEditorModule_OperatorNamesTabWidget);
+	}
+}
+
+void FPhysicsControlOperatorViewer::CloseOperatorNamesTab()
+{
+	if (!OperatorNamesTab)
+	{
+		OperatorNamesTab = FGlobalTabmanager::Get()->FindExistingLiveTab(PhysicsControlEditorModule_OperatorNamesTabWidget);
+	}
+
+	OperatorNamesTab->RequestCloseTab();
+	OperatorNamesTab.Reset();
+}
+
+void FPhysicsControlOperatorViewer::ToggleOperatorNamesTab()
+{
+	if (IsOperatorNamesTabOpen())
+	{
+		CloseOperatorNamesTab();
+	}
+	else
+	{
+		OpenOperatorNamesTab();
+	}
+}
+
+bool FPhysicsControlOperatorViewer::IsOperatorNamesTabOpen()
+{
+	return OperatorNamesTab.IsValid();
+}
+
+void FPhysicsControlOperatorViewer::RequestRefresh()
+{
+	if (PersistantTabWidget)
+	{
+		PersistantTabWidget->RequestRefresh();
+	}
+}
+
+TSharedRef<SDockTab> FPhysicsControlOperatorViewer::OnCreateTab(const FSpawnTabArgs& SpawnTabArgs)
+{
+	const int TabIndex = 1;
+
+	PersistantTabWidget = SNew(SOperatorViewerTabWidget, TabIndex);
+
+	return SNew(SDockTab)
+		.TabRole(ETabRole::NomadTab)
+		.OnTabClosed_Lambda([this](TSharedRef<class SDockTab> InParentTab) { this->OnTabClosed(InParentTab); })
+		[
+			SNew(SBox)
+			.HAlign(HAlign_Fill)
+			.VAlign(VAlign_Fill)
+			[
+				PersistantTabWidget.ToSharedRef()
+			]
+		];
+}
+
+void FPhysicsControlOperatorViewer::OnTabClosed(TSharedRef<SDockTab> DockTab)
+{
+	OperatorNamesTab.Reset();
+	PersistantTabWidget.Reset();
+}
+
+#undef LOCTEXT_NAMESPACE
\ No newline at end of file
diff --git a/Engine/Plugins/Experimental/PhysicsControl/Source/PhysicsControlEditor/Private/OperatorViewer/SOperatorTreeWidget.cpp b/Engine/Plugins/Experimental/PhysicsControl/Source/PhysicsControlEditor/Private/OperatorViewer/SOperatorTreeWidget.cpp
new file mode 100644
--- /dev/null
+++ b/Engine/Plugins/Experimental/PhysicsControl/Source/PhysicsControlEditor/Private/OperatorViewer/SOperatorTreeWidget.cpp
@@ -0,0 +1,415 @@
+// Copyright Epic Games, Inc. All Rights Reserved.
+
+#include "OperatorViewer/SOperatorTreeWidget.h"
+
+#include "AnimGraphNode_RigidBodyWithControl.h"
+#include "BlueprintEditorModule.h"
+#include "Filters/GenericFilter.h"
+#include "HAL/PlatformApplicationMisc.h"
+#include "Kismet2/BlueprintEditorUtils.h"
+#include "Widgets/Views/STableRow.h"
+#include "Widgets/Input/SSearchBox.h"
+#include "Widgets/Text/SInlineEditableTextBlock.h"
+
+// UE_DISABLE_OPTIMIZATION;
+
+#define LOCTEXT_NAMESPACE "PhysicsControlOperatorViewer"
+
+template<typename TFilterItemType> bool PassesAnyFilters(const TFilterItemType InItem, const TSharedPtr<TFilterCollection<TFilterItemType>> FilterCollection)
+{
+	bool bPassesAnyFilter = true;
+
+	if (FilterCollection)
+	{
+		bPassesAnyFilter = false;
+
+		for (int32 FilterIndex = 0, FilterCount = FilterCollection->Num(); !bPassesAnyFilter && (FilterIndex < FilterCount); ++FilterIndex)
+		{
+			bPassesAnyFilter = FilterCollection->GetFilterAtIndex(FilterIndex)->PassesFilter(InItem);
+		}
+	}
+
+	return bPassesAnyFilter;
+}
+
+template<typename TFilterItemType> bool PassesAllFilters(const TFilterItemType InItem, const TSharedPtr<TFilterCollection<TFilterItemType>> FilterCollection)
+{
+	return !FilterCollection || FilterCollection->PassesAllFilters(InItem);
+}
+
+TArray<TArray<FString>> BuildStructuredSearchCriteria(const FText& SearchCriteriaText)
+{
+	const FString OROperatorChar(TEXT("|"));
+	const FString OROperatorText(TEXT(" OR "));
+
+	FString SearchCriteriaString = SearchCriteriaText.ToString();
+	SearchCriteriaString.ReplaceInline(*OROperatorText, *OROperatorChar);
+
+	TArray<FString> OrSeparatedCriteria;
+	SearchCriteriaString.ParseIntoArray(OrSeparatedCriteria, *OROperatorChar, true);
+
+	TArray<TArray<FString>> StructuredSearchCriteria;
+	const TCHAR* Delimeters[2] = { TEXT(" "), TEXT("+") };
+
+	for (const FString& CriteriaText : OrSeparatedCriteria)
+	{
+		StructuredSearchCriteria.Add(TArray<FString>());
+		CriteriaText.ParseIntoArray(StructuredSearchCriteria.Last(), Delimeters, 2, true);
+	}
+
+	return StructuredSearchCriteria;
+}
+
+void FormatStringForClipboard(FString& CopyBfr, const TCHAR* Delimeter)
+{
+	// Format buffer for pasting to an array in the details panel.
+	CopyBfr.ReplaceInline(Delimeter, TEXT("\",\""));
+	CopyBfr.InsertAt(0, TEXT("(\""));
+	CopyBfr.InsertAt(CopyBfr.Len(), TEXT("\")"));
+}
+
+// class SFilteredTreeWidget //
+void SOperatorTreeWidget::Construct(const FArguments& InArgs)
+{
+	Construct_Internal();
+
+	Refresh();
+}
+
+void SOperatorTreeWidget::Tick(const FGeometry& AllottedGeometry, const double InCurrentTime, const float InDeltaTime)
+{
+	Parent::Tick(AllottedGeometry, InCurrentTime, InDeltaTime);
+
+	const bool bIsScrollBarVisible = TreeView->GetScrollbarVisibility() == EVisibility::Visible;
+
+	if (bRefreshRequested)
+	{
+		Refresh();
+	}
+}
+
+void SOperatorTreeWidget::Construct_Internal()
+{
+	CreateCustomFilters();
+
+	SAssignNew(SearchBox, SSearchBox)
+		.SelectAllTextWhenFocused(true)
+		.OnTextChanged(this, &SOperatorTreeWidget::OnFilterTextChanged)
+		.HintText(LOCTEXT("SearchBoxHint", "Search the list of assets..."))
+		.AddMetaData<FTagMetaData>(TEXT("PerfToolAssetBrowser.Search"));
+
+	SAssignNew(FilterBar, FilterBarType)
+		.OnFilterChanged(this, &SOperatorTreeWidget::OnFilterBarFilterChanged)
+		.CustomFilters(CustomFilters);
+
+	TSharedPtr<SVerticalBox> VerticalBox = SNew(SVerticalBox);
+	TSharedPtr<SHorizontalBox> Toolbar = SNew(SHorizontalBox);
+
+	// Add Filter Menu Button
+	Toolbar->AddSlot()
+		.VAlign(VAlign_Center)
+		.Padding(0.0f, 0.0f, 2.0f, 0.0f)
+		.AutoWidth()
+		[
+			SBasicFilterBar<ItemType>::MakeAddFilterButton(FilterBar.ToSharedRef())
+		];
+
+	// Add the Search Text Box
+	Toolbar->AddSlot()
+		.VAlign(VAlign_Center)
+		[
+			SearchBox.ToSharedRef()
+		];
+
+	// Add the toolbar to the main layout vertical box
+	VerticalBox->AddSlot()
+		.AutoHeight()
+		.MaxHeight(24)
+		[
+			Toolbar.ToSharedRef()
+		];
+
+	// Add the filter bar to the main layout vertical box
+	VerticalBox->AddSlot()
+		.AutoHeight()
+		.MaxHeight(24)
+		.Padding(0.0f, 0.0f, 0.0f, 4.0f)
+		[
+			FilterBar.ToSharedRef()
+		];
+
+	SAssignNew(TreeView, STreeView<ItemType>)
+		.TreeItemsSource(&TreeItems)
+		.OnGetChildren(this, &SOperatorTreeWidget::OnItemGetChildren)
+		.OnGenerateRow(this, &SOperatorTreeWidget::GenerateItemRow)
+		.SelectionMode(ESelectionMode::Multi)
+		.OnContextMenuOpening(this, &SOperatorTreeWidget::CreateContextMenu)
+		.OnKeyDownHandler(this, &SOperatorTreeWidget::OnKeyDown)
+		.OnMouseButtonDoubleClick(this, &SOperatorTreeWidget::OnItemDoubleClicked);
+
+	VerticalBox->AddSlot()
+		.VAlign(VAlign_Fill)
+		[
+			TreeView.ToSharedRef()
+		];
+
+	ChildSlot
+		[
+			VerticalBox.ToSharedRef()
+		];
+}
+
+TSharedPtr<SWidget> SOperatorTreeWidget::CreateContextMenu()
+{
+	FMenuBuilder MenuBuilder(true, nullptr, nullptr);
+	MenuBuilder.BeginSection(NAME_None, LOCTEXT("ChannelDrawStyle", "Draw Style"));
+	{
+		MenuBuilder.AddMenuEntry(
+			LOCTEXT("OperatorTreeCopyNames", "Copy Names"),
+			LOCTEXT("OperatorTreeCopyNames_Tooltip", "Copy the names of all the selected operators to the clipboard."),
+			FSlateIcon(FAppStyle::GetAppStyleSetName(), "Icons.Copy"),
+			FUIAction(FExecuteAction::CreateSP(this, &SOperatorTreeWidget::CopySelectedItemsNamesToClipboard)));
+
+		MenuBuilder.AddMenuEntry(
+			LOCTEXT("OperatorTreeCopyAll", "Copy All"),
+			LOCTEXT("OperatorTreeCopyAll_Tooltip", "Copy the full description of all the selected operators to the clipboard."),
+			FSlateIcon(FAppStyle::GetAppStyleSetName(), "Icons.Copy"),
+			FUIAction(FExecuteAction::CreateSP(this, &SOperatorTreeWidget::CopySelectedItemsDescriptionsToClipboard)));
+	}
+	MenuBuilder.EndSection();
+
+	return MenuBuilder.MakeWidget();
+}
+
+TSharedRef<ITableRow> SOperatorTreeWidget::GenerateItemRow(ItemType InItem, const TSharedRef<STableViewBase>& OwnerTable)
+{
+	return InItem->GenerateRow(OwnerTable);
+}
+
+void SOperatorTreeWidget::CopySelectedItemsNamesToClipboard() const
+{
+	const TCHAR* Delimeter = TEXT("\n");
+	FString CopyBfr;
+	uint32 ItemCount = 0;
+
+	for (ItemType SelectedItem : TreeView->GetSelectedItems())
+	{
+		CopyBfr += SelectedItem->Name().ToString() + Delimeter;
+		++ItemCount;
+	}
+
+	CopyBfr.RemoveFromEnd(Delimeter);
+
+	if (ItemCount > 1)
+	{
+		FormatStringForClipboard(CopyBfr, Delimeter);
+	}
+
+	FPlatformApplicationMisc::ClipboardCopy(*CopyBfr);
+}
+
+void SOperatorTreeWidget::CopySelectedItemsDescriptionsToClipboard() const
+{
+	constexpr TCHAR Delimeter = '\n';
+	FString CopyBfr;
+
+	for (ItemType SelectedItem : TreeView->GetSelectedItems())
+	{
+		CopyBfr += SelectedItem->Description().ToString() + Delimeter;
+	}
+
+	CopyBfr.TrimEndInline();
+
+	FPlatformApplicationMisc::ClipboardCopy(*CopyBfr);
+}
+
+FReply SOperatorTreeWidget::OnKeyDown(const FGeometry& MyGeometry, const FKeyEvent& InKeyEvent)
+{
+	if (InKeyEvent.IsControlDown() && (InKeyEvent.GetKey() == EKeys::C))
+	{
+		CopySelectedItemsNamesToClipboard();
+		return FReply::Handled();
+	}
+
+	return FReply::Unhandled();
+}
+
+void SOperatorTreeWidget::OnItemDoubleClicked(ItemType InItem)
+{
+	if (!InItem->OnDoubleClick().IsEventHandled())
+	{
+		TreeView->SetItemSelection(InItem, true, ESelectInfo::OnMouseClick);
+		CopySelectedItemsNamesToClipboard();
+	}
+}
+
+void SOperatorTreeWidget::OnFilterBarFilterChanged()
+{
+	Refresh();
+}
+
+void SOperatorTreeWidget::OnFilterTextChanged(const FText& SearchCriteriaText)
+{
+	FilterStructuredCriteria = BuildStructuredSearchCriteria(SearchCriteriaText);
+
+	for (const FName Set : SetNames)
+	{
+		if (MatchSearchText(Set.ToString(), FilterStructuredCriteria))
+		{
+			FilterSetNames.Add(Set);
+		}
+	}
+
+	Refresh();
+}
+
+void SOperatorTreeWidget::RequestRefresh()
+{
+	bRefreshRequested = true;
+}
+
+void SOperatorTreeWidget::Refresh()
+{
+	using OperatorNameAndTags = OperatorTreeControlItem::OperatorNameAndTags;
+
+	check(FilterBar.IsValid());
+
+	TSet<ItemType> ExpandedItems;
+	TreeView->GetExpandedItems(ExpandedItems);
+
+	TreeItems.Reset();
+	SetNames.Reset();
+
+	TArray<OperatorNameAndTags> OperatorNamesAndTags;
+	TSet<FString> ProcessedNodeNames;
+
+	for (TObjectIterator<UAnimGraphNode_RigidBodyWithControl> Itr; Itr; ++Itr)
+	{	
+		FString BlueprintName = "[Unknown]";
+
+		if (UAnimBlueprint* const AnimBlueprint = Cast<UAnimBlueprint>(FBlueprintEditorUtils::FindBlueprintForNode(*Itr)))
+		{
+			BlueprintName = AnimBlueprint->GetFriendlyName();
+		}
+
+		const FString NodeUniqueName = Itr->GetName() + " [" + BlueprintName + "]";
+
+		if (!ProcessedNodeNames.Contains(NodeUniqueName))
+		{
+			// This is the first time we've encountered this node.
+			ProcessedNodeNames.Add(NodeUniqueName);
+
+			OperatorNamesAndTags.Reset();
+			OperatorNamesAndTags = Itr->GenerateControlsAndBodyModifierNames();
+
+			for (const OperatorNameAndTags& Operator : OperatorNamesAndTags)
+			{
+				// TODO - Define a way to sort these operators.
+
+				for (const FName SetName : Operator.Value)
+				{
+					SetNames.Add(SetName);
+				}
+			}
+
+			// Create node item
+			NodeItemPtr Node = MakeShared<OperatorTreeNodeItem>(*Itr, FName(*NodeUniqueName), BlueprintName);
+
+			if (!IsValid(Itr->Node.OverridePhysicsAsset))
+			{
+				Node->AddChild(MakeShared<OperatorTreeMessageItem>(*Itr, OperatorTreeMessageItem::Warning, LOCTEXT("OperatorNoPhysicsAssetDefinedText", "Can't create operators for limbs because no Physics Asset has been defined for this node.")));
+			}
+
+			// Create items for controls and modifiers
+			for (const OperatorNameAndTags& Operator : OperatorNamesAndTags)
+			{
+				OperatorTreeItemPtr Item = MakeShared<OperatorTreeControlItem>(Operator);
+
+				if (MatchesSearchAndFilter(Item))
+				{
+					Node->AddChild(Item);
+				}
+			}
+
+			if (OperatorNamesAndTags.IsEmpty())
+			{
+				Node->AddChild(MakeShared<OperatorTreeMessageItem>(*Itr, OperatorTreeMessageItem::Log, LOCTEXT("OperatorNoOperatorsDefinedText", "No operators defined for this node.")));
+			}
+			else if (!Node->HasChildren())
+			{
+				Node->AddChild(MakeShared<OperatorTreeMessageItem>(*Itr, OperatorTreeMessageItem::Log, LOCTEXT("OperatorNoOperatorsMatchSearchCriteria", "No operators match search criteria.")));
+			}
+			
+			TreeItems.Add(Node);
+		}
+	}
+
+	TreeView->RequestTreeRefresh();
+
+	// Find any new tree items that match old tree items that were expanded and expand them.
+	for (ItemType Item : ExpandedItems)
+	{
+		if (ItemType* const MatchingExpandedItem = TreeItems.FindByPredicate([Item](const ItemType InItem) -> bool { return IsEqual(InItem, Item); }))
+		{
+			TreeView->SetItemExpansion(*MatchingExpandedItem, true);
+		}
+	}
+
+	bRefreshRequested = false;
+}
+
+void SOperatorTreeWidget::OnItemGetChildren(ItemType InItem, TArray<ItemType>& OutChildren)
+{
+	OutChildren = InItem->Children;
+}
+
+void SOperatorTreeWidget::CreateCustomFilters()
+{
+	using FilterType = FGenericFilter<ItemType>;
+
+	TSharedPtr<FFilterCategory> OperatorTypeFilterCategory = MakeShared<FFilterCategory>(LOCTEXT("OperatorTypeFilterCategory", "Type"), FText::GetEmpty());
+	TSharedPtr<FFilterCategory> SetFilterCategory = MakeShared<FFilterCategory>(LOCTEXT("SetFilterCategory", "Set"), FText::GetEmpty());
+
+	auto CreateOperatorTypeFilter = [this, &OperatorTypeFilterCategory](const FName OperatorTypeName)
+	{
+		const FString FilterName = FString("OperatorType") + OperatorTypeName.ToString();
+
+		FilterType::FOnItemFiltered FilterDelegate = FilterType::FOnItemFiltered::CreateSP(this, &SOperatorTreeWidget::HasTag, OperatorTypeName);
+		TSharedPtr<FilterType> Filter = MakeShared<FilterType>(OperatorTypeFilterCategory, FilterName, FText::FromName(OperatorTypeName), FilterDelegate);
+		Filter->SetToolTipText(FText::Format(LOCTEXT("OperatorTypeFilterTooltip", "Only show {0}s."), FText::FromName(OperatorTypeName)));
+		this->CustomFilters.Add(Filter.ToSharedRef());
+	};
+
+	auto CreateSetFilter = [this, &SetFilterCategory](const FName SetName)
+	{
+		const FString FilterName = FString("Set") + SetName.ToString();
+
+		FilterType::FOnItemFiltered FilterDelegate = FilterType::FOnItemFiltered::CreateSP(this, &SOperatorTreeWidget::HasTag, SetName);
+		TSharedPtr<FilterType> Filter = MakeShared<FilterType>(SetFilterCategory, FilterName, FText::FromName(SetName), FilterDelegate);
+		Filter->SetToolTipText(FText::Format(LOCTEXT("SetFilterTooltip", "Only show controls or modifiers in the {0} set."), FText::FromName(SetName)));
+		this->CustomFilters.Add(Filter.ToSharedRef());
+	};
+
+	this->CustomFilters.Reset();
+
+	CreateOperatorTypeFilter(FName("Control"));
+	CreateOperatorTypeFilter(FName("Modifier"));
+
+	CreateSetFilter(FName("ParentSpace"));
+	CreateSetFilter(FName("WorldSpace"));	
+}
+
+bool SOperatorTreeWidget::HasTag(const ItemType Item, const FName InTag)
+{
+	return Item->HasTag(InTag);
+}
+
+bool SOperatorTreeWidget::MatchesSearchAndFilter(const ItemType InItem) const
+{
+	TSharedPtr<TFilterCollection<ItemType>> FilterCollection = FilterBar->GetAllActiveFilters();
+	check(FilterCollection);
+
+	return PassesAllFilters(InItem, FilterCollection) && (FilterStructuredCriteria.IsEmpty() || InItem->MatchesFilterText(FilterStructuredCriteria));
+}
+
+#undef LOCTEXT_NAMESPACE
diff --git a/Engine/Plugins/Experimental/PhysicsControl/Source/PhysicsControlEditor/Private/OperatorViewer/SOperatorViewerTabWidget.cpp b/Engine/Plugins/Experimental/PhysicsControl/Source/PhysicsControlEditor/Private/OperatorViewer/SOperatorViewerTabWidget.cpp
new file mode 100644
--- /dev/null
+++ b/Engine/Plugins/Experimental/PhysicsControl/Source/PhysicsControlEditor/Private/OperatorViewer/SOperatorViewerTabWidget.cpp
@@ -0,0 +1,46 @@
+// Copyright Epic Games, Inc. All Rights Reserved.
+
+#include "OperatorViewer/SOperatorViewerTabWidget.h"
+#include "OperatorViewer/SOperatorTreeWidget.h"
+#include "Framework/Docking/TabManager.h"
+#include "Widgets/Docking/SDockTab.h"
+#include "Widgets/Input/SButton.h"
+
+// UE_DISABLE_OPTIMIZATION;
+
+#define LOCTEXT_NAMESPACE "PhysicsControlOperatorViewer"
+
+void SOperatorViewerTabWidget::Construct(const FArguments& InArgs, int32 InTabIndex)
+{
+	TSharedPtr<SVerticalBox> TreeViewBox;
+
+	ChildSlot.Padding(5, 5, 5, 5)
+		[
+			SNew(SVerticalBox)
+			+ SVerticalBox::Slot()
+			[
+				SAssignNew(TreeViewBox, SVerticalBox)
+			]
+	
+		];
+
+	// Tree View
+	{
+		TreeViewWidget = SNew(SOperatorTreeWidget);
+
+		TreeViewBox->AddSlot()
+			[
+				TreeViewWidget.ToSharedRef()
+			];
+	}
+}
+
+void SOperatorViewerTabWidget::RequestRefresh()
+{
+	if (TreeViewWidget.IsValid())
+	{
+		TreeViewWidget->RequestRefresh();
+	}
+}
+
+#undef LOCTEXT_NAMESPACE
\ No newline at end of file
diff --git a/Engine/Plugins/Experimental/PhysicsControl/Source/PhysicsControlEditor/Private/PhysicsControlEditorModule.cpp b/Engine/Plugins/Experimental/PhysicsControl/Source/PhysicsControlEditor/Private/PhysicsControlEditorModule.cpp
--- a/Engine/Plugins/Experimental/PhysicsControl/Source/PhysicsControlEditor/Private/PhysicsControlEditorModule.cpp
+++ b/Engine/Plugins/Experimental/PhysicsControl/Source/PhysicsControlEditor/Private/PhysicsControlEditorModule.cpp
@@ -3,15 +3,13 @@
 #include "PhysicsControlEditorModule.h"
 #include "PhysicsControlComponent.h"
 #include "PhysicsControlComponentVisualizer.h"
-#include "OperatorEditor/OperatorEditor.h"
+#include "OperatorViewer/OperatorViewer.h"
 #include "PhysicsControlAssetActions.h"
 #include "PhysicsControlAssetEditorEditMode.h"
 #include "UnrealEdGlobals.h"
 #include "Editor/UnrealEdEngine.h"
 #include "Features/IModularFeatures.h"
 
-static const FName PhysicsControlEditorModule_PhysicsControlEditorInterface("PhysicsControlEditorInterface");
-
 #define LOCTEXT_NAMESPACE "PhysicsControlModule"
 
 //======================================================================================================================
@@ -38,13 +36,13 @@ void FPhysicsControlEditorModule::StartupModule()
 
 	if (!EditorInterface)
 	{
-		EditorInterface = new FPhysicsControlOperatorEditor;
+		EditorInterface = new FPhysicsControlOperatorViewer;
 	}
 
 	if (EditorInterface)
 	{
 		EditorInterface->Startup();
-		IModularFeatures::Get().RegisterModularFeature(PhysicsControlEditorModule_PhysicsControlEditorInterface, EditorInterface);
+		IModularFeatures::Get().RegisterModularFeature(IPhysicsControlOperatorViewerInterface::GetModularFeatureName(), EditorInterface);
 	}
 }
 
@@ -54,7 +52,7 @@ void FPhysicsControlEditorModule::ShutdownModule()
 	if (EditorInterface)
 	{
 		EditorInterface->Shutdown();
-		IModularFeatures::Get().UnregisterModularFeature(PhysicsControlEditorModule_PhysicsControlEditorInterface, EditorInterface);
+		IModularFeatures::Get().UnregisterModularFeature(IPhysicsControlOperatorViewerInterface::GetModularFeatureName(), EditorInterface);
 		delete EditorInterface;
 		EditorInterface = nullptr;
 	}
diff --git a/Engine/Plugins/Experimental/PhysicsControl/Source/PhysicsControlEditor/Public/OperatorEditor/OperatorEditor.h b/Engine/Plugins/Experimental/PhysicsControl/Source/PhysicsControlEditor/Public/OperatorEditor/OperatorEditor.h
deleted file mode 100644
diff --git a/Engine/Plugins/Experimental/PhysicsControl/Source/PhysicsControlEditor/Public/OperatorEditor/OperatorTreeElements.h b/Engine/Plugins/Experimental/PhysicsControl/Source/PhysicsControlEditor/Public/OperatorEditor/OperatorTreeElements.h
deleted file mode 100644
diff --git a/Engine/Plugins/Experimental/PhysicsControl/Source/PhysicsControlEditor/Public/OperatorEditor/SOperatorEditorTabWidget.h b/Engine/Plugins/Experimental/PhysicsControl/Source/PhysicsControlEditor/Public/OperatorEditor/SOperatorEditorTabWidget.h
deleted file mode 100644
diff --git a/Engine/Plugins/Experimental/PhysicsControl/Source/PhysicsControlEditor/Public/OperatorEditor/SOperatorTreeWidget.h b/Engine/Plugins/Experimental/PhysicsControl/Source/PhysicsControlEditor/Public/OperatorEditor/SOperatorTreeWidget.h
deleted file mode 100644
diff --git a/Engine/Plugins/Experimental/PhysicsControl/Source/PhysicsControlEditor/Public/OperatorViewer/OperatorTreeElements.h b/Engine/Plugins/Experimental/PhysicsControl/Source/PhysicsControlEditor/Public/OperatorViewer/OperatorTreeElements.h
new file mode 100644
--- /dev/null
+++ b/Engine/Plugins/Experimental/PhysicsControl/Source/PhysicsControlEditor/Public/OperatorViewer/OperatorTreeElements.h
@@ -0,0 +1,189 @@
+// Copyright Epic Games, Inc. All Rights Reserved.
+
+#pragma once
+
+#include "CoreMinimal.h"
+
+enum EHorizontalAlignment : int;
+enum EVerticalAlignment : int;
+
+class FReply;
+class FString;
+class ITableRow;
+class STableViewBase;
+class UEdGraphNode;
+
+struct FSlateColor;
+
+class OperatorTreeItem;
+
+using OperatorTreeItemPtr = TSharedPtr<OperatorTreeItem>;
+
+class OperatorTreeItem
+{
+public:
+
+	virtual ~OperatorTreeItem() {}
+
+	enum OperatorTreeItemType
+	{
+		ControlType,
+		ModifierType,
+		TargetType,
+		NodeType,
+		MessageType,
+		None
+	};
+
+	virtual FText Name() const = 0;
+	virtual FText Description() const { return Name(); }
+
+	virtual TSharedRef<ITableRow> GenerateRow(const TSharedRef<STableViewBase>& OwnerTable) const = 0;
+	virtual FReply OnDoubleClick();
+	virtual bool MatchesFilterText(const TArray<TArray<FString>>& FilterStructuredCriteria) const { return true; };
+	virtual bool HasTag(const FName Tag) const { return false; }
+	virtual bool IsExpandable() const { return false; }
+	virtual bool IsExpanded() const { return true; }
+
+	void AddChild(OperatorTreeItemPtr ItemPtr) { Children.Add(ItemPtr); }
+	bool HasChildren() const { return !Children.IsEmpty(); }
+
+	virtual OperatorTreeItemType GetType() const { return None; }
+	bool HasType(const OperatorTreeItemType InType) const { return InType == GetType(); }
+
+	TArray<OperatorTreeItemPtr> Children;
+
+	static const EHorizontalAlignment HorizontalAlignment;
+	static const EVerticalAlignment VerticalAlignment;
+	static const float HorizontalPadding;
+	static const float VerticalPadding;
+};
+
+class OperatorTreeControlItem : public OperatorTreeItem
+{
+	using Super = OperatorTreeItem;
+
+public:
+
+	using OperatorNameAndTags = TPair<FName, TArray<FName>>;
+
+	OperatorTreeControlItem()
+	{}
+
+	OperatorTreeControlItem(const OperatorTreeControlItem&) = default;
+
+	OperatorTreeControlItem(const OperatorNameAndTags& InOperatorNameAndTags)
+	: NameAndTags(InOperatorNameAndTags)
+	{}
+
+	virtual FText Name() const override;
+	virtual FText Description() const override;
+
+	virtual TSharedRef<ITableRow> GenerateRow(const TSharedRef<STableViewBase>& OwnerTable) const override;
+	virtual bool MatchesFilterText(const TArray<TArray<FString>>& FilterStructuredCriteria) const override;
+	virtual bool HasTag(const FName Tag) const override;
+	FName OperatorTypeName() const;
+
+	virtual OperatorTreeItemType GetType() const override { return StaticType(); }
+	static OperatorTreeItemType StaticType() { return OperatorTreeItem::ControlType; }
+
+	OperatorNameAndTags NameAndTags;
+
+	static const TArray<FName> HiddenTags;
+};
+
+using ControlItemPtr = TSharedPtr<OperatorTreeControlItem>;
+
+inline bool operator==(const OperatorTreeControlItem& Lhs, const OperatorTreeControlItem& Rhs)
+{
+	return (Lhs.Name().ToString() == Rhs.Name().ToString()) && (Lhs.HasTag(FName("Modifier")) == Rhs.HasTag(FName("Modifier")));
+}
+
+class OperatorTreeNodeItem : public OperatorTreeItem
+{
+	using Super = OperatorTreeItem;
+
+public:
+
+	OperatorTreeNodeItem()
+	{}
+
+	OperatorTreeNodeItem(const OperatorTreeNodeItem&) = default;
+
+	OperatorTreeNodeItem(const UEdGraphNode* InNode, const FName& InName, const FString& InBlueprintName);
+
+	virtual FText Name() const override;
+
+	virtual TSharedRef<ITableRow> GenerateRow(const TSharedRef<STableViewBase>& OwnerTable) const override;
+	virtual FReply OnDoubleClick() override;
+
+	virtual OperatorTreeItemType GetType() const override { return StaticType(); }
+	static OperatorTreeItemType StaticType() { return OperatorTreeItem::NodeType; }
+
+	const UEdGraphNode* GraphNode() const { return Node; }
+	bool IsSelectedInGraph() const;
+	FSlateColor ColorAndOpacity() const;
+
+private:
+	const UEdGraphNode* Node;
+	FText NodeDisplayName;
+
+};
+
+using NodeItemPtr = TSharedPtr<OperatorTreeNodeItem>;
+
+bool operator==(const OperatorTreeNodeItem& Lhs, const OperatorTreeNodeItem& Rhs);
+
+class OperatorTreeMessageItem : public OperatorTreeItem
+{
+	using Super = OperatorTreeItem;
+
+public:
+
+	enum MessageType
+	{
+		Log,
+		Warning,
+		Error
+	};
+
+	OperatorTreeMessageItem() = default;
+	OperatorTreeMessageItem(const OperatorTreeMessageItem&) = default;
+	OperatorTreeMessageItem(const UEdGraphNode* InNode, const MessageType InType, const FText& InText);
+
+	virtual FText Name() const override;
+
+	virtual TSharedRef<ITableRow> GenerateRow(const TSharedRef<STableViewBase>& OwnerTable) const override;
+
+	virtual OperatorTreeItemType GetType() const override { return StaticType(); }
+	static OperatorTreeItemType StaticType() { return OperatorTreeItem::MessageType; }
+
+	const UEdGraphNode* GraphNode() const { return Node; }
+
+private:
+	FText Text;
+	const UEdGraphNode* Node;
+	MessageType Type;
+};
+
+using MessageItemPtr = TSharedPtr<OperatorTreeMessageItem>;
+
+inline bool operator==(const OperatorTreeMessageItem& Lhs, const OperatorTreeMessageItem& Rhs)
+{
+	return Lhs.Name().ToString() == Rhs.Name().ToString(); // TODO - don't convert to strings + incomplete
+}
+
+template<typename DerivedType> TSharedPtr<DerivedType> StaticCastOperatorItemPtr(const TSharedPtr<OperatorTreeItem> InItem)
+{
+	if (InItem->GetType() == DerivedType::StaticType())
+	{
+		return StaticCastSharedPtr<DerivedType>(InItem);
+	}
+
+	return TSharedPtr<DerivedType>();
+}
+
+bool IsEqual(const TSharedPtr<OperatorTreeItem> Lhs, const TSharedPtr<OperatorTreeItem> Rhs);
+
+// Returns true if the TextToSearch string matches the search criteria. Used to identify operator or set names that match user specified search strings in the editor.
+bool MatchSearchText(const FString& TextToSearch, const TArray<TArray<FString>>& StructuredSearchCriteria);
diff --git a/Engine/Plugins/Experimental/PhysicsControl/Source/PhysicsControlEditor/Public/OperatorViewer/OperatorViewer.h b/Engine/Plugins/Experimental/PhysicsControl/Source/PhysicsControlEditor/Public/OperatorViewer/OperatorViewer.h
new file mode 100644
--- /dev/null
+++ b/Engine/Plugins/Experimental/PhysicsControl/Source/PhysicsControlEditor/Public/OperatorViewer/OperatorViewer.h
@@ -0,0 +1,35 @@
+// Copyright Epic Games, Inc. All Rights Reserved.
+
+#pragma once
+
+#include "CoreMinimal.h"
+#include "IPhysicsControlOperatorViewerInterface.h"
+#include "Containers/Ticker.h"
+
+class FSpawnTabArgs;
+class SOperatorViewerTabWidget;
+class SDockTab;
+
+class FPhysicsControlOperatorViewer : public IPhysicsControlOperatorViewerInterface
+{
+public:
+
+	virtual ~FPhysicsControlOperatorViewer() {}
+
+	virtual void OpenOperatorNamesTab() override;
+	virtual void CloseOperatorNamesTab() override;
+	virtual void ToggleOperatorNamesTab() override;
+	virtual bool IsOperatorNamesTabOpen() override;
+	virtual void RequestRefresh() override;
+
+	void Startup();
+	void Shutdown();
+
+	void OnTabClosed(TSharedRef<SDockTab> DockTab);
+
+private:
+	TSharedRef<SDockTab> OnCreateTab(const FSpawnTabArgs& SpawnTabArgs);
+
+	TSharedPtr<SOperatorViewerTabWidget> PersistantTabWidget;
+	TSharedPtr<SDockTab> OperatorNamesTab;
+};
\ No newline at end of file
diff --git a/Engine/Plugins/Experimental/PhysicsControl/Source/PhysicsControlEditor/Public/OperatorViewer/SOperatorTreeWidget.h b/Engine/Plugins/Experimental/PhysicsControl/Source/PhysicsControlEditor/Public/OperatorViewer/SOperatorTreeWidget.h
new file mode 100644
--- /dev/null
+++ b/Engine/Plugins/Experimental/PhysicsControl/Source/PhysicsControlEditor/Public/OperatorViewer/SOperatorTreeWidget.h
@@ -0,0 +1,80 @@
+// Copyright Epic Games, Inc. All Rights Reserved.
+
+#pragma once
+
+#include "CoreMinimal.h"
+
+#include "Widgets/SCompoundWidget.h"
+#include "Filters/SFilterSearchBox.h"
+#include "Filters/SBasicFilterBar.h"
+#include "Framework/SlateDelegates.h"
+#include "Framework/Application/SlateApplication.h"
+#include "Styling/AppStyle.h"
+
+#include "Widgets/Views/STreeView.h"
+
+#include "AnimNode_RigidBodyWithControl.h"
+
+#include "OperatorViewer/OperatorTreeElements.h"
+
+/**
+ * class SOperatorTreeWidget
+ *
+ * 
+ */
+class SOperatorTreeWidget : public SCompoundWidget
+{
+	using Parent = SCompoundWidget;
+
+public:
+	using ItemType = OperatorTreeItemPtr;
+
+	using FilterBarType = SBasicFilterBar<ItemType>;
+
+	SLATE_BEGIN_ARGS(SOperatorTreeWidget) {}
+	SLATE_END_ARGS()
+
+	/** Constructs this widget with InArgs */
+	void Construct(const FArguments& InArgs);
+	
+	virtual void Tick(const FGeometry& AllottedGeometry, const double InCurrentTime, const float InDeltaTime) override;
+
+	TSharedRef<ITableRow> GenerateItemRow(ItemType InItem, const TSharedRef<STableViewBase>& OwnerTable);
+	void OnItemGetChildren(ItemType InItem, TArray<ItemType>& OutChildren);
+	void OnFilterBarFilterChanged();
+	void OnFilterTextChanged(const FText& SearchText);
+
+	bool MatchesSearchAndFilter(const ItemType InItem) const;
+
+	bool HasTag(const ItemType Item, const FName Tag);
+
+	virtual FReply OnKeyDown(const FGeometry& MyGeometry, const FKeyEvent& InKeyEvent) override;
+
+	void OnItemDoubleClicked(ItemType InItem);
+	TSharedPtr<SWidget> CreateContextMenu();
+
+	void RequestRefresh();
+
+	void CopySelectedItemsNamesToClipboard() const;
+	void CopySelectedItemsDescriptionsToClipboard() const;
+
+protected:
+
+	void Construct_Internal();
+	virtual void CreateCustomFilters();
+	virtual void Refresh();
+
+	TSharedPtr<STreeView<ItemType>> TreeView;
+	TSharedPtr<SSearchBox> SearchBox;
+	TSharedPtr<FilterBarType> FilterBar;
+
+	TArray<TSharedRef<FFilterBase<ItemType>>> CustomFilters;
+	TArray<ItemType> TreeItems;
+	TSet<FName> SetNames;
+	FText FilterText;
+
+	TArray<TArray<FString>> FilterStructuredCriteria;
+	TSet<FName> FilterSetNames;
+
+	bool bRefreshRequested;
+};
diff --git a/Engine/Plugins/Experimental/PhysicsControl/Source/PhysicsControlEditor/Public/OperatorViewer/SOperatorViewerTabWidget.h b/Engine/Plugins/Experimental/PhysicsControl/Source/PhysicsControlEditor/Public/OperatorViewer/SOperatorViewerTabWidget.h
new file mode 100644
--- /dev/null
+++ b/Engine/Plugins/Experimental/PhysicsControl/Source/PhysicsControlEditor/Public/OperatorViewer/SOperatorViewerTabWidget.h
@@ -0,0 +1,27 @@
+// Copyright Epic Games, Inc. All Rights Reserved.
+
+#pragma once
+
+#include "Templates/SharedPointer.h"
+//#include "Widgets/DeclarativeSyntaxSupport.h"
+#include "Widgets/SCompoundWidget.h"
+
+class SCompoundWidget;
+class SOperatorTreeWidget;
+
+class SOperatorViewerTabWidget : public SCompoundWidget
+{
+public:
+
+	SLATE_BEGIN_ARGS(SOperatorViewerTabWidget)
+	{}
+
+	SLATE_END_ARGS()
+
+public:
+	void Construct(const FArguments& InArgs, int32 InTabIndex);
+	void RequestRefresh();
+
+private:
+	TSharedPtr<SOperatorTreeWidget> TreeViewWidget;
+};
diff --git a/Engine/Plugins/Experimental/PhysicsControl/Source/PhysicsControlEditor/Public/PhysicsControlEditorModule.h b/Engine/Plugins/Experimental/PhysicsControl/Source/PhysicsControlEditor/Public/PhysicsControlEditorModule.h
--- a/Engine/Plugins/Experimental/PhysicsControl/Source/PhysicsControlEditor/Public/PhysicsControlEditorModule.h
+++ b/Engine/Plugins/Experimental/PhysicsControl/Source/PhysicsControlEditor/Public/PhysicsControlEditorModule.h
@@ -7,7 +7,7 @@
 
 
 
-class FPhysicsControlOperatorEditor;
+class FPhysicsControlOperatorViewer;
 class FPhysicsControlAssetActions;
 
 /**
@@ -26,6 +26,6 @@ private:
 	TArray<FName> VisualizersToUnregisterOnShutdown;
 	TSharedPtr<FPhysicsControlAssetActions> PhysicsControlAssetActions;
 
-	FPhysicsControlOperatorEditor* EditorInterface = nullptr;
+	FPhysicsControlOperatorViewer* EditorInterface = nullptr;
 };
 
diff --git a/Engine/Plugins/Experimental/PhysicsControl/Source/PhysicsControlUncookedOnly/Private/AnimGraphNode_RigidBodyWithControl.cpp b/Engine/Plugins/Experimental/PhysicsControl/Source/PhysicsControlUncookedOnly/Private/AnimGraphNode_RigidBodyWithControl.cpp
--- a/Engine/Plugins/Experimental/PhysicsControl/Source/PhysicsControlUncookedOnly/Private/AnimGraphNode_RigidBodyWithControl.cpp
+++ b/Engine/Plugins/Experimental/PhysicsControl/Source/PhysicsControlUncookedOnly/Private/AnimGraphNode_RigidBodyWithControl.cpp
@@ -7,7 +7,7 @@
 #include "AnimNode_RigidBodyWithControl.h"
 #include "EditorModeManager.h"
 #include "IPhysicsAssetRenderInterface.h"
-#include "IPhysicsControlOperatorEditorInterface.h"
+#include "IPhysicsControlOperatorViewerInterface.h"
 #include "Kismet2/BlueprintEditorUtils.h"
 #include "PhysicsControlOperatorNameGeneration.h"
 #include "PhysicsControlAsset.h"
@@ -23,6 +23,15 @@
 
 // UE_DISABLE_OPTIMIZATION;
 
+// Use this CVar to enable/disable the viewer for control/modifier sets. It's not really functional/correct enough
+// yet for general use and visibility
+bool bRBANWithControl_EnableControlSetViewer = false;
+FAutoConsoleVariableRef CVarRigidBodyWithControlEnableControlSetViewer(
+	TEXT("p.RigidBodyWithControl.EnableControlSetViewer"), 
+	bRBANWithControl_EnableControlSetViewer, 
+	TEXT("Enable/Disable the simple viewer for control and modifier sets for the RBWC node"), ECVF_Default);
+
+
 /////////////////////////////////////////////////////
 // UAnimGraphNode_RigidBodyWithControl
 
@@ -80,7 +89,8 @@ void UAnimGraphNode_RigidBodyWithControl::Draw(
 		if (UPhysicsAsset* const PhysicsAsset = RuntimeRigidBodyNode->GetPhysicsAsset())
 		{
 			IPhysicsAssetRenderInterface& PhysicsAssetRenderInterface = 
-				IModularFeatures::Get().GetModularFeature<IPhysicsAssetRenderInterface>("PhysicsAssetRenderInterface");
+				IModularFeatures::Get().GetModularFeature<IPhysicsAssetRenderInterface>(
+					IPhysicsAssetRenderInterface::GetModularFeatureName());
 
 			// Draw Bodies.
 			if (bIsSelected || 
@@ -148,12 +158,13 @@ void UAnimGraphNode_RigidBodyWithControl::CustomizeDetails(IDetailLayoutBuilder&
 	Super::CustomizeDetails(DetailBuilder);
 	IDetailCategoryBuilder& ViewportCategory = DetailBuilder.EditCategory(TEXT("Debug Visualization"));
 	FAnimNode_RigidBodyWithControl* const RigidBodyNode = static_cast<FAnimNode_RigidBodyWithControl*>(GetDebuggedAnimNode());
-	
+
+	if (bRBANWithControl_EnableControlSetViewer)
 	{
-		FDetailWidgetRow& ControlsEditorWidgetRow = ViewportCategory.AddCustomRow(
-			LOCTEXT("ToggleControlsEditorWidgetRowButtonRow", "ControlsEditor"));
+		FDetailWidgetRow& ControlSetViewerWidgetRow = ViewportCategory.AddCustomRow(
+			LOCTEXT("ToggleControlSetViewerWidgetRowButtonRow", "ControlSetViewer"));
 
-		ControlsEditorWidgetRow
+		ControlSetViewerWidgetRow
 			[
 				SNew(SHorizontalBox)
 				+ SHorizontalBox::Slot()
@@ -161,13 +172,13 @@ void UAnimGraphNode_RigidBodyWithControl::CustomizeDetails(IDetailLayoutBuilder&
 					SNew(SButton)
 					.HAlign(HAlign_Center)
 					.VAlign(VAlign_Center)
-					.OnClicked_Lambda([this]() { this->ToggleControlEditorTab(); return FReply::Handled(); })
-					.ButtonColorAndOpacity_Lambda([this]() { return (IsControlEditorTabOpen()) ? FAppStyle::Get().GetSlateColor("Colors.AccentRed") : FAppStyle::Get().GetSlateColor("Colors.AccentGreen"); })
+					.OnClicked_Lambda([this]() { this->ToggleControlSetViewerTab(); return FReply::Handled(); })
+					.ButtonColorAndOpacity_Lambda([this]() { return (IsControlSetViewerTabOpen()) ? FAppStyle::Get().GetSlateColor("Colors.AccentRed") : FAppStyle::Get().GetSlateColor("Colors.AccentGreen"); })
 					.Content()
 					[
 						SNew(STextBlock)
-						.Text_Lambda([this]() { return (IsControlEditorTabOpen()) ? LOCTEXT("CloseControlEditorTabButtonText", "Close Control Editor") : LOCTEXT("OpenControlEditorTabButtonText", "Open Control Editor"); })
-						.ToolTipText(LOCTEXT("ToggleControlEditorTabButtonToolTip", "Toggle Control Editor Tab"))
+						.Text_Lambda([this]() { return (IsControlSetViewerTabOpen()) ? LOCTEXT("CloseControlSetViewerTabButtonText", "Close Control Set Viewer") : LOCTEXT("OpenControlSetViewerTabButtonText", "Open Control Set Viewer"); })
+						.ToolTipText(LOCTEXT("ToggleControlSetViewerTabButtonToolTip", "Toggle the viewer for control and modifier sets. This lists the controls and modifiers created by each RBWC node, and shows what sets they are in"))
 					]
 				]
 			];
@@ -216,9 +227,14 @@ void UAnimGraphNode_RigidBodyWithControl::CustomizeDetails(IDetailLayoutBuilder&
 
 void UAnimGraphNode_RigidBodyWithControl::PostChange()
 {
-	IPhysicsControlOperatorEditorInterface& PhysicsControlEditorInterface = 
-		IModularFeatures::Get().GetModularFeature<IPhysicsControlOperatorEditorInterface>("PhysicsControlEditorInterface");
-	PhysicsControlEditorInterface.RequestRefresh();
+	if (IModularFeatures::Get().IsModularFeatureAvailable(
+		IPhysicsControlOperatorViewerInterface::GetModularFeatureName()))
+	{
+		IPhysicsControlOperatorViewerInterface& PhysicsControlViewerInterface =
+			IModularFeatures::Get().GetModularFeature<IPhysicsControlOperatorViewerInterface>(
+				IPhysicsControlOperatorViewerInterface::GetModularFeatureName());
+		PhysicsControlViewerInterface.RequestRefresh();
+	}
 }
 
 void UAnimGraphNode_RigidBodyWithControl::PostEditChangeProperty(struct FPropertyChangedEvent& PropertyChangedEvent)
@@ -226,7 +242,8 @@ void UAnimGraphNode_RigidBodyWithControl::PostEditChangeProperty(struct FPropert
 	Super::PostEditChangeProperty(PropertyChangedEvent);
 
 	IPhysicsAssetRenderInterface& PhysicsAssetRenderInterface = 
-		IModularFeatures::Get().GetModularFeature<IPhysicsAssetRenderInterface>("PhysicsAssetRenderInterface");
+		IModularFeatures::Get().GetModularFeature<IPhysicsAssetRenderInterface>(
+			IPhysicsAssetRenderInterface::GetModularFeatureName());
 	PhysicsAssetRenderInterface.SaveConfig();
 
 	PostChange();
@@ -256,7 +273,8 @@ void UAnimGraphNode_RigidBodyWithControl::ToggleBodyVisibility()
 		static_cast<FAnimNode_RigidBodyWithControl*>(GetDebuggedAnimNode());
 	
 	IPhysicsAssetRenderInterface& PhysicsAssetRenderInterface = 
-		IModularFeatures::Get().GetModularFeature<IPhysicsAssetRenderInterface>("PhysicsAssetRenderInterface");
+		IModularFeatures::Get().GetModularFeature<IPhysicsAssetRenderInterface>(
+			IPhysicsAssetRenderInterface::GetModularFeatureName());
 
 	if (RigidBodyNode)
 	{
@@ -269,7 +287,8 @@ void UAnimGraphNode_RigidBodyWithControl::ToggleConstraintVisibility()
 	FAnimNode_RigidBodyWithControl* const RigidBodyNode = 
 		static_cast<FAnimNode_RigidBodyWithControl*>(GetDebuggedAnimNode());
 	IPhysicsAssetRenderInterface& PhysicsAssetRenderInterface = 
-		IModularFeatures::Get().GetModularFeature<IPhysicsAssetRenderInterface>("PhysicsAssetRenderInterface");
+		IModularFeatures::Get().GetModularFeature<IPhysicsAssetRenderInterface>(
+			IPhysicsAssetRenderInterface::GetModularFeatureName());
 
 	if (RigidBodyNode)
 	{
@@ -282,7 +301,8 @@ bool UAnimGraphNode_RigidBodyWithControl::AreAnyBodiesHidden() const
 	FAnimNode_RigidBodyWithControl* const RigidBodyNode = 
 		static_cast<FAnimNode_RigidBodyWithControl*>(GetDebuggedAnimNode());
 	IPhysicsAssetRenderInterface& PhysicsAssetRenderInterface = 
-		IModularFeatures::Get().GetModularFeature<IPhysicsAssetRenderInterface>("PhysicsAssetRenderInterface");
+		IModularFeatures::Get().GetModularFeature<IPhysicsAssetRenderInterface>(
+			IPhysicsAssetRenderInterface::GetModularFeatureName());
 
 	if (RigidBodyNode)
 	{
@@ -297,7 +317,8 @@ bool UAnimGraphNode_RigidBodyWithControl::AreAnyConstraintsHidden() const
 	FAnimNode_RigidBodyWithControl* const RigidBodyNode = 
 		static_cast<FAnimNode_RigidBodyWithControl*>(GetDebuggedAnimNode());
 	IPhysicsAssetRenderInterface& PhysicsAssetRenderInterface = 
-		IModularFeatures::Get().GetModularFeature<IPhysicsAssetRenderInterface>("PhysicsAssetRenderInterface");
+		IModularFeatures::Get().GetModularFeature<IPhysicsAssetRenderInterface>(
+			IPhysicsAssetRenderInterface::GetModularFeatureName());
 
 	if (RigidBodyNode)
 	{
@@ -307,18 +328,29 @@ bool UAnimGraphNode_RigidBodyWithControl::AreAnyConstraintsHidden() const
 	return false;
 }
 
-void UAnimGraphNode_RigidBodyWithControl::ToggleControlEditorTab()
+void UAnimGraphNode_RigidBodyWithControl::ToggleControlSetViewerTab()
 {
-	IPhysicsControlOperatorEditorInterface& PhysicsControlEditorInterface = 
-		IModularFeatures::Get().GetModularFeature<IPhysicsControlOperatorEditorInterface>("PhysicsControlEditorInterface");
-	PhysicsControlEditorInterface.ToggleOperatorNamesTab();
+	if (IModularFeatures::Get().IsModularFeatureAvailable(
+		IPhysicsControlOperatorViewerInterface::GetModularFeatureName()))
+	{
+		IPhysicsControlOperatorViewerInterface& PhysicsControlViewerInterface =
+			IModularFeatures::Get().GetModularFeature<IPhysicsControlOperatorViewerInterface>(
+				IPhysicsControlOperatorViewerInterface::GetModularFeatureName());
+		PhysicsControlViewerInterface.ToggleOperatorNamesTab();
+	}
 }
 
-bool UAnimGraphNode_RigidBodyWithControl::IsControlEditorTabOpen() const
+bool UAnimGraphNode_RigidBodyWithControl::IsControlSetViewerTabOpen() const
 {
-	IPhysicsControlOperatorEditorInterface& PhysicsControlEditorInterface = 
-		IModularFeatures::Get().GetModularFeature<IPhysicsControlOperatorEditorInterface>("PhysicsControlEditorInterface");
-	return PhysicsControlEditorInterface.IsOperatorNamesTabOpen();
+	if (IModularFeatures::Get().IsModularFeatureAvailable(
+		IPhysicsControlOperatorViewerInterface::GetModularFeatureName()))
+	{
+		IPhysicsControlOperatorViewerInterface& PhysicsControlViewerInterface =
+		IModularFeatures::Get().GetModularFeature<IPhysicsControlOperatorViewerInterface>(
+			IPhysicsControlOperatorViewerInterface::GetModularFeatureName());
+		return PhysicsControlViewerInterface.IsOperatorNamesTabOpen();
+	}
+	return false;
 }
 
 TArray<TPair<FName, TArray<FName>>> UAnimGraphNode_RigidBodyWithControl::GenerateControlsAndBodyModifierNames() const
diff --git a/Engine/Plugins/Experimental/PhysicsControl/Source/PhysicsControlUncookedOnly/Public/AnimGraphNode_RigidBodyWithControl.h b/Engine/Plugins/Experimental/PhysicsControl/Source/PhysicsControlUncookedOnly/Public/AnimGraphNode_RigidBodyWithControl.h
--- a/Engine/Plugins/Experimental/PhysicsControl/Source/PhysicsControlUncookedOnly/Public/AnimGraphNode_RigidBodyWithControl.h
+++ b/Engine/Plugins/Experimental/PhysicsControl/Source/PhysicsControlUncookedOnly/Public/AnimGraphNode_RigidBodyWithControl.h
@@ -46,7 +46,7 @@ public:
 
 	void ToggleBodyVisibility();
 	void ToggleConstraintVisibility();
-	void ToggleControlEditorTab();
+	void ToggleControlSetViewerTab();
 
 	TArray<TPair<FName, TArray<FName>>> GenerateControlsAndBodyModifierNames() const;
 
@@ -60,7 +60,7 @@ protected:
 
 	bool AreAnyBodiesHidden() const;
 	bool AreAnyConstraintsHidden() const;
-	bool IsControlEditorTabOpen() const;
+	bool IsControlSetViewerTabOpen() const;
 
 	void PostChange();
 
diff --git a/Engine/Plugins/Experimental/PhysicsControl/Source/PhysicsControlUncookedOnly/Public/IPhysicsControlOperatorEditorInterface.h b/Engine/Plugins/Experimental/PhysicsControl/Source/PhysicsControlUncookedOnly/Public/IPhysicsControlOperatorEditorInterface.h
deleted file mode 100644
diff --git a/Engine/Plugins/Experimental/PhysicsControl/Source/PhysicsControlUncookedOnly/Public/IPhysicsControlOperatorViewerInterface.h b/Engine/Plugins/Experimental/PhysicsControl/Source/PhysicsControlUncookedOnly/Public/IPhysicsControlOperatorViewerInterface.h
new file mode 100644
--- /dev/null
+++ b/Engine/Plugins/Experimental/PhysicsControl/Source/PhysicsControlUncookedOnly/Public/IPhysicsControlOperatorViewerInterface.h
@@ -0,0 +1,23 @@
+// Copyright Epic Games, Inc. All Rights Reserved.
+
+#pragma once
+
+#include "CoreMinimal.h"
+#include "Features/IModularFeature.h"
+
+class IPhysicsControlOperatorViewerInterface : public IModularFeature
+{
+public:
+	static FName GetModularFeatureName()
+	{
+		static FName FeatureName = FName(TEXT("PhysicsControlViewerInterface"));
+		return FeatureName;
+	}
+
+	virtual void OpenOperatorNamesTab() = 0;
+	virtual void CloseOperatorNamesTab() = 0;
+	virtual void ToggleOperatorNamesTab() = 0;
+	virtual bool IsOperatorNamesTabOpen() = 0;
+	virtual void RequestRefresh() = 0;
+
+};
\ No newline at end of file
diff --git a/Engine/Source/Editor/AnimGraph/Private/AnimGraphNode_RigidBody.cpp b/Engine/Source/Editor/AnimGraph/Private/AnimGraphNode_RigidBody.cpp
--- a/Engine/Source/Editor/AnimGraph/Private/AnimGraphNode_RigidBody.cpp
+++ b/Engine/Source/Editor/AnimGraph/Private/AnimGraphNode_RigidBody.cpp
@@ -58,7 +58,7 @@ void UAnimGraphNode_RigidBody::Draw(FPrimitiveDrawInterface* PDI, USkeletalMeshC
 	{
 		if (UPhysicsAsset* const PhysicsAsset = RuntimeRigidBodyNode->GetPhysicsAsset())
 		{
-			IPhysicsAssetRenderInterface& PhysicsAssetRenderInterface = IModularFeatures::Get().GetModularFeature<IPhysicsAssetRenderInterface>("PhysicsAssetRenderInterface");
+			IPhysicsAssetRenderInterface& PhysicsAssetRenderInterface = IModularFeatures::Get().GetModularFeature<IPhysicsAssetRenderInterface>(IPhysicsAssetRenderInterface::GetModularFeatureName());
 
 			// Draw Bodies.
 			if (bIsSelected || (bIsPoseWatchEnabled && PoseWatchElementBodies.IsValid() && PoseWatchElementBodies->GetIsVisible()))
@@ -152,7 +152,7 @@ void UAnimGraphNode_RigidBody::PostEditChangeProperty(struct FPropertyChangedEve
 {
 	Super::PostEditChangeProperty(PropertyChangedEvent);
 
-	IPhysicsAssetRenderInterface& PhysicsAssetRenderInterface = IModularFeatures::Get().GetModularFeature<IPhysicsAssetRenderInterface>("PhysicsAssetRenderInterface");
+	IPhysicsAssetRenderInterface& PhysicsAssetRenderInterface = IModularFeatures::Get().GetModularFeature<IPhysicsAssetRenderInterface>(IPhysicsAssetRenderInterface::GetModularFeatureName());
 	PhysicsAssetRenderInterface.SaveConfig();
 }
 
@@ -160,7 +160,7 @@ void UAnimGraphNode_RigidBody::ToggleBodyVisibility()
 {
 	FAnimNode_RigidBody* const RigidBodyNode = static_cast<FAnimNode_RigidBody*>(GetDebuggedAnimNode());
 	
-	IPhysicsAssetRenderInterface& PhysicsAssetRenderInterface = IModularFeatures::Get().GetModularFeature<IPhysicsAssetRenderInterface>("PhysicsAssetRenderInterface");
+	IPhysicsAssetRenderInterface& PhysicsAssetRenderInterface = IModularFeatures::Get().GetModularFeature<IPhysicsAssetRenderInterface>(IPhysicsAssetRenderInterface::GetModularFeatureName());
 
 	if (RigidBodyNode)
 	{
@@ -171,7 +171,7 @@ void UAnimGraphNode_RigidBody::ToggleBodyVisibility()
 void UAnimGraphNode_RigidBody::ToggleConstraintVisibility()
 {
 	FAnimNode_RigidBody* const RigidBodyNode = static_cast<FAnimNode_RigidBody*>(GetDebuggedAnimNode());
-	IPhysicsAssetRenderInterface& PhysicsAssetRenderInterface = IModularFeatures::Get().GetModularFeature<IPhysicsAssetRenderInterface>("PhysicsAssetRenderInterface");
+	IPhysicsAssetRenderInterface& PhysicsAssetRenderInterface = IModularFeatures::Get().GetModularFeature<IPhysicsAssetRenderInterface>(IPhysicsAssetRenderInterface::GetModularFeatureName());
 
 	if (RigidBodyNode)
 	{
@@ -182,7 +182,7 @@ void UAnimGraphNode_RigidBody::ToggleConstraintVisibility()
 bool UAnimGraphNode_RigidBody::AreAnyBodiesHidden() const
 {
 	FAnimNode_RigidBody* const RigidBodyNode = static_cast<FAnimNode_RigidBody*>(GetDebuggedAnimNode());
-	IPhysicsAssetRenderInterface& PhysicsAssetRenderInterface = IModularFeatures::Get().GetModularFeature<IPhysicsAssetRenderInterface>("PhysicsAssetRenderInterface");
+	IPhysicsAssetRenderInterface& PhysicsAssetRenderInterface = IModularFeatures::Get().GetModularFeature<IPhysicsAssetRenderInterface>(IPhysicsAssetRenderInterface::GetModularFeatureName());
 
 	if (RigidBodyNode)
 	{
@@ -195,7 +195,7 @@ bool UAnimGraphNode_RigidBody::AreAnyBodiesHidden() const
 bool UAnimGraphNode_RigidBody::AreAnyConstraintsHidden() const
 {
 	FAnimNode_RigidBody* const RigidBodyNode = static_cast<FAnimNode_RigidBody*>(GetDebuggedAnimNode());
-	IPhysicsAssetRenderInterface& PhysicsAssetRenderInterface = IModularFeatures::Get().GetModularFeature<IPhysicsAssetRenderInterface>("PhysicsAssetRenderInterface");
+	IPhysicsAssetRenderInterface& PhysicsAssetRenderInterface = IModularFeatures::Get().GetModularFeature<IPhysicsAssetRenderInterface>(IPhysicsAssetRenderInterface::GetModularFeatureName());
 
 	if (RigidBodyNode)
 	{
diff --git a/Engine/Source/Editor/DetailCustomizations/Private/PhysicsConstraintComponentDetails.cpp b/Engine/Source/Editor/DetailCustomizations/Private/PhysicsConstraintComponentDetails.cpp
--- a/Engine/Source/Editor/DetailCustomizations/Private/PhysicsConstraintComponentDetails.cpp
+++ b/Engine/Source/Editor/DetailCustomizations/Private/PhysicsConstraintComponentDetails.cpp
@@ -1189,7 +1189,7 @@ void FPhysicsConstraintComponentDetails::ToggleDisplayConstraintTransformCompone
 {
 	if (ParentPhysicsAsset)
 	{
-		IPhysicsAssetRenderInterface& PhysicsAssetRenderInterface = IModularFeatures::Get().GetModularFeature<IPhysicsAssetRenderInterface>("PhysicsAssetRenderInterface");
+		IPhysicsAssetRenderInterface& PhysicsAssetRenderInterface = IModularFeatures::Get().GetModularFeature<IPhysicsAssetRenderInterface>(IPhysicsAssetRenderInterface::GetModularFeatureName());
 		const bool bCurrentState = IsDisplayingConstraintTransformComponentRelativeToDefault(ComponentFlags);
 		const EConstraintTransformComponentFlags ComponentFlagsToModify = (FSlateApplication::Get().GetModifierKeys().IsShiftDown()) ? EConstraintTransformComponentFlags::All : ComponentFlags;
 		PhysicsAssetRenderInterface.SetDisplayConstraintTransformComponentRelativeToDefault(ParentPhysicsAsset, ComponentFlagsToModify, !bCurrentState);
@@ -1201,7 +1201,7 @@ bool FPhysicsConstraintComponentDetails::IsConstraintTransformComponentEnabled(c
 {
 	if (ParentPhysicsAsset)
 	{
-		IPhysicsAssetRenderInterface& PhysicsAssetRenderInterface = IModularFeatures::Get().GetModularFeature<IPhysicsAssetRenderInterface>("PhysicsAssetRenderInterface");
+		IPhysicsAssetRenderInterface& PhysicsAssetRenderInterface = IModularFeatures::Get().GetModularFeature<IPhysicsAssetRenderInterface>(IPhysicsAssetRenderInterface::GetModularFeatureName());
 		const EConstraintTransformComponentFlags ComponentManipulationFlags = PhysicsAssetRenderInterface.GetConstraintViewportManipulationFlags(ParentPhysicsAsset);
 
 		return EnumHasAllFlags(ComponentManipulationFlags, ComponentFlags);
@@ -1219,7 +1219,7 @@ void FPhysicsConstraintComponentDetails::UpdateTransformProxyDisplayRelativeToDe
 {
 	if (ParentPhysicsAsset)
 	{
-		IPhysicsAssetRenderInterface& PhysicsAssetRenderInterface = IModularFeatures::Get().GetModularFeature<IPhysicsAssetRenderInterface>("PhysicsAssetRenderInterface");
+		IPhysicsAssetRenderInterface& PhysicsAssetRenderInterface = IModularFeatures::Get().GetModularFeature<IPhysicsAssetRenderInterface>(IPhysicsAssetRenderInterface::GetModularFeatureName());
 
 		ChildTransformProxy->SetPositionDisplayRelativeToDefault(PhysicsAssetRenderInterface.IsDisplayingConstraintTransformComponentRelativeToDefault(ParentPhysicsAsset, EConstraintTransformComponentFlags::ChildPosition));
 		ChildTransformProxy->SetRotationDisplayRelativeToDefault(PhysicsAssetRenderInterface.IsDisplayingConstraintTransformComponentRelativeToDefault(ParentPhysicsAsset, EConstraintTransformComponentFlags::ChildRotation));
@@ -1238,7 +1238,7 @@ bool FPhysicsConstraintComponentDetails::IsDisplayingConstraintTransformComponen
 {
 	if (ParentPhysicsAsset)
 	{
-		return IModularFeatures::Get().GetModularFeature<IPhysicsAssetRenderInterface>("PhysicsAssetRenderInterface").IsDisplayingConstraintTransformComponentRelativeToDefault(ParentPhysicsAsset, ComponentFlags);
+		return IModularFeatures::Get().GetModularFeature<IPhysicsAssetRenderInterface>(IPhysicsAssetRenderInterface::GetModularFeatureName()).IsDisplayingConstraintTransformComponentRelativeToDefault(ParentPhysicsAsset, ComponentFlags);
 	}
 
 	return false;
diff --git a/Engine/Source/Editor/Persona/Private/PhysicsAssetRenderUtils.cpp b/Engine/Source/Editor/Persona/Private/PhysicsAssetRenderUtils.cpp
--- a/Engine/Source/Editor/Persona/Private/PhysicsAssetRenderUtils.cpp
+++ b/Engine/Source/Editor/Persona/Private/PhysicsAssetRenderUtils.cpp
@@ -624,7 +624,8 @@ void UPhysicsAssetRenderUtilities::Initialise()
 			PhysicsAssetRenderUtilities->PhysicsAssetRenderInterface = new FPhysicsAssetRenderInterface;
 		}
 
-		IModularFeatures::Get().RegisterModularFeature("PhysicsAssetRenderInterface", PhysicsAssetRenderUtilities->PhysicsAssetRenderInterface);
+		IModularFeatures::Get().RegisterModularFeature(IPhysicsAssetRenderInterface::GetModularFeatureName(), 
+			PhysicsAssetRenderUtilities->PhysicsAssetRenderInterface);
 	}
 }
 
diff --git a/Engine/Source/Editor/Persona/Public/IPhysicsAssetRenderInterface.h b/Engine/Source/Editor/Persona/Public/IPhysicsAssetRenderInterface.h
--- a/Engine/Source/Editor/Persona/Public/IPhysicsAssetRenderInterface.h
+++ b/Engine/Source/Editor/Persona/Public/IPhysicsAssetRenderInterface.h
@@ -13,6 +13,11 @@ class IPhysicsAssetRenderInterface : public IModuleInterface, public IModularFea
 	// virtual ~IPhysicsAssetRenderInterface() = default;
 
 public:
+	static FName GetModularFeatureName()
+	{
+		static FName FeatureName = FName(TEXT("PhysicsAssetRenderInterface"));
+		return FeatureName;
+	}
 
 	virtual void DebugDraw(class USkeletalMeshComponent* const SkeletalMeshComponent, class UPhysicsAsset* const PhysicsAsset, class FPrimitiveDrawInterface* PDI) = 0;
 	virtual void DebugDrawBodies(class USkeletalMeshComponent* const SkeletalMeshComponent, class UPhysicsAsset* const PhysicsAsset, FPrimitiveDrawInterface* PDI, const FColor& PrimitiveColorOverride) = 0;
