[UE] Ensure PBA Physics Asset Conditioning is applied when SkeltalMesh Changes
- Keep a record of the skeletal mesh that was associated with the current asset during the last physics init
- Re-initialize physics if the current asset's skeletal mesh has changed and physics asset conditioning is requirted

diff --git a/Engine/Plugins/Experimental/PhysicsControl/Source/PhysicsControl/Private/AnimNode_RigidBodyWithControl.cpp b/Engine/Plugins/Experimental/PhysicsControl/Source/PhysicsControl/Private/AnimNode_RigidBodyWithControl.cpp
--- a/Engine/Plugins/Experimental/PhysicsControl/Source/PhysicsControl/Private/AnimNode_RigidBodyWithControl.cpp
+++ b/Engine/Plugins/Experimental/PhysicsControl/Source/PhysicsControl/Private/AnimNode_RigidBodyWithControl.cpp
@@ -178,6 +178,7 @@ FAnimNode_RigidBodyWithControl::FAnimNode_RigidBodyWithControl()
 	, QueryParams(NAME_None, FCollisionQueryParams::GetUnknownStatId())
 	, UnsafeWorld(nullptr)
 	, UnsafeOwner(nullptr)
+	, SkeletalMeshAtLastPhysicsInit(nullptr)
 	, CapturedBoneVelocityBoneContainer()
 	, CapturedBoneVelocityPose()
 	, CapturedFrozenPose()
@@ -435,6 +436,13 @@ void FAnimNode_RigidBodyWithControl::CalculateSimulationSpace(
 	}
 }
 
+bool FAnimNode_RigidBodyWithControl::ShouldApplyPhysicsAssetConditioning() const
+{
+	return (PhysicsAssetConditioningConstraintPosition != MapConstraintsBehaviorType::None)
+		&& (PhysicsAssetConditioningConstraintOrientation != MapConstraintsBehaviorType::None)
+		&& (CVarEnableRigidBodyNodeWithControlMatchingConstraintsToSkeleton.GetValueOnAnyThread() > 0);
+}
+
 void FAnimNode_RigidBodyWithControl::RunPhysicsSimulation(float DeltaSeconds, const FVector& SimSpaceGravity)
 {
 	LLM_SCOPE_BYNAME(TEXT("Animation/RigidBodyWithControl"));
@@ -975,6 +983,9 @@ void FAnimNode_RigidBodyWithControl::InitPhysics(const UAnimInstance* InAnimInst
 		return;
 	}
 
+	// Record the Skeletal Mesh used to initialize physics so that it can be re-initialized if necessary when the Skeletal Mesh changes.
+	SkeletalMeshAtLastPhysicsInit = SkeletalMeshAsset;
+
 	const FReferenceSkeleton& SkelMeshRefSkel = SkeletalMeshAsset->GetRefSkeleton();
 	PhysicsAssetToUse = GetPhysicsAssetToBeUsed(InAnimInstance);
 
@@ -1038,7 +1049,7 @@ void FAnimNode_RigidBodyWithControl::InitPhysics(const UAnimInstance* InAnimInst
 			FPhysicsAggregateHandle(),
 			bCreateBodiesInRefPose);
 
-		if (CVarEnableRigidBodyNodeWithControlMatchingConstraintsToSkeleton.GetValueOnAnyThread() > 0)
+		if (ShouldApplyPhysicsAssetConditioning())
 		{
 			TransformConstraintsToMatchSkeletalMesh(SkeletalMeshAsset, PhysicsAssetConditioningConstraintPosition, PhysicsAssetConditioningConstraintOrientation, HighLevelConstraintInstances);
 		}
@@ -1436,10 +1447,15 @@ void FAnimNode_RigidBodyWithControl::SetOverridePhysicsAsset(UPhysicsAsset* Phys
 
 void FAnimNode_RigidBodyWithControl::PreUpdate(const UAnimInstance* InAnimInstance)
 {
-	// Detect changes in the physics asset to be used. This can happen when using the override
-	// physics asset feature.
+	USkeletalMeshComponent* const SKC = InAnimInstance->GetSkelMeshComponent();
+
+	// Re-initialize physics if the physics asset to be used has changed (This can happen when 
+	// using the override physics asset feature) or if the Skeletal Mesh has changed and Physics
+	// Asset Conditioning is required.
 	UPhysicsAsset* PhysicsAssetToBeUsed = GetPhysicsAssetToBeUsed(InAnimInstance);
-	if (GetPhysicsAsset() != PhysicsAssetToBeUsed)
+	const USkeletalMesh* const SkeletalMeshAsset = (SKC) ? SKC->GetSkeletalMeshAsset() : nullptr;
+
+	if ((GetPhysicsAsset() != PhysicsAssetToBeUsed) || ((SkeletalMeshAsset != SkeletalMeshAtLastPhysicsInit) && ShouldApplyPhysicsAssetConditioning()))
 	{
 		InitPhysics(InAnimInstance);
 
@@ -1456,7 +1472,6 @@ void FAnimNode_RigidBodyWithControl::PreUpdate(const UAnimInstance* InAnimInstan
 
 	SCOPE_CYCLE_COUNTER(STAT_RigidBodyNodeWithControl_PreUpdate);
 
-	USkeletalMeshComponent* SKC = InAnimInstance->GetSkelMeshComponent();
 	APawn* PawnOwner = InAnimInstance->TryGetPawnOwner();
 	UPawnMovementComponent* MovementComp = PawnOwner ? PawnOwner->GetMovementComponent() : nullptr;
 
diff --git a/Engine/Plugins/Experimental/PhysicsControl/Source/PhysicsControl/Public/AnimNode_RigidBodyWithControl.h b/Engine/Plugins/Experimental/PhysicsControl/Source/PhysicsControl/Public/AnimNode_RigidBodyWithControl.h
--- a/Engine/Plugins/Experimental/PhysicsControl/Source/PhysicsControl/Public/AnimNode_RigidBodyWithControl.h
+++ b/Engine/Plugins/Experimental/PhysicsControl/Source/PhysicsControl/Public/AnimNode_RigidBodyWithControl.h
@@ -495,6 +495,9 @@ private:
 		FSimulationSpaceData&    OutSimulationSpaceData
 	);
 
+	// Returns true if constraint transforms in the physics asset should be modified to match the current skeleton.
+	bool ShouldApplyPhysicsAssetConditioning() const;
+
 	// Modify Constraint transforms relative to the parent bone to correct for the difference
 	// between the Skeleton used to create the Physics asset and the current skeleton.
 	void TransformConstraintsToMatchSkeletalMesh(
@@ -665,6 +668,9 @@ private:
 	// Used by CollectWorldObjects and UpdateWorldGeometry in Task Thread. Only used for a pointer comparison.
 	const AActor* UnsafeOwner;
 
+	// The Skeletal Mesh parsed during the last physics initialization
+	const USkeletalMesh* SkeletalMeshAtLastPhysicsInit = nullptr;
+
 	FBoneContainer CapturedBoneVelocityBoneContainer;
 	FCSPose<FCompactHeapPose> CapturedBoneVelocityPose;
 	FCSPose<FCompactHeapPose> CapturedFrozenPose;
