Moving Physics in Control Rig to the PhysicsControl plugin, and switching over to doing it through components rather than hierarchy elements

diff --git a/Engine/Plugins/Animation/ControlRig/Source/ControlRig/Private/ControlRig.cpp b/Engine/Plugins/Animation/ControlRig/Source/ControlRig/Private/ControlRig.cpp
--- a/Engine/Plugins/Animation/ControlRig/Source/ControlRig/Private/ControlRig.cpp
+++ b/Engine/Plugins/Animation/ControlRig/Source/ControlRig/Private/ControlRig.cpp
@@ -566,8 +566,10 @@ void UControlRig::InitializeFromCDO()
 			UControlRig* CDO = GetClass()->GetDefaultObject<UControlRig>();
 			URigHierarchy* Hierarchy = GetHierarchy();
 
-			// copy physics solvers
-			PhysicsSolvers = CDO->PhysicsSolvers;
+			// Danny TODO this could copy the physics solvers from the CDO, but they are transient
+			// objects that only exist through construction, so I think we just need to start afresh
+			// each time. 
+			PhysicsSolvers.Empty();
 
 			// copy hierarchy
 			{
@@ -579,25 +581,6 @@ void UControlRig::InitializeFromCDO()
 				Hierarchy->ResetPoseToInitial(ERigElementType::All);
 			}
 
-			// update the physics solvers with new unique IDs and remap the physics solver IDs within hierarchy
-			{
-				TMap<FRigPhysicsSolverID, FRigPhysicsSolverID> PhysicsSolverMap;
-				for(FRigPhysicsSolverDescription& Solver : PhysicsSolvers)
-				{
-					const FRigPhysicsSolverID OldID = Solver.ID;
-					Solver.ID = FRigPhysicsSolverDescription::MakeID(GetPathName(), Solver.Name);
-					PhysicsSolverMap.Add(OldID, Solver.ID);
-				}
-				const TArray<FRigPhysicsElement*> PhysicsElements = Hierarchy->GetPhysicsElements();
-				for(FRigPhysicsElement* PhysicsElement : PhysicsElements)
-				{
-					if(const FRigPhysicsSolverID* NewID = PhysicsSolverMap.Find(PhysicsElement->Solver))
-					{
-						PhysicsElement->Solver = *NewID;
-					}
-				}
-			}
-
 #if WITH_EDITOR
 			// current hierarchy should always mirror CDO's hierarchy whenever a change of interest happens
 			CDO->GetHierarchy()->RegisterListeningHierarchy(Hierarchy);
@@ -1807,9 +1790,10 @@ void UControlRig::Serialize(FArchive& Ar)
 
 	Ar.UsingCustomVersion(FControlRigObjectVersion::GUID);
 
+	// Danny TODO solvers are transient and only made from the construction event
 	if (Ar.CustomVer(FControlRigObjectVersion::GUID) >= FControlRigObjectVersion::ControlRigStoresPhysicsSolvers)
 	{
-		Ar << PhysicsSolvers;
+		//Ar << PhysicsSolvers;
 		UE_RIGVM_ARCHIVETRACE_ENTRY(Ar, TEXT("PhysicsSolvers"));
 	}
 	else
@@ -3598,46 +3582,33 @@ void UControlRig::OnHierarchyTransformUndoRedo(URigHierarchy* InHierarchy, const
 	}
 }
 
-int32 UControlRig::NumPhysicsSolvers() const
+FRigPhysicsSolverBase* UControlRig::GetPhysicsSolver(const FRigPhysicsSolverID& InID)
 {
-	return PhysicsSolvers.Num();
+	TSharedPtr<FRigPhysicsSolverBase>* PhysicsSolver = PhysicsSolvers.Find(InID);
+	return PhysicsSolver ? PhysicsSolver->Get() : nullptr;
 }
 
-const FRigPhysicsSolverDescription* UControlRig::GetPhysicsSolver(int32 InIndex) const
+const FRigPhysicsSolverBase* UControlRig::GetPhysicsSolver(const FRigPhysicsSolverID& InID) const
 {
-	if(PhysicsSolvers.IsValidIndex(InIndex))
-	{
-		return &PhysicsSolvers[InIndex];
-	}
-	return nullptr;
+	const TSharedPtr<FRigPhysicsSolverBase>* PhysicsSolver = PhysicsSolvers.Find(InID);
+	return PhysicsSolver ? PhysicsSolver->Get() : nullptr;
 }
 
-const FRigPhysicsSolverDescription* UControlRig::FindPhysicsSolver(const FRigPhysicsSolverID& InID) const
+const FRigPhysicsSolverBase* UControlRig::FindPhysicsSolverByName(FName InName) const
 {
-	return PhysicsSolvers.FindByPredicate([InID](const FRigPhysicsSolverDescription& Solver) -> bool
+	for (const TPair<FRigPhysicsSolverID, TSharedPtr<FRigPhysicsSolverBase>>& SolverPair : PhysicsSolvers)
 	{
-		return Solver.ID == InID;
-	});
-}
-
-const FRigPhysicsSolverDescription* UControlRig::FindPhysicsSolverByName(const FName& InName) const
-{
-	return PhysicsSolvers.FindByPredicate([InName](const FRigPhysicsSolverDescription& Solver) -> bool
-	{
-		return Solver.Name == InName;
-	});
+		if (SolverPair.Value.IsValid() && SolverPair.Value->Name == InName)
+		{
+			return SolverPair.Value.Get();
+		}
+	}
+	return nullptr;
 }
 
-FRigPhysicsSolverID UControlRig::AddPhysicsSolver(FName InName, bool bSetupUndo, bool bPrintPythonCommand)
+FRigPhysicsSolverID UControlRig::RegisterPhysicsSolver(TSharedPtr<FRigPhysicsSolverBase> InPhysicsWorld, FName InName)
 {
-#if WITH_EDITOR
-	if(CVarControlRigHierarchyEnablePhysics.GetValueOnAnyThread() == false)
-	{
-		return FRigPhysicsSolverID();
-	}
-#endif
-
-	if(InName.IsNone())
+	if(InName.IsNone() || !InPhysicsWorld.IsValid())
 	{
 		return FRigPhysicsSolverID();
 	}
@@ -3647,11 +3618,15 @@ FRigPhysicsSolverID UControlRig::AddPhysicsSolver(FName InName, bool bSetupUndo,
 		return FindPhysicsSolverByName(InName) == nullptr;
 	});
 
-	FRigPhysicsSolverDescription Solver;
-	Solver.ID = FRigPhysicsSolverDescription::MakeID(GetPathName(), Name);
-	Solver.Name = Name;
-	PhysicsSolvers.Add(Solver);
-	return Solver.ID;
+	const FString CompletePath = FString::Printf(TEXT("%s|%s"), *GetPathName(), *Name.ToString());
+	FRigPhysicsSolverID SolverID(FGuid::NewDeterministicGuid(CompletePath));
+
+	PhysicsSolvers.FindOrAdd(SolverID, InPhysicsWorld);
+
+	InPhysicsWorld->Name = Name;
+	InPhysicsWorld->SolverID = SolverID;
+
+	return SolverID;
 }
 
 UControlRig::FPoseScope::FPoseScope(UControlRig* InControlRig, ERigElementType InFilter, const TArray<FRigElementKey>& InElements, const ERigTransformType::Type InTransformType)
diff --git a/Engine/Plugins/Animation/ControlRig/Source/ControlRig/Private/Rigs/RigHierarchy.cpp b/Engine/Plugins/Animation/ControlRig/Source/ControlRig/Private/Rigs/RigHierarchy.cpp
--- a/Engine/Plugins/Animation/ControlRig/Source/ControlRig/Private/Rigs/RigHierarchy.cpp
+++ b/Engine/Plugins/Animation/ControlRig/Source/ControlRig/Private/Rigs/RigHierarchy.cpp
@@ -52,7 +52,6 @@ static TAutoConsoleVariable<int32> CVarControlRigHierarchyTracePrecision(TEXT("C
 static TAutoConsoleVariable<int32> CVarControlRigHierarchyTraceOnSpawn(TEXT("ControlRig.Hierarchy.TraceOnSpawn"), 0, TEXT("sets the number of frames to trace when a new hierarchy is spawned"));
 TAutoConsoleVariable<bool> CVarControlRigHierarchyEnableRotationOrder(TEXT("ControlRig.Hierarchy.EnableRotationOrder"), true, TEXT("enables the rotation order for controls"));
 TAutoConsoleVariable<bool> CVarControlRigHierarchyEnableModules(TEXT("ControlRig.Hierarchy.Modules"), true, TEXT("enables the modular rigging functionality"));
-TAutoConsoleVariable<bool> CVarControlRigHierarchyEnablePhysics(TEXT("ControlRig.Hierarchy.Physics"), false, TEXT("enables physics support for the rig hierarchy"));
 static int32 sRigHierarchyLastTrace = INDEX_NONE;
 
 // A console command to trace a single frame / single execution for a control rig anim node / control rig component
@@ -808,7 +807,6 @@ void URigHierarchy::CopyHierarchy(URigHierarchy* InHierarchy)
 			sizeof(FRigNullElement),
 			sizeof(FRigControlElement),
 			sizeof(FRigCurveElement),
-			sizeof(FRigPhysicsElement),
 			sizeof(FRigReferenceElement),
 			sizeof(FRigConnectorElement),
 			sizeof(FRigSocketElement),
@@ -1680,24 +1678,6 @@ TArray<FRigElementKey> URigHierarchy::RestoreConnectorsFromStates(TArray<FRigCon
 	return Keys;
 }
 
-const FRigPhysicsSolverDescription* URigHierarchy::FindPhysicsSolver(const FRigPhysicsSolverID& InID) const
-{
-	if(const UControlRig* ControlRig = Cast<UControlRig>(GetOuter()))
-	{
-		return ControlRig->FindPhysicsSolver(InID);
-	}
-	return nullptr;
-}
-
-const FRigPhysicsSolverDescription* URigHierarchy::FindPhysicsSolverByName(const FName& InName) const
-{
-	if(const UControlRig* ControlRig = Cast<UControlRig>(GetOuter()))
-	{
-		return ControlRig->FindPhysicsSolverByName(InName);
-	}
-	return nullptr;
-}
-
 TArray<FName> URigHierarchy::GetMetadataNames(FRigElementKey InItem) const
 {
 	TArray<FName> Names;
@@ -5893,15 +5873,6 @@ FRigBaseElement* URigHierarchy::MakeElement(ERigElementType InElementType, int32
 			Element = NewElement<FRigCurveElement>(InCount);
 			break;
 		}
-		case ERigElementType::Physics:
-		{
-			if(OutStructureSize)
-			{
-				*OutStructureSize = sizeof(FRigPhysicsElement);
-			}
-			Element = NewElement<FRigPhysicsElement>(InCount);
-			break;
-		}
 		case ERigElementType::Reference:
 		{
 			if(OutStructureSize)
@@ -5999,18 +5970,6 @@ void URigHierarchy::DestroyElement(FRigBaseElement*& InElement)
 			}
 			break;
 		}
-		case ERigElementType::Physics:
-		{
-			FRigPhysicsElement* ExistingElements = Cast<FRigPhysicsElement>(InElement);
-			for(int32 Index=0;Index<Count;Index++)
-			{
-				TGuardValue<const FRigBaseElement*> DestroyGuard(ElementBeingDestroyed, &ExistingElements[Index]);
-				DestroyComponents(&ExistingElements[Index]);
-				DeallocateElementStorage(&ExistingElements[Index]);
-				ExistingElements[Index].~FRigPhysicsElement(); 
-			}
-			break;
-		}
 		case ERigElementType::Reference:
 		{
 			FRigReferenceElement* ExistingElements = Cast<FRigReferenceElement>(InElement);
diff --git a/Engine/Plugins/Animation/ControlRig/Source/ControlRig/Private/Rigs/RigHierarchyController.cpp b/Engine/Plugins/Animation/ControlRig/Source/ControlRig/Private/Rigs/RigHierarchyController.cpp
--- a/Engine/Plugins/Animation/ControlRig/Source/ControlRig/Private/Rigs/RigHierarchyController.cpp
+++ b/Engine/Plugins/Animation/ControlRig/Source/ControlRig/Private/Rigs/RigHierarchyController.cpp
@@ -522,81 +522,6 @@ FRigElementKey URigHierarchyController::AddCurve(FName InName, float InValue, bo
 	return NewElement->Key;
 }
 
-FRigElementKey URigHierarchyController::AddPhysicsElement(FName InName, FRigElementKey InParent, FRigPhysicsSolverID InSolver, 
-                                                          FRigPhysicsSettings InSettings, FTransform InLocalTransform, bool bSetupUndo, bool bPrintPythonCommand)
-{
-	if(!IsValid())
-	{
-		return FRigElementKey();
-	}
-
-#if WITH_EDITOR
-	if(CVarControlRigHierarchyEnablePhysics.GetValueOnAnyThread() == false)
-	{
-		return FRigElementKey();
-	}
-#endif
-
-	URigHierarchy* Hierarchy = GetHierarchy();
-
-	if(!InSolver.IsValid())
-	{
-		ReportError(TEXT("Physics Solver Guid is not valid"));
-		return FRigElementKey();
-	}
-
-	if(Hierarchy->FindPhysicsSolver(InSolver) == nullptr)
-	{
-		ReportErrorf(TEXT("Physics Solver with guid '%s' cannot be found."), *InSolver.ToString());
-		return FRigElementKey();
-	}
-
-#if WITH_EDITOR
-	TSharedPtr<FScopedTransaction> TransactionPtr;
-	if(bSetupUndo)
-	{
-		TransactionPtr = MakeShared<FScopedTransaction>(NSLOCTEXT("RigHierarchyController", "Add Physics Element", "Add Physics Element"));
-		Hierarchy->Modify();
-	}
-#endif
-
-	FRigPhysicsElement* NewElement = MakeElement<FRigPhysicsElement>();
-	{
-		TGuardValue<bool> DisableCacheValidityChecks(Hierarchy->bEnableCacheValidityCheck, false);		
-		NewElement->Key.Type = ERigElementType::Physics;
-		NewElement->Key.Name = GetSafeNewName(InName, NewElement->Key.Type);
-		NewElement->Solver = InSolver;
-		NewElement->Settings = InSettings;
-		AddElement(NewElement, Hierarchy->Get(Hierarchy->GetIndex(InParent)), true, InName);
-
-		Hierarchy->SetTransform(NewElement, InLocalTransform, ERigTransformType::InitialLocal, true, false);
-		NewElement->GetTransform().Current = NewElement->GetTransform().Initial;
-		NewElement->GetDirtyState().Current = NewElement->GetDirtyState().Initial;
-	}
-	
-#if WITH_EDITOR
-	TransactionPtr.Reset();
-
-	if (bPrintPythonCommand && !bSuspendPythonPrinting)
-	{
-		UBlueprint* Blueprint = GetTypedOuter<UBlueprint>();
-		if (Blueprint)
-		{
-			TArray<FString> Commands = GetAddPhysicsElementPythonCommands(NewElement);
-			for (const FString& Command : Commands)
-			{
-				RigVMPythonUtils::Print(Blueprint->GetFName().ToString(),
-					FString::Printf(TEXT("%s"), *Command));
-			}
-		}
-	}
-#endif
-
-	GetHierarchy()->EnsureCacheValidity();
-
-	return NewElement->Key;
-}
-
 FRigElementKey URigHierarchyController::AddReference(FName InName, FRigElementKey InParent,
 	FRigReferenceGetWorldTransformDelegate InDelegate, bool bSetupUndo)
 {
@@ -1583,12 +1508,6 @@ FString URigHierarchyController::ExportToText(TArray<FRigElementKey> InKeys) con
 				FRigCurveElement::StaticStruct()->ExportText(PerElementData.Content, Element, &DefaultElement, nullptr, PPF_None, nullptr);
 				break;
 			}
-			case ERigElementType::Physics:
-			{
-				FRigPhysicsElement DefaultElement;
-				FRigPhysicsElement::StaticStruct()->ExportText(PerElementData.Content, Element, &DefaultElement, nullptr, PPF_None, nullptr);
-				break;
-			}
 			case ERigElementType::Reference:
 			{
 				FRigReferenceElement DefaultElement;
@@ -1739,12 +1658,6 @@ TArray<FRigElementKey> URigHierarchyController::ImportFromText(FString InContent
 				FRigCurveElement::StaticStruct()->ImportText(*PerElementData.Content, NewElement, nullptr, EPropertyPortFlags::PPF_None, &ErrorPipe, FRigCurveElement::StaticStruct()->GetName(), true);
 				break;
 			}
-			case ERigElementType::Physics:
-			{
-				NewElement = MakeElement<FRigPhysicsElement>(true);
-				FRigPhysicsElement::StaticStruct()->ImportText(*PerElementData.Content, NewElement, nullptr, EPropertyPortFlags::PPF_None, &ErrorPipe, FRigPhysicsElement::StaticStruct()->GetName(), true);
-				break;
-			}
 			case ERigElementType::Reference:
 			{
 				NewElement = MakeElement<FRigReferenceElement>(true);
@@ -2084,10 +1997,6 @@ TArray<FString> URigHierarchyController::GetAddElementPythonCommands(FRigBaseEle
 	{
 		return GetAddCurvePythonCommands(CurveElement);
 	}
-	else if(FRigPhysicsElement* PhysicsElement = Cast<FRigPhysicsElement>(Element))
-	{
-		return GetAddPhysicsElementPythonCommands(PhysicsElement);
-	}
 	else if(FRigReferenceElement* ReferenceElement = Cast<FRigReferenceElement>(Element))
 	{
 		ensure(false);
@@ -2198,42 +2107,6 @@ TArray<FString> URigHierarchyController::GetAddCurvePythonCommands(FRigCurveElem
 		Hierarchy->GetCurveValue(Curve))};
 }
 
-TArray<FString> URigHierarchyController::GetAddPhysicsElementPythonCommands(FRigPhysicsElement* PhysicsElement) const
-{
-	TArray<FString> Commands;
-	FString TransformStr = RigVMPythonUtils::TransformToPythonString(PhysicsElement->GetTransform().Initial.Local.Get());
-	
-	FString ParentKeyStr = "''";
-	if (PhysicsElement->ParentElement)
-	{
-		ParentKeyStr = PhysicsElement->ParentElement->GetKey().ToPythonString();
-	}
-
-	FString PhysicsElementNamePythonized = RigVMPythonUtils::PythonizeName(PhysicsElement->GetName());
-	FRigPhysicsSettings& Settings = PhysicsElement->Settings;
-	FString SettingsStr;
-	{
-		SettingsStr =FString::Printf(TEXT("physics_settings%s"),
-			*PhysicsElementNamePythonized);
-			
-		Commands.Add(FString::Printf(TEXT("physics_settings%s = unreal.RigPhysicsSettings()"),
-			*PhysicsElementNamePythonized));
-
-		Commands.Add(FString::Printf(TEXT("control_settings_%s.mass = %f"),
-			*PhysicsElementNamePythonized,
-			Settings.Mass));	
-	}
-	
-	// FRigElementKey AddPhysicsElement(FName InName, FRigElementKey InParent, FRigPhysicsSettings InSettings, FTransform InLocalTransform, bool bSetupUndo = false);
-	Commands.Add(FString::Printf(TEXT("hierarchy_controller.add_physics_element('%s', %s, %s, %s)"),
-		*PhysicsElement->GetName(),
-		*ParentKeyStr,
-		*SettingsStr,
-		*TransformStr));
-
-	return Commands;
-}
-
 TArray<FString> URigHierarchyController::GetAddConnectorPythonCommands(FRigConnectorElement* Connector) const
 {
 	TArray<FString> Commands;
diff --git a/Engine/Plugins/Animation/ControlRig/Source/ControlRig/Private/Rigs/RigHierarchyElements.cpp b/Engine/Plugins/Animation/ControlRig/Source/ControlRig/Private/Rigs/RigHierarchyElements.cpp
--- a/Engine/Plugins/Animation/ControlRig/Source/ControlRig/Private/Rigs/RigHierarchyElements.cpp
+++ b/Engine/Plugins/Animation/ControlRig/Source/ControlRig/Private/Rigs/RigHierarchyElements.cpp
@@ -2055,110 +2055,6 @@ void FRigCurveElement::CopyFrom(const FRigBaseElement* InOther)
 	}
 }
 
-////////////////////////////////////////////////////////////////////////////////
-// FRigPhysicsSolver
-////////////////////////////////////////////////////////////////////////////////
-
-void FRigPhysicsSolverDescription::Serialize(FArchive& Ar)
-{
-	if (Ar.IsSaving() || Ar.IsObjectReferenceCollector() || Ar.IsCountingMemory())
-	{
-		Save(Ar);
-	}
-	else if (Ar.IsLoading())
-	{
-		Load(Ar);
-	}
-}
-
-void FRigPhysicsSolverDescription::Save(FArchive& Ar)
-{
-	Ar << ID;
-	Ar << Name;
-}
-
-void FRigPhysicsSolverDescription::Load(FArchive& Ar)
-{
-	Ar << ID;
-	Ar << Name;
-}
-
-FGuid FRigPhysicsSolverDescription::MakeGuid(const FString& InObjectPath, const FName& InSolverName)
-{
-	const FString CompletePath = FString::Printf(TEXT("%s|%s"), *InObjectPath, *InSolverName.ToString());
-	return FGuid::NewDeterministicGuid(CompletePath);
-}
-
-FRigPhysicsSolverID FRigPhysicsSolverDescription::MakeID(const FString& InObjectPath, const FName& InSolverName)
-{
-	return FRigPhysicsSolverID(MakeGuid(InObjectPath, InSolverName));
-}
-
-void FRigPhysicsSolverDescription::CopyFrom(const FRigPhysicsSolverDescription* InOther)
-{
-	if(InOther)
-	{
-		ID = InOther->ID;
-		Name = InOther->Name;
-	}
-}
-
-////////////////////////////////////////////////////////////////////////////////
-// FRigPhysicsSettings
-////////////////////////////////////////////////////////////////////////////////
-
-FRigPhysicsSettings::FRigPhysicsSettings()
-	: Mass(1.f)
-{
-}
-
-void FRigPhysicsSettings::Save(FArchive& Ar)
-{
-	Ar << Mass;
-}
-
-void FRigPhysicsSettings::Load(FArchive& Ar)
-{
-	Ar << Mass;
-}
-
-////////////////////////////////////////////////////////////////////////////////
-// FRigPhysicsElement
-////////////////////////////////////////////////////////////////////////////////
-
-const FRigBaseElement::EElementIndex FRigPhysicsElement::ElementTypeIndex = PhysicsElement;
-
-void FRigPhysicsElement::Save(FArchive& Ar, const FRigHierarchySerializationSettings& InSettings)
-{
-	Super::Save(Ar, InSettings);
-
-	if(InSettings.SerializationPhase == FRigHierarchySerializationSettings::StaticData)
-	{
-		Ar << Solver;
-		Settings.Save(Ar);
-	}
-}
-
-void FRigPhysicsElement::Load(FArchive& Ar, const FRigHierarchySerializationSettings& InSettings)
-{
-	Super::Load(Ar, InSettings);
-
-	if(InSettings.SerializationPhase == FRigHierarchySerializationSettings::StaticData)
-	{
-		Ar << Solver;
-		Settings.Load(Ar);
-	}
-}
-
-void FRigPhysicsElement::CopyFrom(const FRigBaseElement* InOther)
-{
-	Super::CopyFrom(InOther);
-	
-	const FRigPhysicsElement* Source = CastChecked<FRigPhysicsElement>(InOther);
-	Solver = Source->Solver;
-	Settings = Source->Settings;
-}
-
 ////////////////////////////////////////////////////////////////////////////////
 // FRigReferenceElement
 ////////////////////////////////////////////////////////////////////////////////
diff --git a/Engine/Plugins/Animation/ControlRig/Source/ControlRig/Private/Units/Execution/RigUnit_Physics.cpp b/Engine/Plugins/Animation/ControlRig/Source/ControlRig/Private/Units/Execution/RigUnit_Physics.cpp
deleted file mode 100644
diff --git a/Engine/Plugins/Animation/ControlRig/Source/ControlRig/Public/ControlRig.h b/Engine/Plugins/Animation/ControlRig/Source/ControlRig/Public/ControlRig.h
--- a/Engine/Plugins/Animation/ControlRig/Source/ControlRig/Public/ControlRig.h
+++ b/Engine/Plugins/Animation/ControlRig/Source/ControlRig/Public/ControlRig.h
@@ -21,6 +21,7 @@
 #include "Animation/MeshDeformerInstance.h"
 #include "Animation/MeshDeformerProducer.h"
 #include "Rigs/RigModuleDefines.h"
+#include "Rigs/RigPhysics.h"
 
 #if WITH_EDITOR
 #include "RigVMModel/RigVMPin.h"
@@ -44,6 +45,7 @@ class UTransformableControlHandle;
 struct FReferenceSkeleton;
 struct FRigUnit;
 struct FRigControl;
+struct FRigPhysicsSolverBase;
 
 CONTROLRIG_API DECLARE_LOG_CATEGORY_EXTERN(LogControlRig, Log, All);
 
@@ -714,50 +716,37 @@ protected:
 private:
 	float DebugBoneRadiusMultiplier;
 
-	// Physics Solvers
-	TArray<FRigPhysicsSolverDescription> PhysicsSolvers;
+	// Physics Solvers/Worlds
+	TMap<FRigPhysicsSolverID, TSharedPtr<FRigPhysicsSolverBase>> PhysicsSolvers;
 
 public:
 
 	/**
-	 * Returns the number of physics solvers
-	 * @return The number of physics solvers
-	 */
-	int32 NumPhysicsSolvers() const; 
-
-	/**
-	 * Returns a physics solver by index (or nullptr)
-	 * @param InIndex The index of the physics solver to return
-	 * @return The physics solver
+	 * Returns the physics world given a solver ID
+	 * @param InID The id identifying the physics solver
+	 * @return The physics solver (may be null)
 	 */
-	const FRigPhysicsSolverDescription* GetPhysicsSolver(int32 InIndex) const; 
+	FRigPhysicsSolverBase* GetPhysicsSolver(const FRigPhysicsSolverID& InID);
 
 	/**
-	 * Finds a new physics solver given its guid
+	 * Returns the physics world given a solver ID
 	 * @param InID The id identifying the physics solver
-	 * @return The physics solver
+	 * @return The physics solver (may be null)
 	 */
-	const FRigPhysicsSolverDescription* FindPhysicsSolver(const FRigPhysicsSolverID& InID) const; 
+	const FRigPhysicsSolverBase* GetPhysicsSolver(const FRigPhysicsSolverID& InID) const;
 
 	/**
-	 * Finds a new physics solver given its name
+	 * Finds a physics solver given its name
 	 * @param InName The name identifying the physics solver in the scope of this control rig
 	 * @return The physics solver
 	 */
-	const FRigPhysicsSolverDescription* FindPhysicsSolverByName(const FName& InName) const; 
+	const FRigPhysicsSolverBase* FindPhysicsSolverByName(FName InName) const; 
 
 	/**
-	 * Adds a physics solver to the hierarchy
-	 * @param InName The suggested name of the new physics solver - will eventually be corrected by the namespace
-	 * @param bSetupUndo If set to true the stack will record the change for undo / redo
-	 * @param bPrintPythonCommand If set to true a python command equivalent to this call will be printed out
-	 * @return The ID for the newly created physics solver.
+	 * Registers and stores (taking ownership) of a physics world, returning an ID for it. The world will be 
+	 * responsible for the low-level storage and simulation of objects.
 	 */
-	UFUNCTION(BlueprintCallable, Category = ControlRig)
-	FRigPhysicsSolverID AddPhysicsSolver(
-		FName InName,
-		bool bSetupUndo = false,
-		bool bPrintPythonCommand = false);
+	FRigPhysicsSolverID RegisterPhysicsSolver(TSharedPtr<FRigPhysicsSolverBase> PhysicsWorld, FName Name);
 
 #if WITH_EDITOR	
 
diff --git a/Engine/Plugins/Animation/ControlRig/Source/ControlRig/Public/Rigs/RigHierarchy.h b/Engine/Plugins/Animation/ControlRig/Source/ControlRig/Public/Rigs/RigHierarchy.h
--- a/Engine/Plugins/Animation/ControlRig/Source/ControlRig/Public/Rigs/RigHierarchy.h
+++ b/Engine/Plugins/Animation/ControlRig/Source/ControlRig/Public/Rigs/RigHierarchy.h
@@ -34,7 +34,6 @@ DECLARE_MULTICAST_DELEGATE_ThreeParams(FRigHierarchyMetadataTagChangedDelegate,
 
 extern CONTROLRIG_API TAutoConsoleVariable<bool> CVarControlRigHierarchyEnableRotationOrder;
 extern CONTROLRIG_API TAutoConsoleVariable<bool> CVarControlRigHierarchyEnableModules;
-extern CONTROLRIG_API TAutoConsoleVariable<bool> CVarControlRigHierarchyEnablePhysics;
 
 UENUM()
 enum ERigTransformStackEntryType : int
@@ -1087,25 +1086,6 @@ public:
 		return GetKeysOfType<FRigCurveElement>(false);
 	}
 
-	/**
-	 * Returns all Physics elements
-	 * @param bTraverse Returns the elements in order of a depth first traversal
-	 */
-	TArray<FRigPhysicsElement*> GetPhysicsElements(bool bTraverse = false) const
-	{
-		return GetElementsOfType<FRigPhysicsElement>(bTraverse);
-	}
-
-	/**
-	 * Returns all Physics elements
-	 * @param bTraverse Returns the elements in order of a depth first traversal
-	 */
-	UFUNCTION(BlueprintCallable, Category = URigHierarchy, meta = (DisplayName = "Get Physics Keys", ScriptName = "GetPhysicsKeys"))
-    TArray<FRigElementKey> GetPhysicsKeys(bool bTraverse = true) const
-	{
-		return GetKeysOfType<FRigPhysicsElement>(bTraverse);
-	}
-
 	/**
 	 * Returns all references
 	 * @param bTraverse Returns the elements in order of a depth first traversal
@@ -1210,20 +1190,6 @@ public:
 		}, false);
 	}
 
-	/**
-	 * Finds a new physics solver given its guid
-	 * @param InID The id identifying the physics solver
-	 * @return The physics solver
-	 */
-	const FRigPhysicsSolverDescription* FindPhysicsSolver(const FRigPhysicsSolverID& InID) const; 
-
-	/**
-	 * Finds a new physics solver given its name
-	 * @param InName The name identifying the physics solver in the scope of this hierarchy
-	 * @return The physics solver
-	 */
-	const FRigPhysicsSolverDescription* FindPhysicsSolverByName(const FName& InName) const; 
-
 	/**
 	 * Returns the name of metadata for a given element
 	 * @param InItem The element key to return the metadata keys for
diff --git a/Engine/Plugins/Animation/ControlRig/Source/ControlRig/Public/Rigs/RigHierarchyController.h b/Engine/Plugins/Animation/ControlRig/Source/ControlRig/Public/Rigs/RigHierarchyController.h
--- a/Engine/Plugins/Animation/ControlRig/Source/ControlRig/Public/Rigs/RigHierarchyController.h
+++ b/Engine/Plugins/Animation/ControlRig/Source/ControlRig/Public/Rigs/RigHierarchyController.h
@@ -200,27 +200,6 @@ public:
 		bool bPrintPythonCommand = false
         );
 
-	/**
-	* Adds a physics element to the hierarchy
-	* @param InName The suggested name of the new physics element - will eventually be corrected by the namespace
-	* @param InParent The (optional) parent of the new physics element. If you don't need a parent, pass FRigElementKey()
-	* @param InSolver The guid identifying the solver to use
-	* @param InSettings All of the physics element's settings
-	* @param InLocalTransform The transform for the new physics element - in the space of the provided parent
-	* @param bSetupUndo If set to true the stack will record the change for undo / redo
-	* @param bPrintPythonCommand If set to true a python command equivalent to this call will be printed out
-	* @return The key for the newly created physics element.
-	*/
-	UFUNCTION(BlueprintCallable, Category = URigHierarchyController)
-    FRigElementKey AddPhysicsElement(
-    	FName InName,
-    	FRigElementKey InParent,
-    	FRigPhysicsSolverID InSolver,
-        FRigPhysicsSettings InSettings,
-    	FTransform InLocalTransform,
-    	bool bSetupUndo = false,
-		bool bPrintPythonCommand = false);
-
 	/**
 	* Adds an reference to the hierarchy
 	* @param InName The suggested name of the new reference - will eventually be corrected by the namespace
@@ -733,8 +712,6 @@ public:
 
 	TArray<FString> GetAddCurvePythonCommands(FRigCurveElement* Curve) const;
 
-	TArray<FString> GetAddPhysicsElementPythonCommands(FRigPhysicsElement* PhysicsElement) const;
-
 	TArray<FString> GetAddConnectorPythonCommands(FRigConnectorElement* Connector) const;
 
 	TArray<FString> GetAddSocketPythonCommands(FRigSocketElement* Socket) const;
diff --git a/Engine/Plugins/Animation/ControlRig/Source/ControlRig/Public/Rigs/RigHierarchyDefines.h b/Engine/Plugins/Animation/ControlRig/Source/ControlRig/Public/Rigs/RigHierarchyDefines.h
--- a/Engine/Plugins/Animation/ControlRig/Source/ControlRig/Public/Rigs/RigHierarchyDefines.h
+++ b/Engine/Plugins/Animation/ControlRig/Source/ControlRig/Public/Rigs/RigHierarchyDefines.h
@@ -29,7 +29,7 @@ enum class ERigElementType : uint8
 	Space = Null UMETA(Hidden),
 	Control = 0x004,
 	Curve = 0x008,
-	Physics = 0x010, 
+	Physics = 0x010 UMETA(Hidden),
 	Reference = 0x020,
 	Connector = 0x040,
 	Socket = 0x080,
diff --git a/Engine/Plugins/Animation/ControlRig/Source/ControlRig/Public/Rigs/RigHierarchyElements.h b/Engine/Plugins/Animation/ControlRig/Source/ControlRig/Public/Rigs/RigHierarchyElements.h
--- a/Engine/Plugins/Animation/ControlRig/Source/ControlRig/Public/Rigs/RigHierarchyElements.h
+++ b/Engine/Plugins/Animation/ControlRig/Source/ControlRig/Public/Rigs/RigHierarchyElements.h
@@ -1814,168 +1814,6 @@ private:
 	friend class FControlRigHierarchyRelinkElementStorage;
 };
 
-USTRUCT(BlueprintType)
-struct CONTROLRIG_API FRigPhysicsSolverID
-{
-public:
-	
-	GENERATED_BODY()
-	
-	FRigPhysicsSolverID()
-		: Guid()
-	{
-	}
-
-	explicit FRigPhysicsSolverID(const FGuid& InGuid)
-		: Guid(InGuid)
-	{
-	}
-
-	UPROPERTY(BlueprintReadOnly, EditAnywhere, Category = Physics, meta=(ShowOnlyInnerProperties))
-	FGuid Guid;
-
-	bool IsValid() const
-	{
-		return Guid.IsValid();
-	}
-
-	FString ToString() const
-	{
-		return Guid.ToString();
-	}
-
-	bool operator == (const FRigPhysicsSolverID& InOther) const
-	{
-		return Guid == InOther.Guid;
-	}
-
-	bool operator != (const FRigPhysicsSolverID& InOther) const
-	{
-		return !(*this == InOther);
-	}
-
-	friend FArchive& operator <<(FArchive& Ar, FRigPhysicsSolverID& ID)
-	{
-		Ar << ID.Guid;
-		return Ar;
-	}
-
-	friend uint32 GetTypeHash(const FRigPhysicsSolverID& InID)
-	{
-		return GetTypeHash(InID.Guid);
-	}
-};
-
-USTRUCT(BlueprintType)
-struct CONTROLRIG_API FRigPhysicsSolverDescription
-{
-public:
-	
-	GENERATED_BODY()
-
-	UPROPERTY(BlueprintReadOnly, EditAnywhere, Category = Physics, meta=(ShowOnlyInnerProperties))
-	FRigPhysicsSolverID ID;
-
-	UPROPERTY(BlueprintReadOnly, EditAnywhere, Category = Physics, meta=(ShowOnlyInnerProperties))
-	FName Name;
-
-	FRigPhysicsSolverDescription()
-	: ID()
-	, Name(NAME_None)
-	{}
-	
-	FRigPhysicsSolverDescription(const FRigPhysicsSolverDescription& InOther)
-	{
-		*this = InOther;
-	}
-	
-	FRigPhysicsSolverDescription& operator=(const FRigPhysicsSolverDescription& InOther)
-	{
-		CopyFrom(&InOther);
-		return *this;
-	}
-
-	void Serialize(FArchive& Ar);
-	void Save(FArchive& Ar);
-	void Load(FArchive& Ar);
-	friend FArchive& operator<<(FArchive& Ar, FRigPhysicsSolverDescription& P)
-	{
-		P.Serialize(Ar);
-		return Ar;
-	}
-
-	static FGuid MakeGuid(const FString& InObjectPath, const FName& InSolverName);
-	static FRigPhysicsSolverID MakeID(const FString& InObjectPath, const FName& InSolverName);
-
-private:
-
-	void CopyFrom(const FRigPhysicsSolverDescription* InOther);
-};
-
-USTRUCT(BlueprintType)
-struct CONTROLRIG_API FRigPhysicsSettings
-{
-	GENERATED_BODY()
-
-	FRigPhysicsSettings();
-
-	void Save(FArchive& Ar);
-	void Load(FArchive& Ar);
-
-
-	UPROPERTY(BlueprintReadWrite, EditAnywhere, Category = Control)
-	float Mass;
-};
-
-USTRUCT(BlueprintType)
-struct CONTROLRIG_API FRigPhysicsElement : public FRigSingleParentElement
-{
-public:
-	
-	GENERATED_BODY()
-	DECLARE_RIG_ELEMENT_METHODS(FRigPhysicsElement)
-
-	static const EElementIndex ElementTypeIndex;
-
-	UPROPERTY(BlueprintReadOnly, EditAnywhere, Category = Physics, meta=(ShowOnlyInnerProperties))
-	FRigPhysicsSolverID Solver;
-
-	UPROPERTY(BlueprintReadOnly, EditAnywhere, Category = Physics, meta=(ShowOnlyInnerProperties))
-	FRigPhysicsSettings Settings;
-	
-	FRigPhysicsElement()
-        : FRigPhysicsElement(nullptr)
-	{ }
-	
-	FRigPhysicsElement(const FRigPhysicsElement& InOther)
-	{
-		*this = InOther;
-	}
-	
-	FRigPhysicsElement& operator=(const FRigPhysicsElement& InOther)
-	{
-		Super::operator=(InOther);
-		Solver = InOther.Solver;
-		Settings = InOther.Settings;
-		return *this;
-	}
-	
-	virtual ~FRigPhysicsElement() override {}
-
-	virtual void Save(FArchive& A, const FRigHierarchySerializationSettings& InSettings) override;
-	virtual void Load(FArchive& Ar, const FRigHierarchySerializationSettings& InSettings) override;
-
-private:
-	explicit FRigPhysicsElement(URigHierarchy* InOwner)
-		: FRigSingleParentElement(InOwner, ERigElementType::Physics)
-	{ }
-	
-	virtual void CopyFrom(const FRigBaseElement* InOther) override;
-
-	friend class URigHierarchy;
-	friend struct FRigBaseElement;
-};
-
 USTRUCT(BlueprintType)
 struct CONTROLRIG_API FRigReferenceElement final : public FRigSingleParentElement
 {
diff --git a/Engine/Plugins/Animation/ControlRig/Source/ControlRig/Public/Rigs/RigPhysics.h b/Engine/Plugins/Animation/ControlRig/Source/ControlRig/Public/Rigs/RigPhysics.h
new file mode 100644
--- /dev/null
+++ b/Engine/Plugins/Animation/ControlRig/Source/ControlRig/Public/Rigs/RigPhysics.h
@@ -0,0 +1,67 @@
+// Copyright Epic Games, Inc. All Rights Reserved.
+
+#pragma once
+
+#include "CoreMinimal.h"
+#include "RigHierarchyComponents.h"
+#include "RigPhysics.generated.h"
+
+// This simply exposes FGuid to blueprint, and thus custom UI at some point.
+USTRUCT(BlueprintType)
+struct CONTROLRIG_API FRigPhysicsSolverID
+{
+public:
+	
+	GENERATED_BODY()
+	
+	UPROPERTY(BlueprintReadOnly, EditAnywhere, Category = Physics, meta=(ShowOnlyInnerProperties))
+	FGuid Guid;
+
+	FRigPhysicsSolverID() : Guid() { }
+	explicit FRigPhysicsSolverID(const FGuid& InGuid) : Guid(InGuid) { }
+
+	bool IsValid() const
+	{
+		return Guid.IsValid();
+	}
+
+	FString ToString() const
+	{
+		return Guid.ToString();
+	}
+
+	bool operator == (const FRigPhysicsSolverID& InOther) const
+	{
+		return Guid == InOther.Guid;
+	}
+
+	bool operator != (const FRigPhysicsSolverID& InOther) const
+	{
+		return !(*this == InOther);
+	}
+
+	friend FArchive& operator <<(FArchive& Ar, FRigPhysicsSolverID& ID)
+	{
+		Ar << ID.Guid;
+		return Ar;
+	}
+
+	friend uint32 GetTypeHash(const FRigPhysicsSolverID& InID)
+	{
+		return GetTypeHash(InID.Guid);
+	}
+};
+
+// Danny TODO This is just inheriting from the base component to acquire the casting functions - is there a better way?
+USTRUCT(BlueprintType)
+struct CONTROLRIG_API FRigPhysicsSolverBase : public FRigBaseComponent
+{
+	GENERATED_BODY()
+	DECLARE_RIG_COMPONENT_METHODS(FRigPhysicsSolverBase)
+
+	virtual ~FRigPhysicsSolverBase() {}
+
+	FName               Name;
+	FRigPhysicsSolverID SolverID;
+};
+
diff --git a/Engine/Plugins/Animation/ControlRig/Source/ControlRig/Public/Units/Execution/RigUnit_Physics.h b/Engine/Plugins/Animation/ControlRig/Source/ControlRig/Public/Units/Execution/RigUnit_Physics.h
deleted file mode 100644
diff --git a/Engine/Plugins/Animation/ControlRig/Source/ControlRigDeveloper/Private/Graph/ControlRigGraph.cpp b/Engine/Plugins/Animation/ControlRig/Source/ControlRigDeveloper/Private/Graph/ControlRigGraph.cpp
--- a/Engine/Plugins/Animation/ControlRig/Source/ControlRigDeveloper/Private/Graph/ControlRigGraph.cpp
+++ b/Engine/Plugins/Animation/ControlRig/Source/ControlRigDeveloper/Private/Graph/ControlRigGraph.cpp
@@ -92,7 +92,6 @@ void UControlRigGraph::CacheNameLists(URigHierarchy* InHierarchy, const FRigVMDr
 		TArray<TSharedPtr<FRigVMStringWithTag>>& NullNameList = ElementNameLists.FindChecked(ERigElementType::Null);
 		TArray<TSharedPtr<FRigVMStringWithTag>>& ControlNameList = ElementNameLists.FindChecked(ERigElementType::Control);
 		TArray<TSharedPtr<FRigVMStringWithTag>>& CurveNameList = ElementNameLists.FindChecked(ERigElementType::Curve);
-		TArray<TSharedPtr<FRigVMStringWithTag>>& PhysicsNameList = ElementNameLists.FindChecked(ERigElementType::Physics);
 		TArray<TSharedPtr<FRigVMStringWithTag>>& ReferenceNameList = ElementNameLists.FindChecked(ERigElementType::Reference);
 		TArray<TSharedPtr<FRigVMStringWithTag>>& ConnectorNameList = ElementNameLists.FindChecked(ERigElementType::Connector);
 		TArray<TSharedPtr<FRigVMStringWithTag>>& SocketNameList = ElementNameLists.FindChecked(ERigElementType::Socket);
@@ -103,7 +102,6 @@ void UControlRigGraph::CacheNameLists(URigHierarchy* InHierarchy, const FRigVMDr
 		CacheNameListForHierarchy<FRigControlElement>(ControlRig, InHierarchy, ControlNameList, false);
 		CacheNameListForHierarchy<FRigControlElement>(ControlRig, InHierarchy, ControlNameListWithoutAnimationChannels, true);
 		CacheNameListForHierarchy<FRigCurveElement>(ControlRig, InHierarchy, CurveNameList, false);
-		CacheNameListForHierarchy<FRigPhysicsElement>(ControlRig, InHierarchy, PhysicsNameList, false);
 		CacheNameListForHierarchy<FRigReferenceElement>(ControlRig, InHierarchy, ReferenceNameList, false);
 		CacheNameListForHierarchy<FRigConnectorElement>(ControlRig, InHierarchy, ConnectorNameList, false);
 		CacheNameListForHierarchy<FRigSocketElement>(ControlRig, InHierarchy, SocketNameList, false);
diff --git a/Engine/Plugins/Animation/ControlRig/Source/ControlRigEditor/Private/ControlRigElementDetails.cpp b/Engine/Plugins/Animation/ControlRig/Source/ControlRigEditor/Private/ControlRigElementDetails.cpp
--- a/Engine/Plugins/Animation/ControlRig/Source/ControlRigEditor/Private/ControlRigElementDetails.cpp
+++ b/Engine/Plugins/Animation/ControlRig/Source/ControlRigEditor/Private/ControlRigElementDetails.cpp
@@ -6476,98 +6476,4 @@ void FRigConnectionRuleDetails::OnRuleContentChanged()
 	StructPropertyHandle->GetParentHandle()->NotifyPostChange(EPropertyChangeType::ValueSet);
 }
 
-void FRigPhysicsElementDetails::CustomizeDetails(IDetailLayoutBuilder& DetailBuilder)
-{
-	FRigTransformElementDetails::CustomizeDetails(DetailBuilder);
-	CustomizeSettings(DetailBuilder);
-	CustomizeTransform(DetailBuilder);
-	CustomizeComponents(DetailBuilder);
-	CustomizeMetadata(DetailBuilder);
-}
-
-void FRigPhysicsElementDetails::CustomizeSettings(IDetailLayoutBuilder& DetailBuilder)
-{
-	if(PerElementInfos.IsEmpty())
-	{
-		return;
-	}
-
-	if(IsAnyElementNotOfType(ERigElementType::Physics))
-	{
-		return;
-	}
-
-	const TSharedPtr<IPropertyHandle> SolverHandle = DetailBuilder.GetProperty(GET_MEMBER_NAME_CHECKED(FRigPhysicsElement, Solver));
-	DetailBuilder.HideProperty(SolverHandle);
-
-	const TSharedPtr<IPropertyHandle> SettingsHandle = DetailBuilder.GetProperty(GET_MEMBER_NAME_CHECKED(FRigPhysicsElement, Settings));
-	DetailBuilder.HideProperty(SettingsHandle);
-
-	IDetailCategoryBuilder& PhysicsCategory = DetailBuilder.EditCategory(TEXT("Physics"), LOCTEXT("Physics", "Physics"));
-
-	FDetailWidgetRow& SolverNameRow = PhysicsCategory.AddCustomRow(LOCTEXT("SolverName", "Solver Name"));
-	SolverNameRow.IsEnabled(false);
-	SolverNameRow.NameContent()
-	[
-		SNew(STextBlock)
-		.Text(LOCTEXT("SolverName", "Solver Name"))
-		.Font(IDetailLayoutBuilder::GetDetailFont())
-	];
-	SolverNameRow.ValueContent()
-	[
-		SNew(SEditableText)
-		.Text(this, &FRigPhysicsElementDetails::GetSolverNameText)
-		.Font(IDetailLayoutBuilder::GetDetailFont())
-	];
-
-	PhysicsCategory
-	.AddProperty(SolverHandle)
-	.DisplayName(LOCTEXT("SolverGuid", "Solver Guid"))
-	.IsEnabled(false);
-
-	const bool bIsProcedural = IsAnyElementProcedural();
-
-	PhysicsCategory
-		.AddProperty(SettingsHandle->GetChildHandle(GET_MEMBER_NAME_CHECKED(FRigPhysicsSettings, Mass)))
-		.IsEnabled(!bIsProcedural);
-}
-
-FText FRigPhysicsElementDetails::GetSolverNameText() const
-{
-	if(PerElementInfos.Num() > 0)
-	{
-		if(const URigHierarchy* Hierarchy = PerElementInfos[0].GetHierarchy())
-		{
-			TOptional<FText> FirstValue;
-			
-			for(const FPerElementInfo& Info : PerElementInfos)
-			{
-				if(const FRigPhysicsElement* PhysicsElement = Info.Element.Get<FRigPhysicsElement>())
-				{
-					if(const FRigPhysicsSolverDescription* Solver = Hierarchy->FindPhysicsSolver(PhysicsElement->Solver))
-					{
-						const FText SolverName = FText::FromName(Solver->Name);
-						if(FirstValue.IsSet())
-						{
-							if(!FirstValue.GetValue().EqualTo(SolverName))
-							{
-								return ControlRigDetailsMultipleValues;	
-							}
-						}
-						else
-						{
-							FirstValue = SolverName;
-						}
-					}
-				}
-			}
-
-			if(FirstValue.IsSet())
-			{
-				return FirstValue.GetValue();
-			}
-		}
-	}
-	return FText();
-}
 #undef LOCTEXT_NAMESPACE
diff --git a/Engine/Plugins/Animation/ControlRig/Source/ControlRigEditor/Private/ControlRigElementDetails.h b/Engine/Plugins/Animation/ControlRig/Source/ControlRigEditor/Private/ControlRigElementDetails.h
--- a/Engine/Plugins/Animation/ControlRig/Source/ControlRigEditor/Private/ControlRigElementDetails.h
+++ b/Engine/Plugins/Animation/ControlRig/Source/ControlRigEditor/Private/ControlRigElementDetails.h
@@ -921,22 +921,3 @@ protected:
 	TAttribute<bool> EnabledAttribute;
 };
 
-class FRigPhysicsElementDetails : public FRigTransformElementDetails
-{
-public:
-
-	// Makes a new instance of this detail layout class for a specific detail view requesting it
-	static TSharedRef<IDetailCustomization> MakeInstance()
-	{
-		return MakeShareable(new FRigPhysicsElementDetails);
-	}
-
-	/** IDetailCustomization interface */
-	virtual void CustomizeDetails(IDetailLayoutBuilder& DetailBuilder) override;
-
-	void CustomizeSettings(IDetailLayoutBuilder& DetailBuilder);
-
-private:
-
-	FText GetSolverNameText() const;
-};
diff --git a/Engine/Plugins/Animation/ControlRig/Source/ControlRigEditor/Private/EditMode/SControlRigOutliner.cpp b/Engine/Plugins/Animation/ControlRig/Source/ControlRigEditor/Private/EditMode/SControlRigOutliner.cpp
--- a/Engine/Plugins/Animation/ControlRig/Source/ControlRigEditor/Private/EditMode/SControlRigOutliner.cpp
+++ b/Engine/Plugins/Animation/ControlRig/Source/ControlRigEditor/Private/EditMode/SControlRigOutliner.cpp
@@ -521,15 +521,7 @@ bool SMultiRigHierarchyTreeView::AddElement(UControlRig* InControlRig, const FRi
 			}
 			case ERigElementType::Physics:
 			{
-				if (!Settings.bShowPhysics)
-				{
-					return false;
-				}
-				if(CVarControlRigHierarchyEnablePhysics.GetValueOnAnyThread() == false)
-				{
-					return false;
-				}
-				break;
+				return false;
 			}
 			case ERigElementType::Reference:
 			{
diff --git a/Engine/Plugins/Animation/ControlRig/Source/ControlRigEditor/Private/Editor/ControlRigWrapperObject.cpp b/Engine/Plugins/Animation/ControlRig/Source/ControlRigEditor/Private/Editor/ControlRigWrapperObject.cpp
--- a/Engine/Plugins/Animation/ControlRig/Source/ControlRigEditor/Private/Editor/ControlRigWrapperObject.cpp
+++ b/Engine/Plugins/Animation/ControlRig/Source/ControlRigEditor/Private/Editor/ControlRigWrapperObject.cpp
@@ -45,10 +45,6 @@ UClass* UControlRigWrapperObject::GetClassForStruct(UScriptStruct* InStruct, boo
 			{
 				PropertyEditorModule.RegisterCustomClassLayout(WrapperClassName, FOnGetDetailCustomizationInstance::CreateStatic(&FRigSocketElementDetails::MakeInstance));
 			}
-			else if(InStruct == FRigPhysicsElement::StaticStruct())
-			{
-				PropertyEditorModule.RegisterCustomClassLayout(WrapperClassName, FOnGetDetailCustomizationInstance::CreateStatic(&FRigPhysicsElementDetails::MakeInstance));
-			}
 		}
 	}
 #endif
diff --git a/Engine/Plugins/Animation/ControlRig/Source/ControlRigEditor/Private/Editor/SRigHierarchyTreeView.cpp b/Engine/Plugins/Animation/ControlRig/Source/ControlRigEditor/Private/Editor/SRigHierarchyTreeView.cpp
--- a/Engine/Plugins/Animation/ControlRig/Source/ControlRigEditor/Private/Editor/SRigHierarchyTreeView.cpp
+++ b/Engine/Plugins/Animation/ControlRig/Source/ControlRigEditor/Private/Editor/SRigHierarchyTreeView.cpp
@@ -563,15 +563,7 @@ bool SRigHierarchyTreeView::AddElement(const FRigBaseElement* InElement)
 		}
 		case ERigElementType::Physics:
 		{
-			if(!Settings.bShowPhysics)
-			{
-				return false;
-			}
-			if(CVarControlRigHierarchyEnablePhysics.GetValueOnAnyThread() == false)
-			{
-				return false;
-			}
-			break;
+			return false;
 		}
 		case ERigElementType::Reference:
 		{
diff --git a/Engine/Plugins/Experimental/PhysicsControl/PhysicsControl.uplugin b/Engine/Plugins/Experimental/PhysicsControl/PhysicsControl.uplugin
--- a/Engine/Plugins/Experimental/PhysicsControl/PhysicsControl.uplugin
+++ b/Engine/Plugins/Experimental/PhysicsControl/PhysicsControl.uplugin
@@ -15,7 +15,18 @@
 	"IsBetaVersion": false,
 	"IsExperimentalVersion": true,
 	"Installed": false,
+	"Plugins": [
+		{
+			"Name": "ControlRig",
+			"Enabled": true
+		},
+	],
 	"Modules": [
+		{
+			"Name": "ControlRigPhysics",
+			"Type": "Runtime",
+			"LoadingPhase": "Default"
+		},
 		{
 			"Name": "PhysicsControl",
 			"Type": "Runtime",
diff --git a/Engine/Plugins/Experimental/PhysicsControl/Source/ControlRigPhysics/ControlRigPhysics.Build.cs b/Engine/Plugins/Experimental/PhysicsControl/Source/ControlRigPhysics/ControlRigPhysics.Build.cs
new file mode 100644
--- /dev/null
+++ b/Engine/Plugins/Experimental/PhysicsControl/Source/ControlRigPhysics/ControlRigPhysics.Build.cs
@@ -0,0 +1,33 @@
+// Copyright Epic Games, Inc. All Rights Reserved.
+
+using UnrealBuildTool;
+
+public class ControlRigPhysics : ModuleRules
+{
+	public ControlRigPhysics(ReadOnlyTargetRules Target) : base(Target)
+	{
+		PCHUsage = ModuleRules.PCHUsageMode.UseExplicitOrSharedPCHs;
+
+		PublicDependencyModuleNames.AddRange(
+			new string[]
+			{
+				"Core",
+				"PhysicsControl"
+			}
+		);
+		
+		PrivateDependencyModuleNames.AddRange(
+			new string[]
+			{
+				"ControlRig",
+				"Chaos",
+				"ChaosCore",
+				"Core",
+				"CoreUObject",
+				"Engine",
+				"PhysicsCore",
+				"RigVM"
+			}
+		);
+	}
+}
diff --git a/Engine/Plugins/Experimental/PhysicsControl/Source/ControlRigPhysics/Private/ControlRigPhysicsModule.cpp b/Engine/Plugins/Experimental/PhysicsControl/Source/ControlRigPhysics/Private/ControlRigPhysicsModule.cpp
new file mode 100644
--- /dev/null
+++ b/Engine/Plugins/Experimental/PhysicsControl/Source/ControlRigPhysics/Private/ControlRigPhysicsModule.cpp
@@ -0,0 +1,18 @@
+// Copyright Epic Games, Inc. All Rights Reserved.
+
+#include "ControlRigPhysicsModule.h"
+
+#define LOCTEXT_NAMESPACE "FControlRigPhysicsModule"
+
+void FControlRigPhysicsModule::StartupModule()
+{
+}
+
+void FControlRigPhysicsModule::ShutdownModule()
+{
+}
+
+#undef LOCTEXT_NAMESPACE
+	
+IMPLEMENT_MODULE(FControlRigPhysicsModule, ControlRigPhysics)
+
diff --git a/Engine/Plugins/Experimental/PhysicsControl/Source/ControlRigPhysics/Private/RigPhysicsComponents.cpp b/Engine/Plugins/Experimental/PhysicsControl/Source/ControlRigPhysics/Private/RigPhysicsComponents.cpp
new file mode 100644
--- /dev/null
+++ b/Engine/Plugins/Experimental/PhysicsControl/Source/ControlRigPhysics/Private/RigPhysicsComponents.cpp
@@ -0,0 +1,27 @@
+// Copyright Epic Games, Inc. All Rights Reserved.
+
+#include "RigPhysicsComponents.h"
+
+////////////////////////////////////////////////////////////////////////////////
+// FRigPhysicsComponent
+////////////////////////////////////////////////////////////////////////////////
+
+void FRigPhysicsComponent::Save(FArchive& Ar)
+{
+	FRigBaseComponent::Save(Ar);
+
+	Ar << BodySolverSettings;
+	Ar << Dynamics;
+	Ar << Collision;
+	Ar << Articulation;
+}
+
+void FRigPhysicsComponent::Load(FArchive& Ar)
+{
+	FRigBaseComponent::Load(Ar);
+	Ar << BodySolverSettings;
+	Ar << Dynamics;
+	Ar << Collision;
+	Ar << Articulation;
+}
+
diff --git a/Engine/Plugins/Experimental/PhysicsControl/Source/ControlRigPhysics/Private/RigPhysicsSolver.cpp b/Engine/Plugins/Experimental/PhysicsControl/Source/ControlRigPhysics/Private/RigPhysicsSolver.cpp
new file mode 100644
--- /dev/null
+++ b/Engine/Plugins/Experimental/PhysicsControl/Source/ControlRigPhysics/Private/RigPhysicsSolver.cpp
@@ -0,0 +1,534 @@
+// Copyright Epic Games, Inc. All Rights Reserved.
+
+#include "RigPhysicsSolver.h"
+
+#include "AnimNode_RigidBodyWithControl.h"
+#include "PhysicsControlPoseData.h"
+
+#include "ControlRig.h"
+#include "Rigs/RigHierarchy.h"
+#include "Rigs/RigHierarchyController.h"
+
+#include "Engine/World.h"
+
+#include "Physics/ImmediatePhysics/ImmediatePhysicsActorHandle.h"
+#include "Physics/ImmediatePhysics/ImmediatePhysicsJointHandle.h"
+#include "Physics/ImmediatePhysics/ImmediatePhysicsSimulation.h"
+#include "Physics/ImmediatePhysics/ImmediatePhysicsChaos/ImmediatePhysicsSimulation_Chaos.h"
+
+#include "Physics/Experimental/PhysScene_Chaos.h"
+
+#include "Chaos/ParticleHandle.h"
+#include "Chaos/PBDJointConstraints.h"
+#include "Chaos/ImplicitObject.h"
+#include "Chaos/ShapeInstance.h"
+#include "Chaos/Capsule.h"
+#include "Chaos/ChaosScene.h"
+#include "Chaos/PBDJointConstraintUtilities.h"
+
+#include "PhysicsEngine/BodyInstance.h"
+#include "PhysicsEngine/ConstraintInstance.h"
+#include "PhysicsEngine/BodySetup.h"
+
+DEFINE_LOG_CATEGORY(LogRigPhysics);
+
+//======================================================================================================================
+FRigPhysicsSolver::FRigPhysicsSolver(
+	TObjectPtr<UControlRig>                     InOwningControlRig,
+	const FRigPhysicsSolverSettings&            InSolverSettings,
+	const FRigPhysicsSimulationSpaceSettings&   InSimulationSpaceSettings)
+	: OwningControlRig(InOwningControlRig)
+	, SolverSettings(InSolverSettings)
+	, SimulationSpaceSettings(InSimulationSpaceSettings)
+{
+	// Danny TODO temporary hackery to allow creation in worker thread
+	Chaos::CVars::bEnableAsyncInitBody = true;
+}
+
+//======================================================================================================================
+void FRigPhysicsSolver::AddBody(FRigComponentKey RigComponentKey)
+{
+	// Just make the record for now. The body details may change, so only actually instantiate later.
+	FRigBodyRecord& Record = BodyRecords.Add(RigComponentKey);
+}
+
+//======================================================================================================================
+void FRigPhysicsSolver::InitialiseSimulation()
+{
+	DestroyPhysicsSimulation();
+
+	Simulation = MakeShared<ImmediatePhysics::FSimulation>();
+}
+
+//======================================================================================================================
+void FRigPhysicsSolver::DestroyPhysicsSimulation()
+{
+	for (TPair<FRigComponentKey, FRigBodyRecord>& BodyRecordPair : BodyRecords)
+	{
+		FRigBodyRecord& Record = BodyRecordPair.Value;
+		if (Record.ArticulationJointHandle)
+		{
+			Simulation->DestroyJoint(Record.ArticulationJointHandle);
+			Record.ArticulationJointHandle = nullptr;
+		}
+		if (Record.ActorHandle)
+		{
+			Simulation->DestroyActor(Record.ActorHandle);
+			Record.ActorHandle = nullptr;
+		}
+	}
+
+	for (ImmediatePhysics::FActorHandle* ActorHandle : AdditionalActorHandles)
+	{
+		Simulation->DestroyActor(ActorHandle);
+	}
+	AdditionalActorHandles.Empty();
+
+	Simulation.Reset();
+}
+
+//======================================================================================================================
+static void SetCommonProperties(const FRigPhysicsCollisionShape& Shape, FKShapeElem& ShapeElem)
+{
+	ShapeElem.RestOffset = Shape.RestOffset;
+	ShapeElem.SetName(Shape.Name);
+	ShapeElem.SetContributeToMass(Shape.bContributeToMass);
+#ifdef PER_SHAPE_COLLISION
+	// Note that FKShapeElem supports enabling/disabling collision per shape, but this is discarded by the immediate
+	// solver.
+	ShapeElem.SetCollisionEnabled(Shape.CollisionEnabled);
+#endif
+}
+
+//======================================================================================================================
+static Chaos::EJointMotionType GetMotionTypeFromLimitValue(float Value)
+{
+	if (Value > 0)
+	{
+		return Chaos::EJointMotionType::Limited;
+	}
+	else if (Value < 0)
+	{
+		return Chaos::EJointMotionType::Free;
+	}
+	else
+	{
+		return Chaos::EJointMotionType::Locked;
+	}
+}
+
+//======================================================================================================================
+ImmediatePhysics::FActorHandle* FRigPhysicsSolver::CreateBody(
+	const FRigPhysicsCollision& Collision, 
+	const FRigPhysicsDynamics*  Dynamics, 
+	const FTransform&           BodyRelSimSpaceTM) const
+{
+	// Danny TODO Use the correct materials (potentially tricky since we may be called on
+	// any thread - see the RBAN code for reference).
+
+	// This appears to be the only way
+	UBodySetup* BodySetup = NewObject<UBodySetup>();
+	BodySetup->CollisionTraceFlag = CTF_UseSimpleAsComplex;
+
+	// Note that the base class (FKShapeElem) for the shapes has got additional
+	// functionality such as enable/disable collision etc.
+
+	bool bAddedCollision = false;
+	for (const FRigPhysicsCollisionBox& Shape : Collision.Boxes)
+	{
+		FKBoxElem Elem(Shape.Extents.X, Shape.Extents.Y, Shape.Extents.Z);
+		SetCommonProperties(Shape, Elem);
+		Elem.Center = Shape.TM.GetTranslation();
+		Elem.Rotation = Shape.TM.Rotator();
+		BodySetup->AggGeom.BoxElems.Add(Elem);
+		bAddedCollision = true;
+	}
+
+	for (const FRigPhysicsCollisionSphere& Shape : Collision.Spheres)
+	{
+		FKSphereElem Elem(Shape.Radius);
+		SetCommonProperties(Shape, Elem);
+		Elem.Center = Shape.TM.GetTranslation();
+		// Note that there is no rotation
+		BodySetup->AggGeom.SphereElems.Add(Elem);
+		bAddedCollision = true;
+	}
+
+	for (const FRigPhysicsCollisionCapsule& Shape : Collision.Capsules)
+	{
+		FKSphylElem Elem(Shape.Radius, Shape.Length);
+		SetCommonProperties(Shape, Elem);
+		Elem.Center = Shape.TM.GetTranslation();
+		Elem.Rotation = Shape.TM.Rotator();
+		BodySetup->AggGeom.SphylElems.Add(Elem);
+		bAddedCollision = true;
+	}
+
+	// This is a bit odd, but in order to create a low-level immediate physics representation, you
+	// have to create a high-level temporary one.
+	//
+	// See USkeletalMeshComponent::InstantiatePhysicsAsset_Internal for how it's done in RBAN (which
+	// makes body instances, converts to immediate physics actors, and then discards the body instances).
+	FBodyInstance BodyInstance;
+	if (Dynamics)
+	{
+		BodyInstance.bSimulatePhysics = true;
+		BodyInstance.bEnableGravity = true;
+		BodyInstance.SetMassOverride(Dynamics->Mass);
+	}
+	else
+	{
+		BodyInstance.bSimulatePhysics = false;
+	}
+	BodyInstance.InitBody(BodySetup, BodyRelSimSpaceTM, nullptr, nullptr);
+
+	ImmediatePhysics::FActorHandle* ActorHandle = Dynamics ?
+		Simulation->CreateDynamicActor(&BodyInstance, BodyRelSimSpaceTM) :
+		Simulation->CreateKinematicActor(&BodyInstance, BodyRelSimSpaceTM);
+
+	// Creation is guaranteed to return something
+	check(ActorHandle);
+
+	Simulation->AddToCollidingPairs(ActorHandle);
+
+	if (bAddedCollision && Dynamics)
+	{
+		// Note that particles are always created disabled. They will simulate when disabled,
+		// but won't collide!
+		ActorHandle->SetEnabled(true);
+	}
+
+	return ActorHandle;
+}
+
+
+//======================================================================================================================
+void FRigPhysicsSolver::Instantiate()
+{
+	if (!bNeedToInstantiate)
+	{
+		return;
+	}
+
+	InitialiseSimulation();
+
+	if (!SolverSettings.Collision.IsEmpty())
+	{
+		FTransform BodyRelSimSpaceTM = ConvertCollisionSpaceTransformToSimSpace(FTransform());
+
+		ImmediatePhysics::FActorHandle* EnvironmentActorHandle = CreateBody(
+			SolverSettings.Collision, nullptr, BodyRelSimSpaceTM);
+		AdditionalActorHandles.Add(EnvironmentActorHandle);
+	}
+
+	URigHierarchy* Hierarchy = OwningControlRig->GetHierarchy();
+
+	for (TPair<FRigComponentKey, FRigBodyRecord>& BodyRecordPair : BodyRecords)
+	{
+		const FRigComponentKey& ComponentKey = BodyRecordPair.Key;
+		FRigBodyRecord& Record = BodyRecordPair.Value;
+
+		if (const FRigPhysicsComponent* PhysicsComponent = Cast<FRigPhysicsComponent>(
+			Hierarchy->FindComponent(ComponentKey)))
+		{
+			const FRigPhysicsCollision& Collision = PhysicsComponent->Collision;
+			const FRigPhysicsDynamics& Dynamics = PhysicsComponent->Dynamics;
+			const FPhysicsControlModifierData& BodyData = PhysicsComponent->BodyData;
+
+			FRigElementKey SourceKey = PhysicsComponent->BodySolverSettings.SourceBone;
+			if (!SourceKey.IsValid())
+			{
+				SourceKey = ComponentKey.ElementKey;
+			}
+
+			FTransform BodyRelComponentTM;
+			// What should we do if the key is not valid? 
+			if (SourceKey.IsValid())
+			{
+				BodyRelComponentTM = Hierarchy->GetGlobalTransform(SourceKey);
+			}
+
+			FTransform BodyRelSimSpaceTM = ConvertComponentSpaceTransformToSimSpace(BodyRelComponentTM);
+
+			Record.ActorHandle = CreateBody(Collision, &Dynamics, BodyRelSimSpaceTM);
+
+			// Note that these data will be applied in the PrePhysics update
+			Record.BodyData = BodyData;
+		}
+	}
+
+	// Once all the bodies are created, we can make their articulations
+	for (TPair<FRigComponentKey, FRigBodyRecord>& BodyRecordPair : BodyRecords)
+	{
+		FRigComponentKey& ChildBodyComponentKey = BodyRecordPair.Key;
+		FRigBodyRecord& ChildBodyRecord = BodyRecordPair.Value;
+
+		ImmediatePhysics::FActorHandle* ChildActorHandle = ChildBodyRecord.ActorHandle;
+		if (!ChildActorHandle)
+		{
+			continue;
+		}
+
+		if (const FRigPhysicsComponent* ChildPhysicsComponent = Cast<FRigPhysicsComponent>(
+			Hierarchy->FindComponent(ChildBodyComponentKey)))
+		{
+			const FRigPhysicsArticulation& Articulation = ChildPhysicsComponent->Articulation;
+			if (!Articulation.bEnable)
+			{
+				continue;
+			}
+
+			// Now we need to find the parent body (etc) for the articulation
+			FRigComponentKey ParentBodyComponentKey = Articulation.ParentBody;
+
+			// And along the way we will discover various other things
+			FRigElementKey ChildBoneKey = ChildBodyComponentKey.ElementKey;
+			if (!ChildBoneKey.IsValid())
+			{
+				continue;
+			}
+
+			if (!ParentBodyComponentKey.IsValid())
+			{
+				// Physics bodies are parented to a bone - we want the parent of the bone
+				FRigElementKey ParentBoneKey = Hierarchy->GetFirstParent(ChildBoneKey);
+				// Now we have the regular bone, we want the physics element that is the child
+				if (ParentBoneKey.IsValid())
+				{
+					TArray<FRigComponentKey> ComponentKeys = Hierarchy->GetComponentKeys(ParentBoneKey);
+					for (FRigComponentKey ComponentKey : ComponentKeys)
+					{
+						const FRigBaseComponent* BaseComponent = Hierarchy->FindComponent(ComponentKey);
+						// Danny TODO what's the best way to find/check this?!
+						if (Cast<FRigPhysicsComponent>(BaseComponent))
+						{
+							ParentBodyComponentKey = ComponentKey;
+							break;
+						}
+					}
+				}
+			}
+
+			const FRigPhysicsComponent* ParentPhysicsComponent = 
+				Cast<FRigPhysicsComponent>(Hierarchy->FindComponent(ParentBodyComponentKey));
+			if (!ParentPhysicsComponent)
+			{
+				continue;
+			}
+
+			FRigBodyRecord* ParentRecord = nullptr;
+			ImmediatePhysics::FActorHandle* ParentActorHandle = nullptr;
+
+			// Find the record that corresponds to the parent, so we can see if there is a parent actor
+			for (TPair<FRigComponentKey, FRigBodyRecord>& OtherBodyRecordPair : BodyRecords)
+			{
+				FRigComponentKey& OtherBodyComponentKey = OtherBodyRecordPair.Key;
+				FRigBodyRecord& OtherBodyRecord = OtherBodyRecordPair.Value;
+				if (OtherBodyComponentKey == ParentPhysicsComponent->GetKey())
+				{
+					ParentRecord = &OtherBodyRecord;
+					ParentActorHandle = OtherBodyRecord.ActorHandle;
+					break;
+				}
+			}
+			if (!ParentRecord)
+			{
+				continue;
+			}
+
+			// Make the articulation (joint constraint). 
+			// UE likes to treat Body1 (index 0) as the child, and Body2 (index 1) as the parent
+
+			FTransform ParentCoMTransform = ParentActorHandle->GetLocalCoMTransform();
+			FTransform ChildCoMTransform = ChildActorHandle->GetLocalCoMTransform();
+
+
+			Chaos::FPBDJointSettings ConstraintSettings;
+			if (Articulation.bAutoCalculateChildOffset)
+			{
+				// Danny TODO - is the "non-override" actually always simply the identity TM?
+				ConstraintSettings.ConnectorTransforms[0] = FTransform();
+			}
+			ConstraintSettings.ConnectorTransforms[0] = 
+				Articulation.ExtraChildOffset * ConstraintSettings.ConnectorTransforms[0];
+
+			if (Articulation.bAutoCalculateParentOffset)
+			{
+				// Danny TODO is it right to use the TM of the child bone relative to parent bone?
+				// What about if the body has got a TM too?
+				ConstraintSettings.ConnectorTransforms[1] = Hierarchy->GetLocalTransform(ChildBoneKey, true);
+			}
+			ConstraintSettings.ConnectorTransforms[1] =
+				Articulation.ExtraParentOffset * ConstraintSettings.ConnectorTransforms[1];
+
+			// Danny TODO The linear limit is likely implemented as a box limit (check) - it's just
+			// not exposed as one. We should change this.
+			ConstraintSettings.LinearLimit = Articulation.LinearLimit.Length(); 
+			Chaos::EJointMotionType LinearMotionType = GetMotionTypeFromLimitValue(ConstraintSettings.LinearLimit);
+			ConstraintSettings.LinearMotionTypes = { LinearMotionType, LinearMotionType, LinearMotionType };
+
+			ConstraintSettings.AngularLimits = FMath::DegreesToRadians(Articulation.AngularLimit);
+			ConstraintSettings.AngularMotionTypes = {
+				GetMotionTypeFromLimitValue(Articulation.AngularLimit[0]),
+				GetMotionTypeFromLimitValue(Articulation.AngularLimit[1]),
+				GetMotionTypeFromLimitValue(Articulation.AngularLimit[2])
+			};
+
+			ConstraintSettings.bCollisionEnabled = false;
+			ConstraintSettings.bUseLinearSolver = SolverSettings.bUseLinearJointSolver;
+
+			ChildBodyRecord.ArticulationJointHandle = Simulation->CreateJoint(
+				ConstraintSettings, ChildActorHandle, ParentActorHandle);
+		}
+	}
+
+	bNeedToInstantiate = false;
+}
+
+//======================================================================================================================
+void FRigPhysicsSolver::SetBodySparseData(
+	FRigComponentKey                         RigComponentKey,
+	const FPhysicsControlModifierSparseData& SparseData)
+{
+	if (FRigBodyRecord* Record = BodyRecords.Find(RigComponentKey))
+	{
+		Record->BodyData.UpdateFromSparseData(SparseData);
+	}
+}
+
+//======================================================================================================================
+void FRigPhysicsSolver::UpdatePrePhysics(const FSimulationSpaceData& SimulationSpaceData)
+{
+	URigHierarchy* Hierarchy = OwningControlRig->GetHierarchy();
+
+	for (ImmediatePhysics::FActorHandle* ActorHandle : AdditionalActorHandles)
+	{
+		FTransform BodyRelSimSpaceTM = ConvertCollisionSpaceTransformToSimSpace(FTransform());
+		ActorHandle->SetKinematicTarget(BodyRelSimSpaceTM);
+	}
+
+	for (TPair<FRigComponentKey, FRigBodyRecord>& BodyRecordPair : BodyRecords)
+	{
+		const FRigComponentKey& ComponentKey = BodyRecordPair.Key;
+		const FRigBodyRecord& Record = BodyRecordPair.Value;
+
+		if (const FRigPhysicsComponent* PhysicsComponent = Cast<FRigPhysicsComponent>(
+			Hierarchy->FindComponent(ComponentKey)))
+		{
+			if (Record.ActorHandle)
+			{
+				UpdateBodyFromModifierData(
+					Record.ActorHandle, Simulation.Get(), Record.BodyData, SimulationSpaceData.Gravity);
+
+				if (Record.ActorHandle->GetIsKinematic())
+				{
+					FRigElementKey SourceKey = PhysicsComponent->BodySolverSettings.SourceBone;
+					if (!SourceKey.IsValid())
+					{
+						SourceKey = ComponentKey.ElementKey;
+					}
+					if (SourceKey.IsValid())
+					{
+						FTransform SourceComponentSpaceTM = Hierarchy->GetGlobalTransform(SourceKey);
+
+						// Danny TODO apply explicit targets here
+						FTransform SourceSimSpaceTM = ConvertComponentSpaceTransformToSimSpace(SourceComponentSpaceTM);
+						Record.ActorHandle->SetKinematicTarget(SourceSimSpaceTM);
+					}
+				}
+			}
+		}
+	}
+}
+
+//======================================================================================================================
+// Note that we read back into a target bone, which may have been specified explicitly, or will
+// otherwise default to the physics element parent.
+void FRigPhysicsSolver::UpdatePostPhysics()
+{
+	URigHierarchy* Hierarchy = OwningControlRig->GetHierarchy();
+	for (TPair<FRigComponentKey, FRigBodyRecord>& BodyRecordPair : BodyRecords)
+	{
+		const FRigComponentKey& ComponentKey = BodyRecordPair.Key;
+		const FRigBodyRecord& Record = BodyRecordPair.Value;
+
+		if (Record.ActorHandle)
+		{
+			if (const FRigPhysicsComponent* PhysicsComponent = Cast<FRigPhysicsComponent>(
+				Hierarchy->FindComponent(ComponentKey)))
+			{
+				// We will only have created this body record if it's associated with our solver, so
+				// we don't need to check again that our solver is in its list.
+				FRigElementKey TargetElementKey = PhysicsComponent->BodySolverSettings.TargetBone;
+				if (!TargetElementKey.IsValid())
+				{
+					TargetElementKey = ComponentKey.ElementKey;
+				}
+
+				if (TargetElementKey.IsValid())
+				{
+					FTransform TargetSimSpaceTM = Record.ActorHandle->GetWorldTransform();
+					FTransform TargetComponentSpaceTM = ConvertSimSpaceTransformToComponentSpace(TargetSimSpaceTM);
+					Hierarchy->SetGlobalTransform(TargetElementKey, TargetComponentSpaceTM);
+				}
+			}
+		}
+	}
+}
+
+//======================================================================================================================
+void FRigPhysicsSolver::StepSimulation(
+	double DeltaTime, const UWorld* DebugWorld)
+{
+	// We need to know about the simulation space etc before we can instantiate anything into the right place
+	FSimulationSpaceData SimulationSpaceData = UpdateSimulationSpaceStateAndCalculateData(DeltaTime);
+
+	// We instantiate when we do the first simulation - this makes sure any changes applied by
+	// the user have been made. It also means there is no overhead if physics is never stepped.
+	// However, there may be a hitch due to the creation, so it may also happen during construction.
+	Instantiate(); // There is an early out if it's already been done
+
+	if (!Simulation)
+	{
+		return;
+	}
+
+	Simulation->SetSolverSettings(
+		SolverSettings.FixedTimeStep,
+		999999999.0, // Danny TODO cull distance - expose this? Yuck
+		SolverSettings.MaxDepenetrationVelocity,
+		SolverSettings.bUseLinearJointSolver,
+		SolverSettings.PositionIterations,
+		SolverSettings.VelocityIterations,
+		SolverSettings.ProjectionIterations,
+		SolverSettings.bUseManifolds);
+
+	ChaosJointSolverSettings.bSolvePositionLast = SolverSettings.bSolveJointPositionsLast;
+
+	Simulation->SetSimulationSpaceSettings(
+		SimulationSpaceSettings.SpaceMovementAmount,
+		SimulationSpaceSettings.ExternalLinearDragV);
+
+	Simulation->UpdateSimulationSpace(
+		SimulationSpaceData.Transform,
+		SimulationSpaceData.LinearVelocity,
+		SimulationSpaceData.AngularVelocity,
+		SimulationSpaceData.LinearAcceleration,
+		SimulationSpaceData.AngularAcceleration);
+
+	UpdatePrePhysics(SimulationSpaceData);
+
+	Simulation->Simulate(
+		DeltaTime, SolverSettings.MaxDeltaTime, SolverSettings.MaxTimeSteps, 
+		SimulationSpaceData.Gravity, &ChaosJointSolverSettings);
+
+#if CHAOS_DEBUG_DRAW
+	// Danny TODO enable/disable chaos debug draw with a cvar or flag
+	if (DebugWorld && DebugWorld->GetPhysicsScene() && DebugWorld->GetPhysicsScene()->GetDebugDrawScene())
+	{
+		Simulation->SetDebugDrawScene(FString("ControlRig"), DebugWorld->GetPhysicsScene()->GetDebugDrawScene());
+		Simulation->DebugDraw();
+	}
+#endif
+	UpdatePostPhysics();
+}
diff --git a/Engine/Plugins/Experimental/PhysicsControl/Source/ControlRigPhysics/Private/RigPhysicsSolver.h b/Engine/Plugins/Experimental/PhysicsControl/Source/ControlRigPhysics/Private/RigPhysicsSolver.h
new file mode 100644
--- /dev/null
+++ b/Engine/Plugins/Experimental/PhysicsControl/Source/ControlRigPhysics/Private/RigPhysicsSolver.h
@@ -0,0 +1,136 @@
+// Copyright Epic Games, Inc. All Rights Reserved.
+
+#pragma once
+
+#include "RigPhysicsComponents.h"
+
+#include "Rigs/RigHierarchyDefines.h"
+#include "Rigs/RigHierarchyComponents.h"
+
+#include "Physics/ImmediatePhysics/ImmediatePhysicsDeclares.h"
+
+#include "PhysicsControlData.h"
+#include "Chaos/PBDJointConstraintTypes.h"
+
+#include "RigPhysicsSolver.generated.h"
+
+class UControlRig;
+
+struct FRigVMDrawInterface;
+
+#if UE_BUILD_SHIPPING || UE_BUILD_TEST
+DECLARE_LOG_CATEGORY_EXTERN(LogRigPhysics, Log, Warning);
+#else
+DECLARE_LOG_CATEGORY_EXTERN(LogRigPhysics, Log, All);
+#endif
+
+USTRUCT()
+struct FRigBodyRecord
+{
+	GENERATED_BODY()
+
+	ImmediatePhysics::FActorHandle* ActorHandle = nullptr;
+	ImmediatePhysics::FJointHandle* ArticulationJointHandle = nullptr;
+	FPhysicsControlModifierData     BodyData;
+};
+
+USTRUCT(BlueprintType)
+struct FRigPhysicsSolver : public FRigPhysicsSolverBase
+{
+public:
+	GENERATED_BODY()
+
+	FRigPhysicsSolver(
+		TObjectPtr<UControlRig>                     InOwningControlRig = nullptr, 
+		const FRigPhysicsSolverSettings&            InSolverSettings = FRigPhysicsSolverSettings(),
+		const FRigPhysicsSimulationSpaceSettings&   InSimulationSpaceSettings = FRigPhysicsSimulationSpaceSettings());
+
+	void AddBody(FRigComponentKey RigComponentKey);
+
+	void SetBodySparseData(FRigComponentKey RigComponentKey, const FPhysicsControlModifierSparseData& SparseData);
+
+	// This will create all the bodies etc in the simulation
+	void Instantiate();
+
+	// Integrates the simulation forwards by DeltaTime
+	void StepSimulation(double DeltaTime, const UWorld* DebugWorld);
+
+	void Draw(FRigVMDrawInterface* DI) const;
+
+private:
+	// Used by the world-space to simulation-space motion transfer system in Component- or Bone-Space sims
+	struct FSimulationSpaceState
+	{
+		FTransform ComponentTM;
+		FTransform BoneRelComponentTM;
+		FVector    ComponentLinearVelocity;
+		FVector    ComponentAngularVelocity;
+		FVector    BoneLinearVelocity;
+		FVector    BoneAngularVelocity;
+	};
+	// Represents the properties of the simulation space
+	struct FSimulationSpaceData
+	{
+		FSimulationSpaceData(const FTransform& TM) : Transform(TM) {}
+		// The world transform of the simulation space
+		FTransform Transform;
+		FVector    LinearVelocity;
+		FVector    AngularVelocity;
+		FVector    LinearAcceleration;
+		FVector    AngularAcceleration;
+		FVector    Gravity;
+	};
+
+	void InitialiseSimulation();
+
+	// Creates an actor with collision. This will by dynamic if Dynamics is valid, or otherwise kinematic
+	ImmediatePhysics::FActorHandle* CreateBody(
+		const FRigPhysicsCollision& Collision,
+		const FRigPhysicsDynamics*  Dynamics,
+		const FTransform&           BodyRelSimSpaceTM) const;
+
+	// This releases the actors etc in the simulation, and destroys the simulation, but it doesn't
+	// change our records (apart from resetting pointers to actors etc).
+	void DestroyPhysicsSimulation();
+
+	void UpdatePrePhysics(const FSimulationSpaceData& SimulationSpaceData);
+	void UpdatePostPhysics();
+
+	void InitSimulationSpace(const FTransform& ComponentToWorld, const FTransform& BoneToComponent);
+
+	FSimulationSpaceData UpdateSimulationSpaceStateAndCalculateData(const float Dt);
+
+	// Returns the simulation space transform, in world space
+	FTransform GetSimulationSpaceTransform() const;
+
+	// Converts a transform from component space (e.g. coming from the owning control rig) into the
+	// simulation space
+	FTransform ConvertComponentSpaceTransformToSimSpace(const FTransform& TM) const;
+
+	// Converts a transform from the simulation space to component space (e.g. for writing back to
+	// the owning control rig)
+	FTransform ConvertSimSpaceTransformToComponentSpace(const FTransform& TM) const;
+
+	// Converts a vector specified in world space to the simulation space (e.g. converting gravity)
+	FVector ConvertWorldVectorToSimSpaceNoScale(const FVector& WorldVector) const;
+
+	FTransform ConvertCollisionSpaceTransformToSimSpace(const FTransform& TM) const;
+
+private:
+	TObjectPtr<UControlRig> OwningControlRig;
+
+	TMap<FRigComponentKey, FRigBodyRecord> BodyRecords;
+	TSharedPtr<ImmediatePhysics::FSimulation> Simulation;
+
+	FRigPhysicsSolverSettings SolverSettings;
+	FRigPhysicsSimulationSpaceSettings SimulationSpaceSettings;
+
+	TArray<ImmediatePhysics::FActorHandle*> AdditionalActorHandles;
+
+	Chaos::FPBDJointSolverSettings ChaosJointSolverSettings;
+
+	FSimulationSpaceState SimulationSpaceState;
+
+	bool bNeedToInstantiate = true;
+};
+
diff --git a/Engine/Plugins/Experimental/PhysicsControl/Source/ControlRigPhysics/Private/RigPhysicsSolverSimulationSpace.cpp b/Engine/Plugins/Experimental/PhysicsControl/Source/ControlRigPhysics/Private/RigPhysicsSolverSimulationSpace.cpp
new file mode 100644
--- /dev/null
+++ b/Engine/Plugins/Experimental/PhysicsControl/Source/ControlRigPhysics/Private/RigPhysicsSolverSimulationSpace.cpp
@@ -0,0 +1,247 @@
+// Copyright Epic Games, Inc. All Rights Reserved.
+
+#include "RigPhysicsSolver.h"
+
+#include "Rigs/RigHierarchy.h"
+#include "Rigs/RigHierarchyController.h"
+#include "ControlRig.h"
+
+#include "Physics/ImmediatePhysics/ImmediatePhysicsActorHandle.h"
+#include "Physics/ImmediatePhysics/ImmediatePhysicsJointHandle.h"
+#include "Physics/ImmediatePhysics/ImmediatePhysicsSimulation.h"
+
+#include "Physics/Experimental/PhysScene_Chaos.h"
+#include "Chaos/ParticleHandle.h"
+#include "Chaos/PBDJointConstraints.h"
+#include "Chaos/ImplicitObject.h"
+#include "Chaos/ShapeInstance.h"
+#include "Chaos/Capsule.h"
+#include "Chaos/PBDJointConstraintUtilities.h"
+
+#include "PhysicsEngine/BodyInstance.h"
+#include "PhysicsEngine/ConstraintInstance.h"
+#include "PhysicsEngine/BodySetup.h"
+
+//======================================================================================================================
+FORCEINLINE_DEBUGGABLE FTransform GetSpaceTransform(
+	ERigPhysicsSimulationSpace Space, const FTransform& ComponentTM, const FTransform& BoneTM)
+{
+	switch (Space)
+	{
+	case ERigPhysicsSimulationSpace::ComponentSpace: return ComponentTM;
+	case ERigPhysicsSimulationSpace::WorldSpace: return FTransform::Identity;
+	case ERigPhysicsSimulationSpace::BoneSpace: return BoneTM * ComponentTM;
+	default: ensureMsgf(false, TEXT("Unsupported Simulation Space")); return FTransform::Identity;
+	}
+}
+
+//======================================================================================================================
+FTransform FRigPhysicsSolver::GetSimulationSpaceTransform() const
+{
+	switch (SolverSettings.SimulationSpace)
+	{
+	case ERigPhysicsSimulationSpace::ComponentSpace: return SimulationSpaceState.ComponentTM;
+	case ERigPhysicsSimulationSpace::WorldSpace: return FTransform::Identity;
+	case ERigPhysicsSimulationSpace::BoneSpace: return SimulationSpaceState.BoneRelComponentTM * SimulationSpaceState.ComponentTM;
+	default: ensureMsgf(false, TEXT("Unsupported Simulation Space")); return FTransform::Identity;
+	}
+}
+
+//======================================================================================================================
+FTransform FRigPhysicsSolver::ConvertComponentSpaceTransformToSimSpace(const FTransform& TM) const
+{
+	switch (SolverSettings.SimulationSpace)
+	{
+	case ERigPhysicsSimulationSpace::ComponentSpace: return TM;
+	case ERigPhysicsSimulationSpace::WorldSpace:  return TM * SimulationSpaceState.ComponentTM;
+	case ERigPhysicsSimulationSpace::BoneSpace: return TM.GetRelativeTransform(SimulationSpaceState.BoneRelComponentTM);
+	default: ensureMsgf(false, TEXT("Unsupported Simulation Space")); return TM;
+	}
+}
+
+//======================================================================================================================
+FTransform FRigPhysicsSolver::ConvertSimSpaceTransformToComponentSpace(const FTransform& TM) const
+{
+	switch (SolverSettings.SimulationSpace)
+	{
+	case ERigPhysicsSimulationSpace::ComponentSpace: return TM;
+	case ERigPhysicsSimulationSpace::WorldSpace:  return TM.GetRelativeTransform(SimulationSpaceState.ComponentTM);
+	case ERigPhysicsSimulationSpace::BoneSpace: return TM * SimulationSpaceState.BoneRelComponentTM;
+	default: ensureMsgf(false, TEXT("Unsupported Simulation Space")); return TM;
+	}
+}
+
+//======================================================================================================================
+FORCEINLINE_DEBUGGABLE FVector ConvertWorldVectorToSimSpaceNoScale(
+	ERigPhysicsSimulationSpace Space, const FVector& WorldVector, 
+	const FTransform& ComponentTM, const FTransform& BoneTM)
+{
+	switch (Space)
+	{
+	case ERigPhysicsSimulationSpace::ComponentSpace: return ComponentTM.InverseTransformVectorNoScale(WorldVector);
+	case ERigPhysicsSimulationSpace::WorldSpace: return WorldVector;
+	case ERigPhysicsSimulationSpace::BoneSpace:
+		return BoneTM.InverseTransformVectorNoScale(ComponentTM.InverseTransformVectorNoScale(WorldVector));
+	default: ensureMsgf(false, TEXT("Unsupported Simulation Space")); return WorldVector;
+	}
+}
+
+//======================================================================================================================
+FVector FRigPhysicsSolver::ConvertWorldVectorToSimSpaceNoScale(const FVector& WorldVector) const
+{
+	switch (SolverSettings.SimulationSpace)
+	{
+	case ERigPhysicsSimulationSpace::ComponentSpace: return SimulationSpaceState.ComponentTM.InverseTransformVectorNoScale(WorldVector);
+	case ERigPhysicsSimulationSpace::WorldSpace: return WorldVector;
+	case ERigPhysicsSimulationSpace::BoneSpace:
+		return SimulationSpaceState.BoneRelComponentTM.InverseTransformVectorNoScale(
+			SimulationSpaceState.ComponentTM.InverseTransformVectorNoScale(WorldVector));
+	default: ensureMsgf(false, TEXT("Unsupported Simulation Space")); return WorldVector;
+	}
+}
+
+//======================================================================================================================
+FTransform FRigPhysicsSolver::ConvertCollisionSpaceTransformToSimSpace(const FTransform& TM) const
+{
+	FTransform SimSpaceTM = GetSpaceTransform(
+		SolverSettings.SimulationSpace, SimulationSpaceState.ComponentTM, SimulationSpaceState.BoneRelComponentTM);
+	FTransform CollisionSpaceTM = GetSpaceTransform(
+		SolverSettings.CollisionSpace, SimulationSpaceState.ComponentTM, SimulationSpaceState.BoneRelComponentTM);
+
+	FTransform WorldSpaceTM = TM * CollisionSpaceTM;
+	return WorldSpaceTM.GetRelativeTransform(SimSpaceTM);
+}
+
+
+//======================================================================================================================
+void FRigPhysicsSolver::InitSimulationSpace(
+	const FTransform& ComponentTM,
+	const FTransform& BoneRelComponentTM)
+{
+	SimulationSpaceState.ComponentTM = ComponentTM;
+	SimulationSpaceState.BoneRelComponentTM = BoneRelComponentTM;
+	SimulationSpaceState.ComponentLinearVelocity = FVector::ZeroVector;
+	SimulationSpaceState.ComponentAngularVelocity = FVector::ZeroVector;
+	SimulationSpaceState.BoneLinearVelocity = FVector::ZeroVector;
+	SimulationSpaceState.BoneAngularVelocity = FVector::ZeroVector;
+}
+
+//======================================================================================================================
+// Note - don't use the space conversion functions here as the state won't have been set yet.
+FRigPhysicsSolver::FSimulationSpaceData FRigPhysicsSolver::UpdateSimulationSpaceStateAndCalculateData(const float Dt)
+{
+	URigHierarchy* Hierarchy = OwningControlRig->GetHierarchy();
+
+	FTransform ComponentTM;
+	if (const USceneComponent* OwningSceneComponent = OwningControlRig->GetOwningSceneComponent())
+	{
+		ComponentTM = OwningSceneComponent->GetComponentTransform();
+	}
+
+	FTransform BoneRelComponentTM;
+	if (SolverSettings.SpaceBone.IsValid())
+	{
+		BoneRelComponentTM = Hierarchy->GetGlobalTransform(SolverSettings.SpaceBone);
+	}
+
+	FSimulationSpaceData SimulationSpaceData(GetSpaceTransform(
+		SolverSettings.SimulationSpace, ComponentTM, BoneRelComponentTM));
+
+	// Note that this depends on the simulation space having been calculated
+	SimulationSpaceData.Gravity = ::ConvertWorldVectorToSimSpaceNoScale(
+		SolverSettings.SimulationSpace, SolverSettings.Gravity, ComponentTM, BoneRelComponentTM);
+
+	if (SolverSettings.SimulationSpace == ERigPhysicsSimulationSpace::WorldSpace)
+	{
+		InitSimulationSpace(ComponentTM, BoneRelComponentTM);
+		SimulationSpaceData.LinearVelocity = SimulationSpaceSettings.ExternalLinearVelocity;
+		SimulationSpaceData.AngularVelocity = SimulationSpaceSettings.ExternalAngularVelocity;
+		return SimulationSpaceData;
+	}
+
+	// Don't estimate velocities if the timestep is effectively zero - just preserve them
+	if (Dt < SMALL_NUMBER)
+	{
+		SimulationSpaceState.ComponentTM = ComponentTM;
+		SimulationSpaceState.BoneRelComponentTM = BoneRelComponentTM;
+		return SimulationSpaceData;
+	}
+
+	FVector CompLinVel = (ComponentTM.GetTranslation() - SimulationSpaceState.ComponentTM.GetTranslation()) / Dt;
+	FVector CompAngVel = (ComponentTM.GetRotation() * SimulationSpaceState.ComponentTM.GetRotation().Inverse()).
+		GetShortestArcWith(FQuat::Identity).ToRotationVector() / Dt;
+	FVector CompLinAcc = (CompLinVel - SimulationSpaceState.ComponentLinearVelocity) / Dt;
+	FVector CompAngAcc = (CompAngVel - SimulationSpaceState.ComponentAngularVelocity) / Dt;
+	SimulationSpaceState.ComponentTM = ComponentTM;
+	SimulationSpaceState.ComponentLinearVelocity = CompLinVel;
+	SimulationSpaceState.ComponentAngularVelocity = CompAngVel;
+
+	if (SolverSettings.SimulationSpace == ERigPhysicsSimulationSpace::ComponentSpace)
+	{
+		// Danny TODO this per-component scaling needs to be done in the right space
+		CompLinVel *= SimulationSpaceSettings.MovementScale;
+		CompLinAcc *= SimulationSpaceSettings.MovementScale;
+
+		SimulationSpaceData.LinearVelocity = CompLinVel.GetClampedToMaxSize(SimulationSpaceSettings.MaxLinearVelocity) + 
+			SimulationSpaceSettings.ExternalLinearVelocity;
+		SimulationSpaceData.AngularVelocity = CompAngVel.GetClampedToMaxSize(SimulationSpaceSettings.MaxAngularVelocity) + 
+			SimulationSpaceSettings.ExternalAngularVelocity;
+		SimulationSpaceData.LinearAcceleration = CompLinAcc.GetClampedToMaxSize(SimulationSpaceSettings.MaxLinearAcceleration);
+		SimulationSpaceData.AngularAcceleration = CompAngAcc.GetClampedToMaxSize(SimulationSpaceSettings.MaxAngularAcceleration);
+		return SimulationSpaceData;
+	}
+
+	if (SolverSettings.SimulationSpace == ERigPhysicsSimulationSpace::BoneSpace)
+	{
+		// World-space component-relative bone velocity and acceleration
+		FVector BoneLinVel = 
+			(BoneRelComponentTM.GetTranslation() - SimulationSpaceState.BoneRelComponentTM.GetTranslation()) / Dt;
+		FVector BoneAngVel = 
+			(BoneRelComponentTM.GetRotation() * SimulationSpaceState.BoneRelComponentTM.GetRotation().Inverse()).
+			GetShortestArcWith(FQuat::Identity).ToRotationVector() / Dt;
+		BoneLinVel = ComponentTM.TransformVector(BoneLinVel);
+		BoneAngVel = ComponentTM.TransformVector(BoneAngVel);
+		FVector BoneLinAcc = (BoneLinVel - SimulationSpaceState.BoneLinearVelocity) / Dt;
+		FVector BoneAngAcc = (BoneAngVel - SimulationSpaceState.BoneAngularVelocity) / Dt;
+		SimulationSpaceState.BoneRelComponentTM = BoneRelComponentTM;
+		SimulationSpaceState.BoneLinearVelocity = BoneLinVel;
+		SimulationSpaceState.BoneAngularVelocity = BoneAngVel;
+
+		// World-space bone velocity and acceleration
+		FVector NetAngVel = CompAngVel + BoneAngVel;
+		FVector NetAngAcc = CompAngAcc + BoneAngAcc;
+
+		// If we limit the angular velocity, we also need to limit the component of linear velocity
+		// that comes from (angvel x offset)
+		float AngVelScale = 1.0f;
+		float NetAngVelLenSq = static_cast<float>(NetAngVel.SizeSquared());
+		if (NetAngVelLenSq > FMath::Square(SimulationSpaceSettings.MaxAngularVelocity))
+		{
+			AngVelScale = SimulationSpaceSettings.MaxAngularVelocity * FMath::InvSqrt(NetAngVelLenSq);
+		}
+
+		// Add the linear velocity and acceleration that comes from rotation of the space about the component
+		// NOTE: Component angular velocity contribution is scaled
+		FVector SpaceCompOffset = ComponentTM.TransformVector(BoneRelComponentTM.GetTranslation());
+		FVector NetLinVel = CompLinVel + BoneLinVel + FVector::CrossProduct(AngVelScale * CompAngVel, SpaceCompOffset);
+		FVector NetLinAcc = CompLinAcc + BoneLinAcc + FVector::CrossProduct(AngVelScale * CompAngAcc, SpaceCompOffset);
+
+		NetLinVel *= SimulationSpaceSettings.MovementScale;
+		NetLinAcc *= SimulationSpaceSettings.MovementScale;
+
+		SimulationSpaceData.LinearVelocity = 
+			NetLinVel.GetClampedToMaxSize(SimulationSpaceSettings.MaxLinearVelocity) + 
+			SimulationSpaceSettings.ExternalLinearVelocity;
+		SimulationSpaceData.AngularVelocity = 
+			NetAngVel.GetClampedToMaxSize(SimulationSpaceSettings.MaxAngularVelocity) + 
+			SimulationSpaceSettings.ExternalAngularVelocity;
+		SimulationSpaceData.LinearAcceleration = 
+			NetLinAcc.GetClampedToMaxSize(SimulationSpaceSettings.MaxLinearAcceleration);
+		SimulationSpaceData.AngularAcceleration = 
+			NetAngAcc.GetClampedToMaxSize(SimulationSpaceSettings.MaxAngularAcceleration);
+		return SimulationSpaceData;
+	}
+
+	// Shouldn't get here
+	return SimulationSpaceData;
+}
diff --git a/Engine/Plugins/Experimental/PhysicsControl/Source/ControlRigPhysics/Private/RigPhysicsSolverVisualization.cpp b/Engine/Plugins/Experimental/PhysicsControl/Source/ControlRigPhysics/Private/RigPhysicsSolverVisualization.cpp
new file mode 100644
--- /dev/null
+++ b/Engine/Plugins/Experimental/PhysicsControl/Source/ControlRigPhysics/Private/RigPhysicsSolverVisualization.cpp
@@ -0,0 +1,334 @@
+// Copyright Epic Games, Inc. All Rights Reserved.
+
+#include "RigPhysicsSolver.h"
+
+#include "Rigs/RigHierarchy.h"
+#include "Rigs/RigHierarchyController.h"
+#include "ControlRig.h"
+
+#include "RigVMCore/RigVMDrawInterface.h"
+
+#include "Physics/ImmediatePhysics/ImmediatePhysicsActorHandle.h"
+#include "Physics/ImmediatePhysics/ImmediatePhysicsJointHandle.h"
+#include "Physics/ImmediatePhysics/ImmediatePhysicsSimulation.h"
+
+#include "Physics/Experimental/PhysScene_Chaos.h"
+#include "Chaos/ParticleHandle.h"
+#include "Chaos/PBDJointConstraints.h"
+#include "Chaos/ImplicitObject.h"
+#include "Chaos/ShapeInstance.h"
+#include "Chaos/Capsule.h"
+#include "Chaos/PBDJointConstraintUtilities.h"
+
+#include "Engine/Engine.h"
+#include "Materials/MaterialInstanceDynamic.h"
+
+#include "Math/Color.h"
+
+namespace RigPhysicsSolverDraw
+{
+	float Thickness = 0.3f;
+	FLinearColor DynamicColor = FColor::Yellow;
+	FLinearColor KinematicColor = FColor::Blue;
+	int32 Detail = 16;
+}
+
+//======================================================================================================================
+static void DrawShapes(
+	FRigVMDrawInterface*                  DI,
+	const bool                            bIsKinematic,
+	const Chaos::FGeometryParticleHandle* Particle,
+	const FTransform&                     ShapeTransform, 
+	const Chaos::FImplicitObject*         ImplicitObject,
+	const Chaos::FPerShapeData*           Shape)
+{
+	using namespace Chaos;
+
+	if (!ImplicitObject)
+	{
+		return;
+	}
+	const EImplicitObjectType PackedType = ImplicitObject->GetType(); // Type includes scaling and instancing data
+	const EImplicitObjectType InnerType = GetInnerType(ImplicitObject->GetType());
+
+	// For simplicity, we're going to assume no scaling, instancing, transforming etc
+
+	if (!ImplicitObject)
+	{
+		return;
+	}
+
+	switch (InnerType)
+	{
+	case ImplicitObjectType::Transformed:
+	{
+		const TImplicitObjectTransformed<FReal, 3>* Transformed = 
+			ImplicitObject->GetObject<TImplicitObjectTransformed<FReal, 3>>();
+		FTransform TransformedTransform = FTransform(
+			ShapeTransform.GetRotation() * Transformed->GetTransform().GetRotation(),
+			ShapeTransform.TransformPosition(Transformed->GetTransform().GetLocation()));
+		DrawShapes(DI, bIsKinematic, Particle, TransformedTransform, Transformed->GetTransformedObject(), Shape);
+		return;
+	}
+	case ImplicitObjectType::Sphere:
+	{
+		const Chaos::FSphere* Sphere = ImplicitObject->GetObject<Chaos::FSphere>();
+		DI->DrawSphere(ShapeTransform, FTransform(FVec3(Sphere->GetCenterOfMass())), Sphere->GetRadiusf(),
+			bIsKinematic ? RigPhysicsSolverDraw::KinematicColor : RigPhysicsSolverDraw::DynamicColor,
+			RigPhysicsSolverDraw::Thickness, RigPhysicsSolverDraw::Detail);
+		break;
+	}
+	case ImplicitObjectType::Box:
+	{
+		const TBox<FReal, 3>* Box = ImplicitObject->GetObject<TBox<FReal, 3>>();
+		DI->DrawBox(ShapeTransform, FTransform(FQuat::Identity, Box->GetCenter(), Box->Extents()),
+			bIsKinematic ? RigPhysicsSolverDraw::KinematicColor : RigPhysicsSolverDraw::DynamicColor,
+			RigPhysicsSolverDraw::Thickness);
+		break;
+	}
+	case ImplicitObjectType::Capsule:
+	{
+		const Chaos::FCapsule* Capsule = ImplicitObject->GetObject<Chaos::FCapsule>();
+		const FQuat Q = FRotationMatrix::MakeFromZ(Capsule->GetAxis()).ToQuat();
+		DI->DrawCapsule(ShapeTransform, FTransform(Q, Capsule->GetCenterOfMass()),
+			Capsule->GetRadiusf(), Capsule->GetHeightf(),
+			bIsKinematic ? RigPhysicsSolverDraw::KinematicColor : RigPhysicsSolverDraw::DynamicColor,
+			RigPhysicsSolverDraw::Thickness);
+		break;
+	}
+	default:
+		// If we don't know what it is, don't draw it.
+		break;
+	}
+
+}
+
+//======================================================================================================================
+static void DrawActor(
+	FRigVMDrawInterface*                  DI,
+	const FTransform&                     SpaceTransform,
+	const ImmediatePhysics::FActorHandle& ActorHandle)
+{
+	bool bIsKinematic = ActorHandle.GetIsKinematic();
+	const Chaos::FGeometryParticleHandle* GeometryParticleHandle = ActorHandle.GetParticle();
+	const FTransform ParticleTransform = GeometryParticleHandle->GetTransformXR() * SpaceTransform;
+	for (const Chaos::FShapeInstancePtr& ShapeInstance : GeometryParticleHandle->ShapeInstances())
+	{
+		const Chaos::FImplicitObject* ImplicitObject = ShapeInstance->GetGeometry();
+		DrawShapes(
+			DI, bIsKinematic, GeometryParticleHandle, 
+			ParticleTransform, ImplicitObject, ShapeInstance.Get());
+	}
+}
+
+enum
+{
+	TwistIndex = 0,
+	Swing1Index = 1,
+	Swing2Index = 2
+};
+
+// Returns a number to use for the limit for limit based on limit type
+float GetLimitAngleRadians(float LimitAngle, Chaos::EJointMotionType LimitType)
+{
+	switch (LimitType)
+	{
+	case Chaos::EJointMotionType::Free: return UE_PI;
+	case Chaos::EJointMotionType::Locked: return 0.f;
+	default: return LimitAngle;
+	}
+}
+
+
+//======================================================================================================================
+static void DrawJoint(
+	FRigVMDrawInterface*                  DI,
+	const FTransform&                     SpaceTransform,
+	const ImmediatePhysics::FJointHandle& JointHandle)
+{
+	const Chaos::FPBDJointConstraintHandle* ConstraintHandle = JointHandle.GetConstraint();
+	if (!ConstraintHandle)
+	{
+		return;
+	}
+
+	const Chaos::TVec2<const ImmediatePhysics_Chaos::FActorHandle*>& ActorHandles = JointHandle.GetActorHandles();
+	const ImmediatePhysics_Chaos::FActorHandle* ChildActor = ActorHandles[0];
+	const ImmediatePhysics_Chaos::FActorHandle* ParentActor = ActorHandles[1];
+
+	FTransform ChildActorTM = (ChildActor ? ChildActor->GetWorldTransform() : FTransform()) * SpaceTransform;
+	FTransform ParentActorTM = (ParentActor ? ParentActor->GetWorldTransform() : FTransform()) * SpaceTransform;
+
+	bool bChildIsKinematic = ChildActor ? ChildActor->GetIsKinematic() : true;
+	bool bParentIsKinematic = ParentActor ? ParentActor->GetIsKinematic() : true;
+
+	float Size = 5.0f;
+
+	const Chaos::FPBDJointSettings& JointSettings = ConstraintHandle->GetSettings();
+
+	FTransform DialFrame = JointSettings.ConnectorTransforms[0] * ChildActorTM;
+	FTransform LimitFrame = JointSettings.ConnectorTransforms[1] * ParentActorTM;
+
+	static bool bDrawAsAxes = false;
+	if (bDrawAsAxes)
+	{
+		DI->DrawAxes(LimitFrame, FTransform(), Size, RigPhysicsSolverDraw::Thickness * 2);
+		DI->DrawAxes(DialFrame, FTransform(), Size, RigPhysicsSolverDraw::Thickness);
+	}
+
+	{
+		// See FConstraintInstance::DrawConstraintImp for inspiration
+
+		// There seems to be a swap between swing1 and swing2 compared to GetSwingTwistAngles
+		FVector LimitAngleRadians(
+			GetLimitAngleRadians(JointSettings.AngularLimits[0], JointSettings.AngularMotionTypes[0]),
+			GetLimitAngleRadians(JointSettings.AngularLimits[2], JointSettings.AngularMotionTypes[2]),
+			GetLimitAngleRadians(JointSettings.AngularLimits[1], JointSettings.AngularMotionTypes[1])
+		);
+
+		Chaos::FReal TwistAngle, Swing1Angle, Swing2Angle;
+		const FQuat LimitQ = LimitFrame.GetRotation();
+		FQuat DialQ = DialFrame.GetRotation();
+		DialQ.EnforceShortestArcWith(LimitQ);
+		Chaos::FPBDJointUtilities::GetSwingTwistAngles(LimitQ, DialQ, TwistAngle, Swing1Angle, Swing2Angle);
+
+		bool bDrawViolatedLimits = true;
+
+		bool bTwistViolated = false;
+		bool bSwing1Violated = false;
+		bool bSwing2Violated = false;
+
+		if (bDrawViolatedLimits)
+		{
+			bTwistViolated = JointSettings.AngularMotionTypes[TwistIndex] == Chaos::EJointMotionType::Limited &&
+				FMath::Abs(TwistAngle) > LimitAngleRadians[0];
+
+			bSwing1Violated = JointSettings.AngularMotionTypes[Swing1Index] == Chaos::EJointMotionType::Limited &&
+				FMath::Abs(Swing1Angle) > LimitAngleRadians[1];
+
+			bSwing2Violated = JointSettings.AngularMotionTypes[Swing2Index] == Chaos::EJointMotionType::Limited &&
+				FMath::Abs(Swing2Angle) > LimitAngleRadians[2];
+		}
+
+		const bool bLockSwing1 = JointSettings.AngularMotionTypes[Swing1Index] == Chaos::EJointMotionType::Locked;
+		const bool bLockSwing2 = JointSettings.AngularMotionTypes[Swing2Index] == Chaos::EJointMotionType::Locked;
+		const bool bLockAllSwing = bLockSwing1 && bLockSwing2;
+
+		check(GEngine->ConstraintLimitMaterialX && GEngine->ConstraintLimitMaterialY && GEngine->ConstraintLimitMaterialZ);
+		static UMaterialInterface* LimitMaterialX = GEngine->ConstraintLimitMaterialX;
+		static UMaterialInterface* LimitMaterialXAxis = GEngine->ConstraintLimitMaterialXAxis;
+		static UMaterialInterface* LimitMaterialY = GEngine->ConstraintLimitMaterialY;
+		static UMaterialInterface* LimitMaterialYAxis = GEngine->ConstraintLimitMaterialYAxis;
+		static UMaterialInterface* LimitMaterialZ = GEngine->ConstraintLimitMaterialZ;
+		static UMaterialInterface* LimitMaterialZAxis = GEngine->ConstraintLimitMaterialZAxis;
+
+		float ArrowLength = Size * 1.05f; // stick out a little bit 
+
+		// If swing is limited (but not locked) - draw the swing limit cone.
+		if (!bLockAllSwing)
+		{
+			if (JointSettings.AngularMotionTypes[Swing1Index] == Chaos::EJointMotionType::Free && 
+				JointSettings.AngularMotionTypes[Swing2Index] == Chaos::EJointMotionType::Free)
+			{
+				DI->DrawSphere(LimitFrame, FTransform(), Size * 0.2f, FColor::White, 
+					RigPhysicsSolverDraw::Thickness, RigPhysicsSolverDraw::Detail);
+			}
+			else
+			{
+				FTransform ConeLimitTM = LimitFrame;
+				ConeLimitTM.SetScale3D(FVector(Size));
+				const float Swing1LimitAngle = LimitAngleRadians[Swing1Index];
+				const float Swing2LimitAngle = LimitAngleRadians[Swing2Index];
+				DI->DrawCone(
+					ConeLimitTM, FTransform(), Swing1LimitAngle, Swing2LimitAngle, RigPhysicsSolverDraw::Detail, 
+					true, FColor::Green, LimitMaterialX->GetRenderProxy(), RigPhysicsSolverDraw::Thickness);
+			}
+			// Draw the swing Dial indicator - shows the current orientation of the child frame
+			// relative to the parent frame on the swing axis.
+			// Start the arrow at the limit position as it can be confusing if there is joint separation
+			DI->DrawArrow(FTransform(LimitFrame.GetTranslation()), 
+				DialFrame.GetUnitAxis(EAxis::X) * ArrowLength,
+				DialFrame.GetUnitAxis(EAxis::Y), FColor::Red, RigPhysicsSolverDraw::Thickness);
+
+			if (bSwing1Violated || bSwing2Violated)
+			{
+				// Drawing a sphere seems a bit heavy... but it appears that if DrawPoint is used then they get
+				// culled when you get close!
+				DI->DrawSphere(FTransform(
+					DialFrame.GetRotation(),
+					LimitFrame.GetTranslation() + DialFrame.GetUnitAxis(EAxis::X) * ArrowLength),
+					FTransform(),
+					ArrowLength * 0.01f, FColor::Orange, RigPhysicsSolverDraw::Thickness * 4);
+			}
+
+		}
+
+		// Draw the twist limit
+		if (JointSettings.AngularMotionTypes[TwistIndex] != Chaos::EJointMotionType::Locked)
+		{
+			// Draw as a flat cone
+			FTransform ConeLimitTM = LimitFrame;
+			ConeLimitTM.SetScale3D(FVector(Size));
+			const float TwistLimitAngle = LimitAngleRadians[TwistIndex];
+			DI->DrawCone(ConeLimitTM, FTransform(FQuat::MakeFromRotationVector(FVector(0, 1, 0) * -HALF_PI)),
+				TwistLimitAngle, 0.0f, RigPhysicsSolverDraw::Detail, true, FColor::Green,
+				LimitMaterialYAxis->GetRenderProxy(), RigPhysicsSolverDraw::Thickness);
+
+			// Draw the twist Dial indicator - shows the current orientation of the child frame
+			// relative to the parent frame on the twist axis.
+			FQuat Rot(LimitFrame.GetUnitAxis(EAxis::X), TwistAngle);
+			FVector TwistArrow = Rot * LimitFrame.GetUnitAxis(EAxis::Z);
+			FVector TwistSideDir = FVector::CrossProduct(LimitFrame.GetUnitAxis(EAxis::X), TwistArrow);
+
+			DI->DrawArrow(FTransform(LimitFrame.GetTranslation()),
+				TwistArrow * ArrowLength,
+				TwistSideDir, FColor::Blue, RigPhysicsSolverDraw::Thickness);
+
+			if (bTwistViolated)
+			{
+				DI->DrawSphere(FTransform(
+					DialFrame.GetRotation(),
+					LimitFrame.GetTranslation() + TwistArrow * ArrowLength),
+					FTransform(),
+					ArrowLength * 0.01f, FColor::Orange, RigPhysicsSolverDraw::Thickness * 4);
+			}
+		}
+	}
+}
+
+//======================================================================================================================
+void FRigPhysicsSolver::Draw(FRigVMDrawInterface* DI) const
+{
+	const URigHierarchy* Hierarchy = OwningControlRig->GetHierarchy();
+	if (!DI || !Hierarchy)
+	{
+		return;
+	}
+
+	// All rendering is done relative to the component, so convert the sim space (identity) into the
+	// component space.
+	const FTransform SpaceTransform = ConvertSimSpaceTransformToComponentSpace(FTransform());
+
+	for (const ImmediatePhysics::FActorHandle* ActorHandle : AdditionalActorHandles)
+	{
+		DrawActor(DI, SpaceTransform, *ActorHandle);
+	}
+
+	for (const TPair<FRigComponentKey, FRigBodyRecord>& BodyRecordPair : BodyRecords)
+	{
+		const FRigBodyRecord& Record = BodyRecordPair.Value;
+		const ImmediatePhysics::FActorHandle* ActorHandle = Record.ActorHandle;
+		const ImmediatePhysics::FJointHandle* JointHandle = Record.ArticulationJointHandle;
+
+		if (JointHandle)
+		{
+			DrawJoint(DI, SpaceTransform, *JointHandle);
+		}
+
+		if (ActorHandle)
+		{
+			DrawActor(DI, SpaceTransform, *ActorHandle);
+		}
+	}
+}
+
diff --git a/Engine/Plugins/Experimental/PhysicsControl/Source/ControlRigPhysics/Private/RigUnit_Physics.cpp b/Engine/Plugins/Experimental/PhysicsControl/Source/ControlRigPhysics/Private/RigUnit_Physics.cpp
new file mode 100644
--- /dev/null
+++ b/Engine/Plugins/Experimental/PhysicsControl/Source/ControlRigPhysics/Private/RigUnit_Physics.cpp
@@ -0,0 +1,182 @@
+// Copyright Epic Games, Inc. All Rights Reserved.
+
+#include "RigUnit_Physics.h"
+#include "Rigs/RigHierarchyController.h"
+#include "Units/RigUnitContext.h"
+#include "Units/Execution/RigUnit_PrepareForExecution.h"
+#include "RigPhysicsSolver.h"
+#include "ControlRig.h"
+
+#include UE_INLINE_GENERATED_CPP_BY_NAME(RigUnit_Physics)
+
+//======================================================================================================================
+FRigUnit_HierarchyAddPhysicsSolver_Execute()
+{
+	FString ErrorMessage;
+	if(!FRigUnit_DynamicHierarchyBase::IsValidToRunInContext(ExecuteContext, true, &ErrorMessage))
+	{
+		if(!ErrorMessage.IsEmpty())
+		{
+			UE_CONTROLRIG_RIGUNIT_REPORT_ERROR(TEXT("%s"), *ErrorMessage);
+		}
+		return;
+	}
+
+	SolverID = FRigPhysicsSolverID();
+
+	if (!ExecuteContext.Hierarchy)
+	{
+		return;
+	}
+
+	if(UControlRig* ControlRig = Cast<UControlRig>(ExecuteContext.Hierarchy->GetOuter()))
+	{
+		TSharedPtr<FRigPhysicsSolver> PhysicsSolver = MakeShared<FRigPhysicsSolver>(
+			ControlRig, SolverSettings, SimulationSpaceSettings);
+		SolverID = ControlRig->RegisterPhysicsSolver(PhysicsSolver, Name);
+	}
+}
+
+//======================================================================================================================
+FRigUnit_AddPhysicsBody_Execute()
+{
+	Result.Reset();
+
+	if (!SolverID.IsValid())
+	{
+		UE_CONTROLRIG_RIGUNIT_REPORT_ERROR(TEXT("Physics Solver Guid is not valid"));
+		return;
+	}
+
+	if (ExecuteContext.GetEventName() == FRigUnit_PrepareForExecution::EventName)
+	{
+		if (URigHierarchy* Hierarchy = ExecuteContext.Hierarchy)
+		{
+			if (UControlRig* ControlRig = Cast<UControlRig>(Hierarchy->GetOuter()))
+			{
+				if (ControlRig->GetPhysicsSolver(SolverID) == nullptr)
+				{
+					UE_CONTROLRIG_RIGUNIT_REPORT_ERROR(
+						TEXT("Physics Solver with guid '%s' cannot be found."), *SolverID.ToString());
+					return;
+				}
+
+				if (URigHierarchyController* Controller = Hierarchy->GetController())
+				{
+					FRigHierarchyControllerInstructionBracket InstructionBracket(Controller, ExecuteContext.GetInstructionIndex());
+					Result = Controller->AddComponent(FRigPhysicsComponent::StaticStruct(), TEXT("PhysicsBody"), Item);
+					if (Result.IsValid())
+					{
+						if (FRigPhysicsComponent* Component = Cast<FRigPhysicsComponent>(Hierarchy->FindComponent(Result)))
+						{
+							Component->BodySolverSettings.SolverID = SolverID;
+							Component->Dynamics = Dynamics;
+							Component->Collision = Collision;
+							Component->Articulation = Articulation;
+							Component->BodyData = BodyData;
+
+							if (FRigPhysicsSolver* PhysicsSolver = Cast<FRigPhysicsSolver>(ControlRig->GetPhysicsSolver(SolverID)))
+							{
+								PhysicsSolver->AddBody(Result);
+							}
+						}
+					}
+				}
+			}
+		}
+	}
+	else
+	{
+		UE_CONTROLRIG_RIGUNIT_REPORT_ERROR(TEXT("AddPhysicsBody can only be used during Setup"));
+	}
+}
+
+//======================================================================================================================
+FRigUnit_HierarchySetPhysicsBodySourceBone_Execute()
+{
+	DECLARE_SCOPE_HIERARCHICAL_COUNTER_RIGUNIT();
+
+	if (ExecuteContext.Hierarchy)
+	{
+		if (FRigPhysicsComponent* PhysicsComponent = Cast<FRigPhysicsComponent>(
+			ExecuteContext.Hierarchy->FindComponent(PhysicsBody)))
+		{
+			PhysicsComponent->BodySolverSettings.SourceBone = SourceBone;
+		}
+	}
+}
+
+//======================================================================================================================
+FRigUnit_HierarchySetPhysicsBodyTargetBone_Execute()
+{
+	DECLARE_SCOPE_HIERARCHICAL_COUNTER_RIGUNIT();
+
+	if (ExecuteContext.Hierarchy)
+	{
+		if (FRigPhysicsComponent* PhysicsComponent = Cast<FRigPhysicsComponent>(
+			ExecuteContext.Hierarchy->FindComponent(PhysicsBody)))
+		{
+			PhysicsComponent->BodySolverSettings.TargetBone = TargetBone;
+		}
+	}
+}
+
+//======================================================================================================================
+FRigUnit_HierarchySetPhysicsBodySparseData_Execute()
+{
+	DECLARE_SCOPE_HIERARCHICAL_COUNTER_RIGUNIT();
+
+	if (ExecuteContext.Hierarchy)
+	{
+		if (FRigPhysicsComponent* PhysicsComponent = Cast<FRigPhysicsComponent>(
+			ExecuteContext.Hierarchy->FindComponent(PhysicsBody)))
+		{
+			if (UControlRig* ControlRig = Cast<UControlRig>(ExecuteContext.Hierarchy->GetOuter()))
+			{
+				if (FRigPhysicsSolver* PhysicsSolver = Cast<FRigPhysicsSolver>(
+					ControlRig->GetPhysicsSolver((PhysicsComponent->BodySolverSettings.SolverID))))
+				{
+					PhysicsSolver->SetBodySparseData(PhysicsBody, Data);
+				}
+			}
+		}
+	}
+}
+
+//======================================================================================================================
+FRigUnit_StepPhysicsSolver_Execute()
+{
+	DECLARE_SCOPE_HIERARCHICAL_COUNTER_RIGUNIT();
+
+	if (ExecuteContext.Hierarchy)
+	{
+		if (UControlRig* ControlRig = Cast<UControlRig>(ExecuteContext.Hierarchy->GetOuter()))
+		{
+			if (FRigPhysicsSolver* PhysicsSolver = Cast<FRigPhysicsSolver>(ControlRig->GetPhysicsSolver(SolverID)))
+			{
+				double DeltaTime = DeltaTimeOverride > 0 ? DeltaTimeOverride : ExecuteContext.GetDeltaTime();
+				PhysicsSolver->StepSimulation(DeltaTime, ExecuteContext.GetWorld());
+				PhysicsSolver->Draw(ExecuteContext.GetDrawInterface());
+			}
+		}
+	}
+}
+
+//======================================================================================================================
+FRigUnit_InstantiatePhysics_Execute()
+{
+	DECLARE_SCOPE_HIERARCHICAL_COUNTER_RIGUNIT();
+
+	if (ExecuteContext.Hierarchy)
+	{
+		if (UControlRig* ControlRig = Cast<UControlRig>(ExecuteContext.Hierarchy->GetOuter()))
+		{
+			if (FRigPhysicsSolver* PhysicsSolver = Cast<FRigPhysicsSolver>(ControlRig->GetPhysicsSolver(SolverID)))
+			{
+				PhysicsSolver->Instantiate();
+			}
+		}
+	}
+}
+
+
diff --git a/Engine/Plugins/Experimental/PhysicsControl/Source/ControlRigPhysics/Public/ControlRigPhysicsModule.h b/Engine/Plugins/Experimental/PhysicsControl/Source/ControlRigPhysics/Public/ControlRigPhysicsModule.h
new file mode 100644
--- /dev/null
+++ b/Engine/Plugins/Experimental/PhysicsControl/Source/ControlRigPhysics/Public/ControlRigPhysicsModule.h
@@ -0,0 +1,14 @@
+// Copyright Epic Games, Inc. All Rights Reserved.
+
+#pragma once
+
+#include "Modules/ModuleManager.h"
+
+class FControlRigPhysicsModule : public IModuleInterface
+{
+public:
+
+	/** IModuleInterface implementation */
+	virtual void StartupModule() override;
+	virtual void ShutdownModule() override;
+};
diff --git a/Engine/Plugins/Experimental/PhysicsControl/Source/ControlRigPhysics/Public/RigPhysicsComponents.h b/Engine/Plugins/Experimental/PhysicsControl/Source/ControlRigPhysics/Public/RigPhysicsComponents.h
new file mode 100644
--- /dev/null
+++ b/Engine/Plugins/Experimental/PhysicsControl/Source/ControlRigPhysics/Public/RigPhysicsComponents.h
@@ -0,0 +1,437 @@
+// Copyright Epic Games, Inc. All Rights Reserved.
+
+#pragma once
+
+#include "Rigs/RigPhysics.h"
+#include "PhysicsControlData.h"
+#include "RigPhysicsComponents.generated.h"
+
+struct FRigVMExecuteContext;
+struct FRigBaseElement;
+struct FRigControlElement;
+class URigHierarchy;
+class FRigElementKeyRedirector;
+
+/** Determines in what space the simulation should run */
+UENUM()
+enum class ERigPhysicsSimulationSpace : uint8
+{
+	/** Simulate in world space. Moving the skeletal mesh will generate velocity changes */
+	WorldSpace,
+	/** Simulate in component space. Moving the entire skeletal mesh will have no affect on velocities */
+	ComponentSpace,
+	/**
+	 * Simulate in the space of the specified bone. Moving the entire skeletal mesh and/or the Bone
+	 * will have no affect on in the simulation.
+	 */
+	BoneSpace,
+};
+
+USTRUCT(BlueprintType)
+struct FRigPhysicsSimulationSpaceSettings
+{
+	GENERATED_BODY()
+
+	// Multiplier on the components of velocity and acceleration that are passed to the simulation. Usually from 0.0 to 1.0 to 
+	// reduce the effects of jumping and crouching on the simulation, but it can be higher than 1.0 if you need to exaggerate this motion for some reason.
+	UPROPERTY(EditAnywhere, BlueprintReadWrite, Category = Settings)
+	FVector MovementScale = FVector(1, 1, 1);
+
+	// A clamp on the effective world-space velocity that is passed to the simulation. Units are cm/s. The default value effectively means "unlimited". It is not usually required to
+	// change this but you would reduce this to limit the effects of drag on the bodies in the simulation (if you have bodies that have LinearDrag set to non-zero in the physics asset). 
+	// Expected values in this case would be somewhat less than the usual velocities of your object which is commonly a few hundred for a character.
+	UPROPERTY(EditAnywhere, BlueprintReadWrite, Category = Settings, meta = (ClampMin = "0.0"))
+	float MaxLinearVelocity = 10000;
+
+	// A clamp on the effective world-space angular velocity that is passed to the simulation. Units are radian/s, so a value of about 6.0 is one rotation per second.
+	// The default value effectively means "unlimited". You would reduce this (and MaxAngularAcceleration) to limit how much bodies "fly out" when the actor spins on the spot. 
+	// This is especially useful if you have characters than can rotate very quickly and you would probably want values around or less than 10 in this case.
+	UPROPERTY(EditAnywhere, BlueprintReadWrite, Category = Settings, meta = (ClampMin = "0.0"))
+	float MaxAngularVelocity = 10000;
+
+	// A clamp on the effective world-space acceleration that is passed to the simulation. Units are cm/s/s. The default value effectively means "unlimited". 
+	// This property is used to stop the bodies of the simulation flying out when suddenly changing linear speed. It is useful when you have characters than can 
+	// changes from stationary to running very quickly such as in an FPS. A common value for a character might be in the few hundreds.
+	UPROPERTY(EditAnywhere, BlueprintReadWrite, Category = Settings, meta = (ClampMin = "0.0"))
+	float MaxLinearAcceleration = 10000;
+
+	// A clamp on the effective world-space angular accleration that is passed to the simulation. Units are radian/s/s. The default value effectively means "unlimited". 
+	// This has a similar effect to MaxAngularVelocity, except that it is related to the flying out of bodies when the rotation speed suddenly changes. Typical limist for
+	// a character might be around 100.
+	UPROPERTY(EditAnywhere, BlueprintReadWrite, Category = Settings, meta = (ClampMin = "0.0"))
+	float MaxAngularAcceleration = 10000;
+
+	// Global multiplier on the effects of simulation space movement on the simulation
+	UPROPERTY(EditAnywhere, BlueprintReadWrite, Category = Settings)
+	float SpaceMovementAmount = 1.0f;
+
+	// Additional linear drag applied to every body in addition to linear drag specified on them in the physics asset. 
+	// When combined with ExternalLinearVelocity, this can be used to add a temporary wind-blown effect without having to tune linear drag on 
+	// all the bodies in the physics asset. The result is that each body has a force equal to -ExternalLinearDragV * ExternalLinearVelocity applied to it, in 
+	// additional to all other forces. The vector is in simulation local space.
+	UPROPERTY(EditAnywhere, BlueprintReadWrite, Category = Settings)
+	FVector ExternalLinearDragV = FVector::ZeroVector;
+
+	// Additional velocity that is added to the component velocity so the simulation acts as if the actor is moving at speed, even when stationary. 
+	// Vector is in world space. Units are cm/s. Could be used for a wind effects etc. Typical values are similar to the velocity of the object or effect, 
+	// and usually around or less than 1000 for characters/wind.
+	UPROPERTY(EditAnywhere, BlueprintReadWrite, Category = Settings)
+	FVector ExternalLinearVelocity = FVector::ZeroVector;
+
+	// Additional angular velocity that is added to the component angular velocity. This can be used to make the simulation act as if the actor is rotating
+	// even when it is not. E.g., to apply physics to a character on a podium as the camera rotates around it, to emulate the podium itself rotating.
+	// Vector is in world space. Units are rad/s.
+	UPROPERTY(EditAnywhere, BlueprintReadWrite, Category = Settings)
+	FVector ExternalAngularVelocity = FVector::ZeroVector;
+};
+
+USTRUCT(BlueprintType)
+struct CONTROLRIGPHYSICS_API FRigPhysicsDynamics
+{
+	GENERATED_BODY()
+
+	friend FArchive& operator <<(FArchive& Ar, FRigPhysicsDynamics& Data)
+	{
+		Ar << Data.Mass;
+		Ar << Data.bOverrideCentreOfMass;
+		Ar << Data.CentreOfMassOverride;
+		Ar << Data.bOverrideMomentsOfInertia;
+		Ar << Data.MomentsOfInertiaOverride;
+		return Ar;
+	}
+
+	// The total mass of this body.
+	UPROPERTY(BlueprintReadWrite, EditAnywhere, Category = Dynamics)
+	float Mass = 1.0f;
+
+	// Use this to specify the location of the centre of mass
+	UPROPERTY(BlueprintReadWrite, EditAnywhere, Category = Dynamics)
+	bool bOverrideCentreOfMass = false;
+
+	UPROPERTY(BlueprintReadWrite, EditAnywhere, Category = Dynamics, meta = (EditCondition = "bOverrideCentreOfMass"))
+	FVector CentreOfMassOverride;
+
+	// Use this to specify the moments of inertia
+	UPROPERTY(BlueprintReadWrite, EditAnywhere, Category = Dynamics)
+	bool bOverrideMomentsOfInertia = false;
+
+	UPROPERTY(BlueprintReadWrite, EditAnywhere, Category = Dynamics, meta = (EditCondition = "bOverrideMomentsOfInertia"))
+	FVector MomentsOfInertiaOverride = FVector(1.0f, 1.0f, 1.0f);
+
+};
+
+USTRUCT(BlueprintType)
+struct CONTROLRIGPHYSICS_API FRigPhysicsCollisionShape
+{
+	GENERATED_BODY()
+
+	/** Offset used when generating contact points. This allows you to smooth out
+	the Minkowski sum by radius R. Useful for making objects slide smoothly
+	on top of irregularities  */
+	UPROPERTY(Category = Shape, EditAnywhere)
+	float RestOffset = 0.0f;
+
+	UPROPERTY(Category = Shape, EditAnywhere)
+	FName Name;
+
+	/** True if this shape should contribute to the overall mass of the body it
+		belongs to. This lets you create extra collision volumes which do not affect
+		the mass properties of an object. */
+	UPROPERTY(Category = Shape, EditAnywhere)
+	bool bContributeToMass = true;
+};
+
+USTRUCT(BlueprintType)
+struct CONTROLRIGPHYSICS_API FRigPhysicsCollisionBox : public FRigPhysicsCollisionShape
+{
+	GENERATED_BODY()
+
+	friend FArchive& operator <<(FArchive& Ar, FRigPhysicsCollisionBox& Data)
+	{
+		Ar << Data.TM;
+		Ar << Data.Extents;
+		return Ar;
+	}
+
+	UPROPERTY(BlueprintReadWrite, EditAnywhere, Category = Shapes)
+	FTransform TM;
+
+	UPROPERTY(BlueprintReadWrite, EditAnywhere, Category = Shapes)
+	FVector Extents;
+};
+
+USTRUCT(BlueprintType)
+struct CONTROLRIGPHYSICS_API FRigPhysicsCollisionSphere : public FRigPhysicsCollisionShape
+{
+	GENERATED_BODY()
+
+	friend FArchive& operator <<(FArchive& Ar, FRigPhysicsCollisionSphere& Data)
+	{
+		Ar << Data.TM;
+		Ar << Data.Radius;
+		return Ar;
+	}
+
+	UPROPERTY(BlueprintReadWrite, EditAnywhere, Category = Shapes)
+	FTransform TM;
+
+	UPROPERTY(BlueprintReadWrite, EditAnywhere, Category = Shapes)
+	float Radius;
+};
+
+USTRUCT(BlueprintType)
+struct CONTROLRIGPHYSICS_API FRigPhysicsCollisionCapsule : public FRigPhysicsCollisionShape
+{
+	GENERATED_BODY()
+
+	friend FArchive& operator <<(FArchive& Ar, FRigPhysicsCollisionCapsule& Data)
+	{
+		Ar << Data.TM;
+		Ar << Data.Radius;
+		Ar << Data.Length;
+		return Ar;
+	}
+
+	UPROPERTY(BlueprintReadWrite, EditAnywhere, Category = Shapes)
+	FTransform TM;
+
+	UPROPERTY(BlueprintReadWrite, EditAnywhere, Category = Shapes)
+	float Radius;
+
+	UPROPERTY(BlueprintReadWrite, EditAnywhere, Category = Shapes)
+	float Length;
+};
+
+USTRUCT(BlueprintType)
+struct CONTROLRIGPHYSICS_API FRigPhysicsCollision
+{
+	GENERATED_BODY()
+
+	friend FArchive& operator <<(FArchive& Ar, FRigPhysicsCollision& Data)
+	{
+		Ar << Data.Boxes;
+		Ar << Data.Spheres;
+		Ar << Data.Capsules;
+		return Ar;
+	}
+
+	UPROPERTY(BlueprintReadWrite, EditAnywhere, Category = Shapes)
+	TArray<FRigPhysicsCollisionBox> Boxes;
+
+	UPROPERTY(BlueprintReadWrite, EditAnywhere, Category = Shapes)
+	TArray<FRigPhysicsCollisionSphere> Spheres;
+
+	UPROPERTY(BlueprintReadWrite, EditAnywhere, Category = Shapes)
+	TArray<FRigPhysicsCollisionCapsule> Capsules;
+
+	bool IsEmpty() const { return Boxes.IsEmpty() && Spheres.IsEmpty() && Capsules.IsEmpty(); }
+};
+
+USTRUCT(BlueprintType)
+struct CONTROLRIGPHYSICS_API FRigPhysicsArticulation
+{
+	GENERATED_BODY()
+
+	friend FArchive& operator <<(FArchive& Ar, FRigPhysicsArticulation& Data)
+	{
+		Ar << Data.bEnable;
+		Ar << Data.ParentBody;
+		Ar << Data.bAutoCalculateParentOffset;
+		Ar << Data.bAutoCalculateChildOffset;
+		Ar << Data.ExtraParentOffset;
+		Ar << Data.ExtraChildOffset;
+		Ar << Data.LinearLimit;
+		Ar << Data.AngularLimit;
+		return Ar;
+	}
+
+	UPROPERTY(BlueprintReadWrite, EditAnywhere, Category = Articulation)
+	bool bEnable = true;
+
+	// The body element that will be used as the parent for the articulation (the child will always
+	// be whatever body this articulation is owned by. If it is unset, then a suitable body will be
+	// searched for (by assuming this body is the child of a bone which is the child of a parent
+	// bone which has a body. If it is otherwise invalid, or the parent bone is not physical, then
+	// no articulation will be created.
+	UPROPERTY(BlueprintReadWrite, EditAnywhere, Category = Articulation)
+	FRigComponentKey ParentBody;
+
+	UPROPERTY(BlueprintReadWrite, EditAnywhere, Category = Articulation)
+	bool bAutoCalculateParentOffset = true;
+
+	// The parent frame offset applied after any auto-calculation
+	UPROPERTY(BlueprintReadWrite, EditAnywhere, Category = Articulation)
+	FTransform ExtraParentOffset;
+
+	UPROPERTY(BlueprintReadWrite, EditAnywhere, Category = Articulation)
+	bool bAutoCalculateChildOffset = true;
+
+	// The child frame offset applied after any auto-calculation
+	UPROPERTY(BlueprintReadWrite, EditAnywhere, Category = Articulation)
+	FTransform ExtraChildOffset;
+
+	UPROPERTY(BlueprintReadWrite, EditAnywhere, Category = Articulation)
+	FVector LinearLimit = FVector::ZeroVector;
+
+	// Angular limit in degrees
+	UPROPERTY(BlueprintReadWrite, EditAnywhere, Category = Articulation)
+	FVector AngularLimit = FVector::ZeroVector;
+};
+
+USTRUCT(BlueprintType)
+struct CONTROLRIGPHYSICS_API FRigPhysicsBodySolverSettings
+{
+	GENERATED_BODY()
+
+	FRigPhysicsBodySolverSettings(
+		FRigPhysicsSolverID InSolverID = FRigPhysicsSolverID(), FRigElementKey InTargetBone = FRigElementKey())
+		: SolverID(InSolverID), TargetBone(InTargetBone) {}
+
+	friend FArchive& operator <<(FArchive& Ar, FRigPhysicsBodySolverSettings& Data)
+	{
+		Ar << Data.SolverID;
+		Ar << Data.TargetBone;
+		return Ar;
+	}
+
+	UPROPERTY(BlueprintReadWrite, EditAnywhere, Category = Articulation)
+	FRigPhysicsSolverID SolverID;
+
+	/// The bone that is used to initialize physics, as well as what to track when the body is set to be kinematic.
+	/// Note that if this is unset, the it will default to the parent of the physics element.
+	UPROPERTY(BlueprintReadWrite, EditAnywhere, Category = Articulation)
+	FRigElementKey SourceBone;
+
+	/// The bone that is written to following simulation. Note that if this is unset, the it will
+	/// default to the parent of the physics element.
+	UPROPERTY(BlueprintReadWrite, EditAnywhere, Category = Articulation)
+	FRigElementKey TargetBone;
+};
+
+
+USTRUCT(BlueprintType)
+struct CONTROLRIGPHYSICS_API FRigPhysicsComponent : public FRigBaseComponent
+{
+public:
+
+	GENERATED_BODY()
+	DECLARE_RIG_COMPONENT_METHODS(FRigPhysicsComponent)
+
+	UPROPERTY(BlueprintReadOnly, EditAnywhere, Category = Physics, meta = (ShowOnlyInnerProperties))
+	FRigPhysicsBodySolverSettings BodySolverSettings;
+
+	UPROPERTY(BlueprintReadOnly, EditAnywhere, Category = Physics, meta = (ShowOnlyInnerProperties))
+	FRigPhysicsDynamics Dynamics;
+
+	UPROPERTY(BlueprintReadOnly, EditAnywhere, Category = Physics, meta = (ShowOnlyInnerProperties))
+	FRigPhysicsCollision Collision;
+
+	UPROPERTY(BlueprintReadOnly, EditAnywhere, Category = Physics, meta = (ShowOnlyInnerProperties))
+	FRigPhysicsArticulation Articulation;
+
+	UPROPERTY(BlueprintReadOnly, EditAnywhere, Category = Physics)
+	FPhysicsControlModifierData BodyData;
+
+	virtual void Save(FArchive& A) override;
+	virtual void Load(FArchive& Ar) override;
+
+private:
+
+	friend class URigHierarchy;
+	friend struct FRigBaseElement;
+};
+
+USTRUCT(BlueprintType)
+struct CONTROLRIGPHYSICS_API FRigPhysicsSolverSettings
+{
+	GENERATED_BODY()
+
+	UPROPERTY(BlueprintReadOnly, EditAnywhere, Category = SolverSettings)
+	ERigPhysicsSimulationSpace SimulationSpace = ERigPhysicsSimulationSpace::WorldSpace;
+
+	// The space in which collision shapes are defined
+	UPROPERTY(BlueprintReadOnly, EditAnywhere, Category = SolverSettings)
+	ERigPhysicsSimulationSpace CollisionSpace = ERigPhysicsSimulationSpace::WorldSpace;
+
+	/** The bone to use for the collision or simulation space (only relevant if one is set to BoneSpace) */
+	UPROPERTY(EditAnywhere, Category = Settings, meta = (EditCondition = "SimulationSpace == ERigPhysicsSimulationSpace::BoneSpace"))
+	FRigElementKey SpaceBone;
+
+	// The collision shapes defined in the collision space - e.g. for representing a ground etc
+	UPROPERTY(BlueprintReadOnly, EditAnywhere, Category = Physics, meta = (ShowOnlyInnerProperties))
+	FRigPhysicsCollision Collision;
+
+	UPROPERTY(EditAnywhere, BlueprintReadWrite, Category = SolverSettings)
+	FVector Gravity = FVector(0, 0, -981.0);
+
+	/**
+	 * The number of position iterations to run. The position solve is responsible for de-penetration.
+	 * Increasing this will improve simulation stability, but increase the cost.
+	*/
+	UPROPERTY(EditAnywhere, BlueprintReadWrite, Category = SolverSettings)
+	int32 PositionIterations = 6;
+
+	/**
+	 * The number of velocity iterations to run. The velocity solve is responsible for restitution (bounce) and friction.
+	 * This should usually be 1, but could be 0 if you don't care about friction and restitution.
+	*/
+	UPROPERTY(EditAnywhere, BlueprintReadWrite, Category = SolverSettings)
+	int32 VelocityIterations = 1;
+
+	/**
+	 * The number of projection iterations to run. The projection phase is a final pass over the constraints, applying
+	 * a semi-physical correction to any joint errors remaining after the position and velocity solves. It can be
+	 * very helpful to stabilize joint chains, but can cause issues with collision response. The projection magnitude
+	 * can be controlled per-constraint in the constraint settings (assuming ProjectionIteration is not zero).
+	 * This should be left as 1 in almost all cases.
+	*/
+	UPROPERTY(EditAnywhere, BlueprintReadWrite, Category = SolverSettings)
+	int32 ProjectionIterations = 1;
+
+	/**
+	 * When bodies are penetrating, this is the maximum velocity delta that can be applied in one frame.
+	*/
+	UPROPERTY(EditAnywhere, BlueprintReadWrite, Category = SolverSettings, meta = (ClampMin = 0))
+	float MaxDepenetrationVelocity = 0.0f;
+
+	/**
+	 * The recommended fixed timestep for the RBAN solver. Set to 0 to run with variable timestep (default).
+	 * NOTE: If this value is non-zero and less than the current frame time, the simulation will step multiple times
+	 * which increases the cost.
+	 */
+	UPROPERTY(EditAnywhere, BlueprintReadWrite, Category = SolverSettings, meta = (ClampMin = 0))
+	float FixedTimeStep = 0.0f;
+
+	// The maximum number of solver steps that can be made
+	UPROPERTY(EditAnywhere, BlueprintReadWrite, Category = SolverSettings, meta = (ClampMin = 1))
+	int32 MaxTimeSteps = 10;
+
+	/**
+	 * If a variable timestep is in use, this is the maximum delta time that can be used. If this is smaller than the
+	 * requested total delta time then multiple steps will be used.
+	 */
+	UPROPERTY(EditAnywhere, BlueprintReadWrite, Category = SolverSettings, meta = (ClampMin = 0))
+	float MaxDeltaTime = 1.f / 30.f;
+
+	/**
+	 * Whether to use the linear or non-linear solver for RBAN Joints. The linear solver is significantly cheaper than
+	 * the non-linear solver when you are running multiple iterations, but is more likely to suffer from jitter.
+	 * In general you should try to use the linear solver and increase the PositionIterations to improve stability if
+	 * possible, only using the non-linear solver as a last resort.
+	*/
+	UPROPERTY(EditAnywhere, BlueprintReadWrite, Category = SolverSettings)
+	bool bUseLinearJointSolver = true;
+
+	/**
+	 * When solving joints, whether to solve the positions last (as opposed to the orientations)
+	 */
+	UPROPERTY(EditAnywhere, BlueprintReadWrite, Category = SolverSettings)
+	bool bSolveJointPositionsLast = true;
+
+	/**
+	 * It enables the use of multi-point contact manifolds, which are created only once at the start of each tick.
+	 * When disabled, a single-point contact is generated in each solver iteration which is more expensive.
+	*/
+	UPROPERTY(EditAnywhere, BlueprintReadWrite, Category = CollisionSettings)
+	bool bUseManifolds = true;
+};
diff --git a/Engine/Plugins/Experimental/PhysicsControl/Source/ControlRigPhysics/Public/RigUnit_Physics.h b/Engine/Plugins/Experimental/PhysicsControl/Source/ControlRigPhysics/Public/RigUnit_Physics.h
new file mode 100644
--- /dev/null
+++ b/Engine/Plugins/Experimental/PhysicsControl/Source/ControlRigPhysics/Public/RigUnit_Physics.h
@@ -0,0 +1,204 @@
+// Copyright Epic Games, Inc. All Rights Reserved.
+
+#pragma once
+
+#include "Units/Execution/RigUnit_DynamicHierarchy.h"
+#include "RigPhysicsComponents.h"
+#include "PhysicsControlData.h"
+
+#include "RigUnit_Physics.generated.h"
+
+/**
+ * Adds a new physics solver to the hierarchy
+ * Note: This node only runs as part of the construction event.
+ */
+USTRUCT(meta=(DisplayName="Spawn Physics Solver", Keywords="Construction,Create,New,Simulation,Physics", Varying))
+struct CONTROLRIGPHYSICS_API FRigUnit_HierarchyAddPhysicsSolver : public FRigUnit_DynamicHierarchyBaseMutable
+{
+	GENERATED_BODY()
+
+	FRigUnit_HierarchyAddPhysicsSolver()
+	{
+		Name = TEXT("Solver");
+		SolverID = FRigPhysicsSolverID();
+	}
+
+	/*
+	 * The name of the new solver to add
+	 */
+	UPROPERTY(meta = (Input))
+	FName Name;
+
+	UPROPERTY(meta = (Input))
+	FRigPhysicsSolverSettings SolverSettings;
+
+	UPROPERTY(meta = (Input))
+	FRigPhysicsSimulationSpaceSettings SimulationSpaceSettings;
+
+	/*
+	 * The identifier of the solver
+	 */
+	UPROPERTY(meta = (Output))
+	FRigPhysicsSolverID SolverID;
+
+	RIGVM_METHOD()
+	virtual void Execute() override;
+};
+
+USTRUCT(meta = (DisplayName = "Set Source Bone", Keywords = "Physics", Varying))
+struct CONTROLRIGPHYSICS_API FRigUnit_HierarchySetPhysicsBodySourceBone : public FRigUnit_DynamicHierarchyBaseMutable
+{
+	GENERATED_BODY()
+
+	RIGVM_METHOD()
+	virtual void Execute() override;
+
+	UPROPERTY(meta = (Input))
+	FRigComponentKey PhysicsBody;
+
+	UPROPERTY(meta = (Input))
+	FRigElementKey SourceBone;
+};
+
+USTRUCT(meta = (DisplayName = "Set Target Bone", Keywords = "Physics", Varying))
+struct CONTROLRIGPHYSICS_API FRigUnit_HierarchySetPhysicsBodyTargetBone : public FRigUnit_DynamicHierarchyBaseMutable
+{
+	GENERATED_BODY()
+
+	RIGVM_METHOD()
+	virtual void Execute() override;
+
+	UPROPERTY(meta = (Input))
+	FRigComponentKey PhysicsBody;
+
+	UPROPERTY(meta = (Input))
+	FRigElementKey TargetBone;
+};
+
+
+USTRUCT(meta = (DisplayName = "Set Physics Body Data", Keywords = "Physics", Varying))
+struct CONTROLRIGPHYSICS_API FRigUnit_HierarchySetPhysicsBodySparseData : public FRigUnit_DynamicHierarchyBaseMutable
+{
+	GENERATED_BODY()
+
+	RIGVM_METHOD()
+	virtual void Execute() override;
+
+	UPROPERTY(meta = (Input))
+	FRigComponentKey PhysicsBody;
+
+	/*
+	 * The initial global transform of the spawned element
+	 */
+	UPROPERTY(meta = (Input))
+	FPhysicsControlModifierSparseData Data;
+};
+
+/**
+ * Adds a new physics body to the hierarchy, which should be a child of a regular animation joint.
+ * Note: This node only runs as part of the construction event.
+ */
+USTRUCT(meta=(DisplayName="Add Physics Body Component", Keywords="Construction,Create,New,Joint,Body,Skeleton,Physics", Varying))
+struct CONTROLRIGPHYSICS_API FRigUnit_AddPhysicsBody : public FRigUnitMutable
+{
+	GENERATED_BODY()
+
+	FRigUnit_AddPhysicsBody()
+	{
+		Item = FRigElementKey();
+		Result = FRigComponentKey();
+		SolverID = FRigPhysicsSolverID();
+	}
+
+	RIGVM_METHOD()
+	virtual void Execute() override;
+
+	UPROPERTY(meta = (Input))
+	FRigElementKey Item;
+
+	UPROPERTY(meta = (Output))
+	FRigComponentKey Result;
+
+	/*
+	 * The solver to relate this new physics element to
+	 */
+	UPROPERTY(meta = (Input))
+	FRigPhysicsSolverID SolverID;
+
+	/*
+	 * The dynamics properties of the new physics element
+	 */
+	UPROPERTY(meta = (Input))
+	FRigPhysicsDynamics Dynamics;
+
+	/*
+	 * The collision properties of the new physics element
+	 */
+	UPROPERTY(meta = (Input))
+	FRigPhysicsCollision Collision;
+
+	/**
+	 * The articulation properties of the new physics element
+	 */
+	UPROPERTY(meta = (Input))
+	FRigPhysicsArticulation Articulation;
+
+	/**
+	 * The runtime modifiable data
+	 */
+	UPROPERTY(meta = (Input))
+	FPhysicsControlModifierData BodyData;
+};
+
+/**
+ * Instantiates all the objects in the physics world. Some properties can't be modified after this happens.
+ * Note that it will happen automatically during the first simulation step if it hasn't been explicitly
+ * requested. Explicit instantiation allows the timing to be controlled, as allocations etc may cause some
+ * delays.
+ */
+USTRUCT(meta = (DisplayName = "Instantiate physics", Keywords = "Physics"))
+struct CONTROLRIGPHYSICS_API FRigUnit_InstantiatePhysics : public FRigUnitMutable
+{
+	GENERATED_BODY()
+
+	RIGVM_METHOD()
+	virtual void Execute() override;
+
+	/*
+	 * The solver to relate this new physics element to
+	 */
+	UPROPERTY(meta = (Input))
+	FRigPhysicsSolverID SolverID;
+};
+
+
+/**
+ * Steps the specified physics solver
+ */
+USTRUCT(meta = (DisplayName = "Step Physics Solver", Keywords = "Simulate,Physics"))
+struct CONTROLRIGPHYSICS_API FRigUnit_StepPhysicsSolver : public FRigUnitMutable
+{
+	GENERATED_BODY()
+
+	FRigUnit_StepPhysicsSolver()
+	{
+		DeltaTimeOverride = -1.0f;
+	}
+
+	RIGVM_METHOD()
+	virtual void Execute() override;
+
+	/*
+	 * The solver to relate this new physics element to
+	 */
+	UPROPERTY(meta = (Input))
+	FRigPhysicsSolverID SolverID;
+
+	/*
+	 * Normally the execute context time will be used, but if this is positive then it will override it.
+	 */
+	UPROPERTY(meta = (Input))
+	float DeltaTimeOverride;
+};
+
+
diff --git a/Engine/Plugins/Experimental/PhysicsControl/Source/PhysicsControl/Private/RigidBodyWithControl.cpp b/Engine/Plugins/Experimental/PhysicsControl/Source/PhysicsControl/Private/RigidBodyWithControl.cpp
--- a/Engine/Plugins/Experimental/PhysicsControl/Source/PhysicsControl/Private/RigidBodyWithControl.cpp
+++ b/Engine/Plugins/Experimental/PhysicsControl/Source/PhysicsControl/Private/RigidBodyWithControl.cpp
@@ -721,35 +721,46 @@ void FAnimNode_RigidBodyWithControl::ApplyControl(FRigidBodyControlRecord& Contr
 }
 
 //======================================================================================================================
-void FAnimNode_RigidBodyWithControl::ApplyModifier(
-	const FRigidBodyModifierRecord& BodyModifierRecord, 
-	const FVector&                  SimSpaceGravity)
+void UpdateBodyFromModifierData(
+	ImmediatePhysics::FActorHandle*    ActorHandle, 
+	ImmediatePhysics::FSimulation*     PhysicsSimulation,
+	const FPhysicsControlModifierData& ModifierData,
+	const FVector&                     SimSpaceGravity)
 {
-	if (BodyModifierRecord.ActorHandle)
+	if (ModifierData.MovementType != EPhysicsMovementType::Default)
 	{
-		if (BodyModifierRecord.ModifierData.MovementType != EPhysicsMovementType::Default)
-		{
-			// Note that there's an early out if there's no change needed, so this should be OK.
-			PhysicsSimulation->SetIsKinematic(
-				BodyModifierRecord.ActorHandle,
-				BodyModifierRecord.ModifierData.MovementType != EPhysicsMovementType::Simulated);
-		}
+		// Note that there's an early out if there's no change needed, so this should be OK.
+		PhysicsSimulation->SetIsKinematic(
+			ActorHandle,
+			ModifierData.MovementType != EPhysicsMovementType::Simulated);
+	}
 
-		// Note that the actual kinematic targets will be set separately, since they need to be set
-		// for all kinematics whether or not they were under a modifier.
+	// Note that the actual kinematic targets will be set separately, since they need to be set
+	// for all kinematics whether or not they were under a modifier.
 
-		// Scale gravity
-		float GravityMultiplier = BodyModifierRecord.ModifierData.GravityMultiplier;
-		if (GravityMultiplier != 1 && BodyModifierRecord.ActorHandle->IsGravityEnabled())
-		{
-			float Mass = (float) BodyModifierRecord.ActorHandle->GetMass();
-			FVector AntiGravityForce = SimSpaceGravity * (-Mass * (1.0f - GravityMultiplier));
-			BodyModifierRecord.ActorHandle->AddForce(AntiGravityForce);
-		}
+	// Scale gravity
+	float GravityMultiplier = ModifierData.GravityMultiplier;
+	if (GravityMultiplier != 1 && ActorHandle->IsGravityEnabled())
+	{
+		float Mass = (float) ActorHandle->GetMass();
+		FVector AntiGravityForce = SimSpaceGravity * (-Mass * (1.0f - GravityMultiplier));
+		ActorHandle->AddForce(AntiGravityForce);
+	}
+
+	// Set collision
+	PhysicsSimulation->SetHasCollision(
+		ActorHandle, CollisionEnabledHasPhysics(ModifierData.CollisionType));
+}
 
-		// Set collision
-		PhysicsSimulation->SetHasCollision(
-			BodyModifierRecord.ActorHandle, CollisionEnabledHasPhysics(BodyModifierRecord.ModifierData.CollisionType));
+//======================================================================================================================
+void FAnimNode_RigidBodyWithControl::ApplyModifier(
+	const FRigidBodyModifierRecord& BodyModifierRecord, 
+	const FVector&                  SimSpaceGravity)
+{
+	if (BodyModifierRecord.ActorHandle && PhysicsSimulation)
+	{
+		UpdateBodyFromModifierData(
+			BodyModifierRecord.ActorHandle, PhysicsSimulation, BodyModifierRecord.ModifierData, SimSpaceGravity);
 	}
 }
 
diff --git a/Engine/Plugins/Experimental/PhysicsControl/Source/PhysicsControl/Public/AnimNode_RigidBodyWithControl.h b/Engine/Plugins/Experimental/PhysicsControl/Source/PhysicsControl/Public/AnimNode_RigidBodyWithControl.h
--- a/Engine/Plugins/Experimental/PhysicsControl/Source/PhysicsControl/Public/AnimNode_RigidBodyWithControl.h
+++ b/Engine/Plugins/Experimental/PhysicsControl/Source/PhysicsControl/Public/AnimNode_RigidBodyWithControl.h
@@ -50,6 +50,16 @@ enum class MapConstraintsBehaviorType : uint8
 	None UMETA(DisplayName = "Nothing", ToolTip = "Nothing"),
 };
 
+/**
+ * A helper function that lets any immediate simulation user update their bodies in the same way that
+ * RigidBodyWithControl does.
+ */
+PHYSICSCONTROL_API void UpdateBodyFromModifierData(
+	ImmediatePhysics::FActorHandle* ActorHandle,
+	ImmediatePhysics::FSimulation* PhysicsSimulation,
+	const FPhysicsControlModifierData& ModifierData,
+	const FVector& SimSpaceGravity);
+
 /**
  * Controller that simulates physics based on the physics asset of the skeletal mesh component
  */
diff --git a/Engine/Plugins/Runtime/RigVM/Source/RigVM/Private/RigVMCore/RigVMDrawInterface.cpp b/Engine/Plugins/Runtime/RigVM/Source/RigVM/Private/RigVMCore/RigVMDrawInterface.cpp
--- a/Engine/Plugins/Runtime/RigVM/Source/RigVM/Private/RigVMCore/RigVMDrawInterface.cpp
+++ b/Engine/Plugins/Runtime/RigVM/Source/RigVM/Private/RigVMCore/RigVMDrawInterface.cpp
@@ -259,13 +259,12 @@ void FRigVMDrawInterface::DrawArc(const FTransform& WorldOffset, const FTransfor
 	V = Rotation.RotateVector(V);
 	Instruction.Positions.Add(V);
 	float StepAngle = (MaximumAngle - MinimumAngle) / float(Count);
-	if (FMath::Abs<float>(MaximumAngle - MinimumAngle) >= PI * 2.f - SMALL_NUMBER)
+	if (FMath::Abs<float>(MaximumAngle - MinimumAngle) >= TWO_PI)
 	{
-		StepAngle = (PI * 2.f) / float(Count);
-		Count++;
+		StepAngle = TWO_PI / float(Count);
 	}
 	Rotation = FQuat(FVector(0.f, 0.f, 1.f), StepAngle);
-	for(int32 Index=1;Index<Count;Index++)
+	for(int32 Index=0;Index<Count;Index++)
 	{
 		V = Rotation.RotateVector(V);
 		Instruction.Positions.Add(V);
@@ -280,7 +279,53 @@ void FRigVMDrawInterface::DrawCircle(const FTransform& WorldOffset, const FTrans
 	DrawArc(WorldOffset, Transform, Radius, 0, TWO_PI, Color, Thickness, Detail);
 }
 
-void FRigVMDrawInterface::DrawCone(const FTransform& WorldOffset, const FTransform& ConeOffset, float Angle1, float Angle2, uint32 NumSides, bool bDrawSideLines, const FLinearColor& SideLineColor, FMaterialRenderProxy* const MaterialRenderProxy)
+void FRigVMDrawInterface::DrawSphere(const FTransform& WorldOffset, const FTransform& Transform, float Radius, 
+	const FLinearColor& Color, float Thickness, int32 Detail)
+{
+	// Circles are drawn in the XY plane
+	FTransform R1(FQuat(FVector(1.f, 0.f, 0.f), HALF_PI));
+	FTransform R2(FQuat(FVector(0.f, 1.f, 0.f), HALF_PI));
+	DrawCircle(WorldOffset, Transform, Radius, Color, Thickness, Detail);
+	DrawCircle(WorldOffset, R1 * Transform, Radius, Color, Thickness, Detail);
+	DrawCircle(WorldOffset, R2 * Transform, Radius, Color, Thickness, Detail);
+}
+
+void FRigVMDrawInterface::DrawHemisphere(const FTransform& WorldOffset, const FTransform& Transform, float Radius,
+	const FLinearColor& Color, float Thickness, int32 Detail)
+{
+	// Circles are drawn in the XY plane
+	FTransform R1(FQuat(FVector(1.f, 0.f, 0.f), HALF_PI));
+	FTransform R2(FQuat(FVector(0.f, 1.f, 0.f), -HALF_PI));
+	DrawCircle(WorldOffset, Transform, Radius, Color, Thickness, Detail);
+	DrawArc(WorldOffset, R1 * Transform, Radius, 0, PI, Color, Thickness, Detail);
+	DrawArc(WorldOffset, R2 * Transform, Radius, -HALF_PI, HALF_PI, Color, Thickness, Detail);
+}
+
+void FRigVMDrawInterface::DrawCapsule(const FTransform& WorldOffset, const FTransform& Transform,
+	float Radius, float Length, const FLinearColor& Color, float Thickness, int32 Detail)
+{
+	FTransform T1(FVector(0, 0, Length * 0.5));
+	FTransform T2(FQuat(FVector(0.f, 1.f, 0.f), PI), FVector(0, 0, -Length * 0.5));
+
+	DrawHemisphere(WorldOffset, T1 * Transform, Radius, Color, Thickness, Detail);
+	DrawHemisphere(WorldOffset, T2 * Transform, Radius, Color, Thickness, Detail);
+
+	DrawLine(WorldOffset, 
+		Transform.TransformPosition(FVector(Radius, 0, Length * 0.5)), 
+		Transform.TransformPosition(FVector(Radius, 0, -Length * 0.5)), Color, Thickness);
+	DrawLine(WorldOffset,
+		Transform.TransformPosition(FVector(-Radius, 0, Length * 0.5)),
+		Transform.TransformPosition(FVector(-Radius, 0, -Length * 0.5)), Color, Thickness);
+	DrawLine(WorldOffset,
+		Transform.TransformPosition(FVector(0, Radius, Length * 0.5)),
+		Transform.TransformPosition(FVector(0, Radius, -Length * 0.5)), Color, Thickness);
+	DrawLine(WorldOffset,
+		Transform.TransformPosition(FVector(0, -Radius, Length * 0.5)),
+		Transform.TransformPosition(FVector(0, -Radius, -Length * 0.5)), Color, Thickness);
+}
+
+
+void FRigVMDrawInterface::DrawCone(const FTransform& WorldOffset, const FTransform& ConeOffset, float Angle1, float Angle2, uint32 NumSides, bool bDrawSideLines, const FLinearColor& SideLineColor, FMaterialRenderProxy* const MaterialRenderProxy, float SideLineThickness)
 {
 	if (!IsEnabled())
 	{
@@ -311,7 +356,7 @@ void FRigVMDrawInterface::DrawCone(const FTransform& WorldOffset, const FTransfo
 			float Fraction = (float)i / (float)(4);
 			float Azi = 2.f * PI * Fraction;
 			FVector ConeVert = CalcConeVert(Angle1, Angle2, Azi);
-			DrawLine(WorldOffset, ConeOffset.GetLocation(), ConeOffset.TransformPosition(ConeVert), SideLineColor, 1.f);
+			DrawLine(WorldOffset, ConeOffset.GetLocation(), ConeOffset.TransformPosition(ConeVert), SideLineColor, SideLineThickness);
 		}
 	}
 }
diff --git a/Engine/Plugins/Runtime/RigVM/Source/RigVM/Public/RigVMCore/RigVMDrawInterface.h b/Engine/Plugins/Runtime/RigVM/Source/RigVM/Public/RigVMCore/RigVMDrawInterface.h
--- a/Engine/Plugins/Runtime/RigVM/Source/RigVM/Public/RigVMCore/RigVMDrawInterface.h
+++ b/Engine/Plugins/Runtime/RigVM/Source/RigVM/Public/RigVMCore/RigVMDrawInterface.h
@@ -20,14 +20,19 @@ public:
 	void DrawLines(const FTransform& WorldOffset, const TArrayView<const FVector>& Positions, const FLinearColor& Color, float Thickness = 0.f);
 	void DrawLineStrip(const FTransform& WorldOffset, const TArrayView<const FVector>& Positions, const FLinearColor& Color, float Thickness = 0.f);
 	void DrawBox(const FTransform& WorldOffset, const FTransform& Transform, const FLinearColor& Color, float Thickness = 0.f);
+	void DrawSphere(const FTransform& WorldOffset, const FTransform& Transform, float Radius, const FLinearColor& Color, float Thickness = 0.f, int32 Detail = 12);
+	// Hemispheres are drawn so that they extend up along +ve Z
+	void DrawHemisphere(const FTransform& WorldOffset, const FTransform& Transform, float Radius, const FLinearColor& Color, float Thickness = 0.f, int32 Detail = 12);
+	void DrawCapsule(const FTransform& WorldOffset, const FTransform& Transform, float Radius, float Length, const FLinearColor& Color, float Thickness = 0.f, int32 Detail = 12);
 	void DrawAxes(const FTransform& WorldOffset, const FTransform& Transform, float Size, float Thickness = 0.f);
 	void DrawAxes(const FTransform& WorldOffset, const TArrayView<const FTransform>& Transforms, float Size, float Thickness = 0.f);
 	void DrawAxes(const FTransform& WorldOffset, const FTransform& Transform, const FLinearColor& InColor, float Size, float Thickness = 0.f);
 	void DrawAxes(const FTransform& WorldOffset, const TArrayView<const FTransform>& Transforms, const FLinearColor& InColor, float Size, float Thickness = 0.f);
 	void DrawRectangle(const FTransform& WorldOffset, const FTransform& Transform, float Size, const FLinearColor& Color, float Thickness);
 	void DrawArc(const FTransform& WorldOffset, const FTransform& Transform, float Radius, float MinimumAngle, float MaximumAngle, const FLinearColor& Color, float Thickness, int32 Detail);
+	// Circles are drawn in the XY plane
 	void DrawCircle(const FTransform& WorldOffset, const FTransform& Transform, float Radius, const FLinearColor& Color, float Thickness, int32 Detail);
-	void DrawCone(const FTransform& WorldOffset, const FTransform& ConeOffset, float Angle1, float Angle2, uint32 NumSides, bool bDrawSideLines, const FLinearColor& SideLineColor, FMaterialRenderProxy* const MaterialRenderProxy);
+	void DrawCone(const FTransform& WorldOffset, const FTransform& ConeOffset, float Angle1, float Angle2, uint32 NumSides, bool bDrawSideLines, const FLinearColor& SideLineColor, FMaterialRenderProxy* const MaterialRenderProxy, float SideLineThickness = 1.0f);
 	void DrawArrow(const FTransform& WorldOffset, const FVector& Direction, const FVector& Side, const FLinearColor& InColor, float Thickness);
 	void DrawPlane(const FTransform& WorldOffset, const FVector2D& Scale, const FLinearColor& MeshColor, bool bDrawLines, const FLinearColor& LineColor, FMaterialRenderProxy* const MaterialRenderProxy);
 
