Fixing some errors in help/comments. no functional changes.

#rnx

diff --git a/Engine/Plugins/Experimental/PhysicsControl/Source/PhysicsControl/Public/PhysicsControlComponent.h b/Engine/Plugins/Experimental/PhysicsControl/Source/PhysicsControl/Public/PhysicsControlComponent.h
--- a/Engine/Plugins/Experimental/PhysicsControl/Source/PhysicsControl/Public/PhysicsControlComponent.h
+++ b/Engine/Plugins/Experimental/PhysicsControl/Source/PhysicsControl/Public/PhysicsControlComponent.h
@@ -94,8 +94,6 @@ public:
 	 * @param ControlTarget Describes the initial target for the new control
 	 * @param Set Which set to include the control in (optional). Note that it automatically
 	 *        gets added to the set "All"
-	 * @param bEnabled If true then the control will be enabled immediately. If false you will need to call
-	 *        SetControlEnabled(true) in order to activate it.
 	 * @param NamePrefix Optional string that is prefixed to the control that is created.
 	 * @return The name of the new control
 	 */
@@ -119,8 +117,6 @@ public:
 	 * @param ControlTarget Describes the initial target for the new control
 	 * @param Set Which set to include the control in (optional). Note that it automatically
 	 *        gets added to the set "All"
-	 * @param bEnabled If true then the control will be enabled immediately. If false you will need to call
-	 *        SetControlEnabled(true) in order to enable it.
 	 * @return True if a new control was created, false if a control of the specified name already exists
 	 */
 	UFUNCTION(BlueprintCallable, Category = PhysicsControl, Meta = (ReturnDisplayName = "Success"))
@@ -144,8 +140,6 @@ public:
 	 * @param ControlData   Describes the initial strength etc of the new control
 	 * @param Set Which set to include the control in (optional). Note that it automatically
 	 *        gets added to the set "All"
-	 * @param bEnabled If true then the control will be enabled immediately. If false you will need to call
-	 *        SetControlEnabled(true) in order to enable it.
 	 * @return An array of the controls that have been created
 	 */
 	UFUNCTION(BlueprintCallable, Category = PhysicsControl)
@@ -198,8 +192,6 @@ public:
 	 * @param ControlData   Describes the initial strength etc of the new control
 	 * @param Set Which set to include the control in (optional). Note that it automatically
 	 *        gets added to the set "All"
-	 * @param bEnabled If true then the control will be enabled immediately. If false you will need to call
-	 *        SetControlEnabled(true) in order to enable it.
 	 * @return An array of the controls that have been created
 	 */
 	UFUNCTION(BlueprintCallable, Category = PhysicsControl)
