Character animation seems to be skydiving diagonally when gliding straight forward

RigidBodyWithCotrol now matches RBAN simulation space, velocities and acceslerations.

#tests Skydive in PIE

diff --git a/Engine/Plugins/Experimental/PhysicsControl/Source/PhysicsControl/Private/AnimNode_RigidBodyWithControl.cpp b/Engine/Plugins/Experimental/PhysicsControl/Source/PhysicsControl/Private/AnimNode_RigidBodyWithControl.cpp
--- a/Engine/Plugins/Experimental/PhysicsControl/Source/PhysicsControl/Private/AnimNode_RigidBodyWithControl.cpp
+++ b/Engine/Plugins/Experimental/PhysicsControl/Source/PhysicsControl/Private/AnimNode_RigidBodyWithControl.cpp
@@ -185,12 +185,8 @@ FAnimNode_RigidBodyWithControl::FAnimNode_RigidBodyWithControl()
 	, CapturedFrozenPose()
 	, CapturedFrozenCurves()
 	, PreviousComponentLinearVelocity(0.0f)
-	, SimSpacePreviousComponentToWorld()
-	, SimSpacePreviousBoneToComponent()
-	, SimSpacePreviousComponentLinearVelocity(0.0f)
-	, SimSpacePreviousComponentAngularVelocity(0.0f)
-	, SimSpacePreviousBoneLinearVelocity(0.0f)
-	, SimSpacePreviousBoneAngularVelocity(0.0f)
+	, PreviousSimulationSpaceTransform()
+	, PreviousPreviousSimulationSpaceTransform()
 {
 }
 
@@ -332,111 +328,72 @@ void FAnimNode_RigidBodyWithControl::InitSimulationSpace(
 	const FTransform& ComponentToWorld,
 	const FTransform& BoneToComponent)
 {
-	SimSpacePreviousComponentToWorld = ComponentToWorld;
-	SimSpacePreviousBoneToComponent = BoneToComponent;
-	SimSpacePreviousComponentLinearVelocity = FVector::ZeroVector;
-	SimSpacePreviousComponentAngularVelocity = FVector::ZeroVector;
-	SimSpacePreviousBoneLinearVelocity = FVector::ZeroVector;
-	SimSpacePreviousBoneAngularVelocity = FVector::ZeroVector;
+	FTransform SimulationSpaceTransform = SpaceToWorldTransform(SimulationSpace, ComponentToWorld, BoneToComponent);
+	PreviousSimulationSpaceTransform = SimulationSpaceTransform;
+	PreviousPreviousSimulationSpaceTransform = SimulationSpaceTransform;
 }
 
-void FAnimNode_RigidBodyWithControl::CalculateSimulationSpace(
+void FAnimNode_RigidBodyWithControl::CalculateSimulationSpaceMotion(
 	ESimulationSpace         Space, 
+	const FTransform& SimulationSpaceTransform,
 	const FTransform&        ComponentToWorld, 
-	const FTransform&        BoneToComponent,
 	const float              Dt,
 	const FSimSpaceSettings& Settings,
-	FSimulationSpaceData&    OutSimulationSpaceData)
+	FVector& SpaceLinearVel, 
+	FVector& SpaceAngularVel, 
+	FVector& SpaceLinearAcc, 
+	FVector& SpaceAngularAcc)
 {
-	// World-space transform of the simulation space
-	OutSimulationSpaceData.Transform = SpaceToWorldTransform(Space, ComponentToWorld, BoneToComponent);
-	OutSimulationSpaceData.LinearVel = FVector::ZeroVector;
-	OutSimulationSpaceData.AngularVel = FVector::ZeroVector;
-	OutSimulationSpaceData.LinearAcc = FVector::ZeroVector;
-	OutSimulationSpaceData.AngularAcc = FVector::ZeroVector;
-
-	// The simulation scale does not change - we scale the inputs and outputs instead.
-	// This means we do not support phantom forces resulting from scale changes, but that's ok.
-	// NOTE: If we don't clear the scale, rapid scaling to zero can introduce large phantom forces
-	// leading to major instability in the simulation
-	OutSimulationSpaceData.Transform.SetScale3D(FVector::One());
+	SpaceLinearVel = FVector::ZeroVector;
+	SpaceAngularVel = FVector::ZeroVector;
+	SpaceLinearAcc = FVector::ZeroVector;
+	SpaceAngularAcc = FVector::ZeroVector;
 
 	// If the system is disabled, nothing else to do
-	if ((Settings.WorldAlpha == 0.0f) || (Dt < SMALL_NUMBER))
+	if ((Settings.WorldAlpha == 0.0f) || (Dt < UE_SMALL_NUMBER))
 	{
 		return;
 	}
 
+	// If the simulation is in world space, the simulation space is stationary
 	if (Space == ESimulationSpace::WorldSpace)
 	{
-		OutSimulationSpaceData.LinearVel = Settings.ExternalLinearVelocity;
-		OutSimulationSpaceData.AngularVel = Settings.ExternalAngularVelocity;
+		SpaceLinearVel = Settings.ExternalLinearVelocity;
+		SpaceAngularVel = Settings.ExternalAngularVelocity;
 		return;
 	}
 
-	// World-space component velocity and acceleration
-	FVector CompLinVel = Chaos::FVec3::CalculateVelocity(SimSpacePreviousComponentToWorld.GetTranslation(), ComponentToWorld.GetTranslation(), Dt);
-	FVector CompAngVel = Chaos::FRotation3::CalculateAngularVelocity(SimSpacePreviousComponentToWorld.GetRotation(), ComponentToWorld.GetRotation(), Dt);
-	FVector CompLinAcc = (CompLinVel - SimSpacePreviousComponentLinearVelocity) / Dt;
-	FVector CompAngAcc = (CompAngVel - SimSpacePreviousComponentAngularVelocity) / Dt;
-	SimSpacePreviousComponentToWorld = ComponentToWorld;
-	SimSpacePreviousComponentLinearVelocity = CompLinVel;
-	SimSpacePreviousComponentAngularVelocity = CompAngVel;
+	// World-space component linear velocity and acceleration
+	FVector PrevSpaceLinearVel = Chaos::FVec3::CalculateVelocity(PreviousPreviousSimulationSpaceTransform.GetTranslation(), PreviousSimulationSpaceTransform.GetTranslation(), Dt);
+	SpaceLinearVel = Chaos::FVec3::CalculateVelocity(PreviousSimulationSpaceTransform.GetTranslation(), SimulationSpaceTransform.GetTranslation(), Dt);
+	SpaceLinearAcc = (SpaceLinearVel - PrevSpaceLinearVel) / Dt;
 
-	if (Space == ESimulationSpace::ComponentSpace)
-	{
-		CompLinVel.Z *= Settings.VelocityScaleZ;
-		CompLinAcc.Z *= Settings.VelocityScaleZ;
+	// World-space component angular velocity and acceleration
+	FVector PrevSpaceAngularVel = Chaos::FRotation3::CalculateAngularVelocity(PreviousPreviousSimulationSpaceTransform.GetRotation(), PreviousSimulationSpaceTransform.GetRotation(), Dt);
+	SpaceAngularVel = Chaos::FRotation3::CalculateAngularVelocity(PreviousSimulationSpaceTransform.GetRotation(), SimulationSpaceTransform.GetRotation(), Dt);
+	SpaceAngularAcc = (SpaceAngularVel - PrevSpaceAngularVel) / Dt;
 
-		OutSimulationSpaceData.LinearVel = CompLinVel.GetClampedToMaxSize(Settings.MaxLinearVelocity) + Settings.ExternalLinearVelocity;
-		OutSimulationSpaceData.AngularVel = CompAngVel.GetClampedToMaxSize(Settings.MaxAngularVelocity) + Settings.ExternalAngularVelocity;
-		OutSimulationSpaceData.LinearAcc = CompLinAcc.GetClampedToMaxSize(Settings.MaxLinearAcceleration);
-		OutSimulationSpaceData.AngularAcc = CompAngAcc.GetClampedToMaxSize(Settings.MaxAngularAcceleration);
-		return;
-	}
+	// Apply Z scale
+	SpaceLinearVel.Z *= Settings.VelocityScaleZ;
+	SpaceLinearAcc.Z *= Settings.VelocityScaleZ;
 	
-	if (Space == ESimulationSpace::BaseBoneSpace)
-	{
-		// World-space component-relative bone velocity and acceleration
-		FVector BoneLinVel = Chaos::FVec3::CalculateVelocity(SimSpacePreviousBoneToComponent.GetTranslation(), BoneToComponent.GetTranslation(), Dt);
-		FVector BoneAngVel = Chaos::FRotation3::CalculateAngularVelocity(SimSpacePreviousBoneToComponent.GetRotation(), BoneToComponent.GetRotation(), Dt);
-		BoneLinVel = ComponentToWorld.TransformVector(BoneLinVel);
-		BoneAngVel = ComponentToWorld.TransformVector(BoneAngVel);
-		FVector BoneLinAcc = (BoneLinVel - SimSpacePreviousBoneLinearVelocity) / Dt;
-		FVector BoneAngAcc = (BoneAngVel - SimSpacePreviousBoneAngularVelocity) / Dt;
-		SimSpacePreviousBoneToComponent = BoneToComponent;
-		SimSpacePreviousBoneLinearVelocity = BoneLinVel;
-		SimSpacePreviousBoneAngularVelocity = BoneAngVel;
-
-		// World-space bone velocity and acceleration
-		FVector NetAngVel = CompAngVel + BoneAngVel;
-		FVector NetAngAcc = CompAngAcc + BoneAngAcc;
-
-		// If we limit the angular velocity, we also need to limit the component of linear velocity that comes from (angvel x offset)
-		float AngVelScale = 1.0f;
-		float NetAngVelLenSq = (float) NetAngVel.SizeSquared();
-		if (NetAngVelLenSq > FMath::Square(Settings.MaxAngularVelocity))
-		{
-			AngVelScale = Settings.MaxAngularVelocity * FMath::InvSqrt(NetAngVelLenSq);
-		}
-
-		// Add the linear velocity and acceleration that comes from rotation of the space about the component
-		// NOTE: Component angular velocity constribution is scaled
-		FVector SpaceCompOffset = ComponentToWorld.TransformVector(BoneToComponent.GetTranslation());
-		FVector NetLinVel = CompLinVel + BoneLinVel + FVector::CrossProduct(AngVelScale * CompAngVel, SpaceCompOffset);
-		FVector NetLinAcc = CompLinAcc + BoneLinAcc + FVector::CrossProduct(AngVelScale * CompAngAcc, SpaceCompOffset);
-
-		NetLinVel.Z *= Settings.VelocityScaleZ;
-		NetLinAcc.Z *= Settings.VelocityScaleZ;
-
-		OutSimulationSpaceData.LinearVel = NetLinVel.GetClampedToMaxSize(Settings.MaxLinearVelocity) + Settings.ExternalLinearVelocity;
-		OutSimulationSpaceData.AngularVel = NetAngVel.GetClampedToMaxSize(Settings.MaxAngularVelocity) + Settings.ExternalAngularVelocity;
-		OutSimulationSpaceData.LinearAcc = NetLinAcc.GetClampedToMaxSize(Settings.MaxLinearAcceleration);
-		OutSimulationSpaceData.AngularAcc = NetAngAcc.GetClampedToMaxSize(Settings.MaxAngularAcceleration);
-		return;
-	}
+	// Clamped world-space motion of the simulation space
+	SpaceLinearVel = SpaceLinearVel.GetClampedToMaxSize(Settings.MaxLinearVelocity) + Settings.ExternalLinearVelocity;
+	SpaceAngularVel = SpaceAngularVel.GetClampedToMaxSize(Settings.MaxAngularVelocity) + Settings.ExternalAngularVelocity;
+	SpaceLinearAcc = SpaceLinearAcc.GetClampedToMaxSize(Settings.MaxLinearAcceleration);
+	SpaceAngularAcc = SpaceAngularAcc.GetClampedToMaxSize(Settings.MaxAngularAcceleration);
+
+	// Transform world-space motion into simulation space
+	SpaceLinearVel = SimulationSpaceTransform.InverseTransformVector(SpaceLinearVel);
+	SpaceAngularVel = SimulationSpaceTransform.InverseTransformVectorNoScale(SpaceAngularVel);
+	SpaceLinearAcc = SimulationSpaceTransform.InverseTransformVector(SpaceLinearAcc);
+	SpaceAngularAcc = SimulationSpaceTransform.InverseTransformVectorNoScale(SpaceAngularAcc);
+
+	PreviousPreviousSimulationSpaceTransform = PreviousSimulationSpaceTransform;
+	PreviousSimulationSpaceTransform = SimulationSpaceTransform;
 }
 
+
 bool FAnimNode_RigidBodyWithControl::ShouldApplyPhysicsAssetConditioning() const
 {
 	return (PhysicsAssetConditioningConstraintPosition != MapConstraintsBehaviorType::None)
@@ -818,17 +775,28 @@ void FAnimNode_RigidBodyWithControl::EvaluateSkeletalControl_AnyThread(FComponen
 				SimSpaceSettingsToUse = &RBANWithControl_SimSpaceOverride;
 			}
 
-			FSimulationSpaceData SimulationSpaceData;
-			CalculateSimulationSpace(
-				SimulationSpace, CompWorldSpaceTM, BaseBoneTM, DeltaSeconds, 
-				*SimSpaceSettingsToUse, SimulationSpaceData);
-
-			UpdateWorldObjects(SimulationSpaceData, DeltaSeconds);
-			UpdateClothColliderObjects(SimulationSpaceData);
+			const FTransform SpaceTransform = SpaceToWorldTransform(SimulationSpace, CompWorldSpaceTM, BaseBoneTM);
+			FVector SimulationSpaceLinearVelocity;
+			FVector SimulationSpaceAngularVelocity;
+			FVector SimulationSpaceLinearAcceleration;
+			FVector SimulationSpaceAngularAcceleration;
+			CalculateSimulationSpaceMotion(
+				SimulationSpace,
+				SpaceTransform,
+				CompWorldSpaceTM,
+				DeltaSeconds,
+				*SimSpaceSettingsToUse,
+				SimulationSpaceLinearVelocity,
+				SimulationSpaceAngularVelocity,
+				SimulationSpaceLinearAcceleration,
+				SimulationSpaceAngularAcceleration);
+
+			UpdateWorldObjects(SpaceTransform, SimulationSpaceLinearVelocity, SimulationSpaceAngularVelocity, DeltaSeconds);
+			UpdateClothColliderObjects(SpaceTransform);
 
 			PhysicsSimulation->UpdateSimulationSpace(
-				SimulationSpaceData.Transform, SimulationSpaceData.LinearVel, SimulationSpaceData.AngularVel,
-				SimulationSpaceData.LinearAcc, SimulationSpaceData.AngularAcc);
+				SpaceTransform, SimulationSpaceLinearVelocity, SimulationSpaceAngularVelocity,
+				SimulationSpaceLinearAcceleration, SimulationSpaceAngularAcceleration);
 
 			PhysicsSimulation->SetSimulationSpaceSettings(
 				SimSpaceSettingsToUse->WorldAlpha, SimSpaceSettingsToUse->ExternalLinearDragV);
@@ -1618,7 +1586,7 @@ void FAnimNode_RigidBodyWithControl::RemoveClothColliderObjects()
 	ClothColliders.Reset();
 }
 
-void FAnimNode_RigidBodyWithControl::UpdateClothColliderObjects(const FSimulationSpaceData& SimulationSpaceData)
+void FAnimNode_RigidBodyWithControl::UpdateClothColliderObjects(const FTransform& SpaceTransform)
 {
 	for (FClothCollider& ClothCollider : ClothColliders)
 	{
@@ -1627,9 +1595,9 @@ void FAnimNode_RigidBodyWithControl::UpdateClothColliderObjects(const FSimulatio
 			// Calculate the sim-space transform of this object
 			const FTransform CompWorldTransform = ClothCollider.SkeletalMeshComponent->GetBoneTransform(ClothCollider.BoneIndex);
 			FTransform CompSpaceTransform;
-			CompSpaceTransform.SetTranslation(SimulationSpaceData.Transform.InverseTransformPosition(
+			CompSpaceTransform.SetTranslation(SpaceTransform.InverseTransformPosition(
 				CompWorldTransform.GetLocation()));
-			CompSpaceTransform.SetRotation(SimulationSpaceData.Transform.InverseTransformRotation(
+			CompSpaceTransform.SetRotation(SpaceTransform.InverseTransformRotation(
 				CompWorldTransform.GetRotation()));
 			CompSpaceTransform.SetScale3D(FVector::OneVector);	// TODO - sort out scale for world objects in local sim
 
@@ -1758,7 +1726,7 @@ void FAnimNode_RigidBodyWithControl::PurgeExpiredWorldObjects()
 // component-based simulation and need to update stationary objects into our space, or simply
 // because the objects are moving.
 void FAnimNode_RigidBodyWithControl::UpdateWorldObjects(
-	const FSimulationSpaceData& SimulationSpaceData, const float DeltaSeconds)
+	const FTransform& SpaceTransform, const FVector& LinearVel, const FVector& AngularVel, const float DeltaSeconds)
 {
 	SCOPE_CYCLE_COUNTER(STAT_RigidBodyNodeWithControl_UpdateWorldObjects);
 	LLM_SCOPE_BYNAME(TEXT("Animation/RigidBodyWithControl"));
@@ -1773,10 +1741,10 @@ void FAnimNode_RigidBodyWithControl::UpdateWorldObjects(
 		// We will want the "previous" simulation space to calculate velocities. Get it from the
 		// velocity, so we can take advantage of that being zero at initialisation/reset times etc.
 		const FVector PrevSimSpacePosition = 
-			SimulationSpaceData.Transform.GetLocation() - SimulationSpaceData.LinearVel * DeltaSeconds;
-		const FQuat SimSpaceDeltaOrientation = FQuat::MakeFromRotationVector(SimulationSpaceData.AngularVel * DeltaSeconds);
+			SpaceTransform.GetLocation() - LinearVel * DeltaSeconds;
+		const FQuat SimSpaceDeltaOrientation = FQuat::MakeFromRotationVector(AngularVel * DeltaSeconds);
 		const FQuat PrevSimSpaceOrientation = 
-			SimSpaceDeltaOrientation.Inverse() * SimulationSpaceData.Transform.GetRotation();
+			SimSpaceDeltaOrientation.Inverse() * SpaceTransform.GetRotation();
 		const FTransform PrevSimSpaceTransform(PrevSimSpaceOrientation, PrevSimSpacePosition, FVector::OneVector);
 
 		for (const auto& WorldEntry : ComponentsInSim)
@@ -1791,8 +1759,8 @@ void FAnimNode_RigidBodyWithControl::UpdateWorldObjects(
 				const FVector WorldPosition = CompWorldTransform.GetLocation();
 				const FQuat WorldOrientation = CompWorldTransform.GetRotation();
 
-				const FVector CompSpacePosition = SimulationSpaceData.Transform.InverseTransformPosition(WorldPosition);
-				const FQuat CompSpaceOrientation = SimulationSpaceData.Transform.InverseTransformRotation(WorldOrientation);
+				const FVector CompSpacePosition = SpaceTransform.InverseTransformPosition(WorldPosition);
+				const FQuat CompSpaceOrientation = SpaceTransform.InverseTransformRotation(WorldOrientation);
 
 				// TODO - sort out scale for world objects in local sim
 				const FTransform CompSpaceTransform(CompSpaceOrientation, CompSpacePosition, FVector::OneVector);
@@ -1838,8 +1806,8 @@ void FAnimNode_RigidBodyWithControl::UpdateWorldObjects(
 				const FVector WorldPosition = CompWorldTransform.GetLocation();
 				const FQuat WorldOrientation = CompWorldTransform.GetRotation();
 
-				const FVector CompSpacePosition = SimulationSpaceData.Transform.InverseTransformPosition(WorldPosition);
-				const FQuat CompSpaceOrientation = SimulationSpaceData.Transform.InverseTransformRotation(WorldOrientation);
+				const FVector CompSpacePosition = SpaceTransform.InverseTransformPosition(WorldPosition);
+				const FQuat CompSpaceOrientation = SpaceTransform.InverseTransformRotation(WorldOrientation);
 
 				// TODO - sort out scale for world objects in local sim
 				const FTransform CompSpaceTransform(CompSpaceOrientation, CompSpacePosition, FVector::OneVector);
diff --git a/Engine/Plugins/Experimental/PhysicsControl/Source/PhysicsControl/Public/AnimNode_RigidBodyWithControl.h b/Engine/Plugins/Experimental/PhysicsControl/Source/PhysicsControl/Public/AnimNode_RigidBodyWithControl.h
--- a/Engine/Plugins/Experimental/PhysicsControl/Source/PhysicsControl/Public/AnimNode_RigidBodyWithControl.h
+++ b/Engine/Plugins/Experimental/PhysicsControl/Source/PhysicsControl/Public/AnimNode_RigidBodyWithControl.h
@@ -486,14 +486,16 @@ private:
 	void InitSimulationSpace(const FTransform& ComponentToWorld, const FTransform& BoneToComponent);
 
 	// Calculate simulation space transform, velocity etc to pass into the solver
-	void CalculateSimulationSpace(
-		ESimulationSpace         Space,
-		const FTransform&        ComponentToWorld,
-		const FTransform&        BoneToComponent,
-		const float              Dt,
+	void CalculateSimulationSpaceMotion(
+		ESimulationSpace Space,
+		const FTransform& SpaceTransform,
+		const FTransform& ComponentToWorld,
+		const float Dt,
 		const FSimSpaceSettings& Settings,
-		FSimulationSpaceData&    OutSimulationSpaceData
-	);
+		FVector& SpaceLinearVel,
+		FVector& SpaceAngularVel,
+		FVector& SpaceLinearAcc,
+		FVector& SpaceAngularAcc);
 
 	// Returns true if constraint transforms in the physics asset should be modified to match the current skeleton.
 	bool ShouldApplyPhysicsAssetConditioning() const;
@@ -514,7 +516,7 @@ private:
 	void RemoveClothColliderObjects();
 
 	// Update the sim-space transforms of all cloth collider objects.
-	void UpdateClothColliderObjects(const FSimulationSpaceData& SimulationSpaceData);
+	void UpdateClothColliderObjects(const FTransform& SpaceTransform);
 
 	// Gather nearby world objects and add them to the sim
 	void CollectWorldObjects();
@@ -526,7 +528,7 @@ private:
 	void PurgeExpiredWorldObjects();
 
 	// Update sim-space transforms of world objects
-	void UpdateWorldObjects(const FSimulationSpaceData& SimulationSpaceData, const float DeltaSeconds);
+	void UpdateWorldObjects(const FTransform& SpaceTransform, const FVector& LinearVel, const FVector& AngularVel, const float DeltaSeconds);
 
 	// Advances the simulation by a given timestep
 	void RunPhysicsSimulation(float DeltaSeconds, const FVector& SimSpaceGravity);
@@ -679,12 +681,8 @@ private:
 	FVector PreviousComponentLinearVelocity;
 
 	// Used by the world-space to simulation-space motion transfer system in Component- or Bone-Space sims
-	FTransform SimSpacePreviousComponentToWorld;
-	FTransform SimSpacePreviousBoneToComponent;
-	FVector SimSpacePreviousComponentLinearVelocity;
-	FVector SimSpacePreviousComponentAngularVelocity;
-	FVector SimSpacePreviousBoneLinearVelocity;
-	FVector SimSpacePreviousBoneAngularVelocity;
+	FTransform PreviousSimulationSpaceTransform;
+	FTransform PreviousPreviousSimulationSpaceTransform;
 };
 
 #if WITH_EDITORONLY_DATA
