Adds support for physics in Control Rig, with the implementation in the experimental Physics Control plugin

diff --git a/Engine/Plugins/Experimental/PhysicsControl/PhysicsControl.uplugin b/Engine/Plugins/Experimental/PhysicsControl/PhysicsControl.uplugin
--- a/Engine/Plugins/Experimental/PhysicsControl/PhysicsControl.uplugin
+++ b/Engine/Plugins/Experimental/PhysicsControl/PhysicsControl.uplugin
@@ -15,7 +15,18 @@
 	"IsBetaVersion": false,
 	"IsExperimentalVersion": true,
 	"Installed": false,
+	"Plugins": [
+		{
+			"Name": "ControlRig",
+			"Enabled": true
+		},
+	],
 	"Modules": [
+		{
+			"Name": "ControlRigPhysics",
+			"Type": "Runtime",
+			"LoadingPhase": "Default"
+		},
 		{
 			"Name": "PhysicsControl",
 			"Type": "Runtime",
diff --git a/Engine/Plugins/Experimental/PhysicsControl/Source/ControlRigPhysics/ControlRigPhysics.Build.cs b/Engine/Plugins/Experimental/PhysicsControl/Source/ControlRigPhysics/ControlRigPhysics.Build.cs
new file mode 100644
--- /dev/null
+++ b/Engine/Plugins/Experimental/PhysicsControl/Source/ControlRigPhysics/ControlRigPhysics.Build.cs
@@ -0,0 +1,33 @@
+// Copyright Epic Games, Inc. All Rights Reserved.
+
+using UnrealBuildTool;
+
+public class ControlRigPhysics : ModuleRules
+{
+	public ControlRigPhysics(ReadOnlyTargetRules Target) : base(Target)
+	{
+		PCHUsage = ModuleRules.PCHUsageMode.UseExplicitOrSharedPCHs;
+
+		PublicDependencyModuleNames.AddRange(
+			new string[]
+			{
+				"Core",
+				"PhysicsControl"
+			}
+		);
+		
+		PrivateDependencyModuleNames.AddRange(
+			new string[]
+			{
+				"ControlRig",
+				"Chaos",
+				"ChaosCore",
+				"Core",
+				"CoreUObject",
+				"Engine",
+				"PhysicsCore",
+				"RigVM"
+			}
+		);
+	}
+}
diff --git a/Engine/Plugins/Experimental/PhysicsControl/Source/ControlRigPhysics/Private/ControlRigPhysicsModule.cpp b/Engine/Plugins/Experimental/PhysicsControl/Source/ControlRigPhysics/Private/ControlRigPhysicsModule.cpp
new file mode 100644
--- /dev/null
+++ b/Engine/Plugins/Experimental/PhysicsControl/Source/ControlRigPhysics/Private/ControlRigPhysicsModule.cpp
@@ -0,0 +1,18 @@
+// Copyright Epic Games, Inc. All Rights Reserved.
+
+#include "ControlRigPhysicsModule.h"
+
+#define LOCTEXT_NAMESPACE "FControlRigPhysicsModule"
+
+void FControlRigPhysicsModule::StartupModule()
+{
+}
+
+void FControlRigPhysicsModule::ShutdownModule()
+{
+}
+
+#undef LOCTEXT_NAMESPACE
+	
+IMPLEMENT_MODULE(FControlRigPhysicsModule, ControlRigPhysics)
+
diff --git a/Engine/Plugins/Experimental/PhysicsControl/Source/ControlRigPhysics/Private/RigPhysicsComponent.cpp b/Engine/Plugins/Experimental/PhysicsControl/Source/ControlRigPhysics/Private/RigPhysicsComponent.cpp
new file mode 100644
--- /dev/null
+++ b/Engine/Plugins/Experimental/PhysicsControl/Source/ControlRigPhysics/Private/RigPhysicsComponent.cpp
@@ -0,0 +1,27 @@
+// Copyright Epic Games, Inc. All Rights Reserved.
+
+#include "RigPhysicsComponent.h"
+
+////////////////////////////////////////////////////////////////////////////////
+// FRigPhysicsComponent
+////////////////////////////////////////////////////////////////////////////////
+
+void FRigPhysicsComponent::Save(FArchive& Ar)
+{
+	FRigBaseComponent::Save(Ar);
+
+	Ar << BodySolverSettings;
+	Ar << Dynamics;
+	Ar << Collision;
+	Ar << Joint;
+}
+
+void FRigPhysicsComponent::Load(FArchive& Ar)
+{
+	FRigBaseComponent::Load(Ar);
+	Ar << BodySolverSettings;
+	Ar << Dynamics;
+	Ar << Collision;
+	Ar << Joint;
+}
+
diff --git a/Engine/Plugins/Experimental/PhysicsControl/Source/ControlRigPhysics/Private/RigPhysicsControlComponent.cpp b/Engine/Plugins/Experimental/PhysicsControl/Source/ControlRigPhysics/Private/RigPhysicsControlComponent.cpp
new file mode 100644
--- /dev/null
+++ b/Engine/Plugins/Experimental/PhysicsControl/Source/ControlRigPhysics/Private/RigPhysicsControlComponent.cpp
@@ -0,0 +1,26 @@
+// Copyright Epic Games, Inc. All Rights Reserved.
+
+#include "RigPhysicsControlComponent.h"
+
+////////////////////////////////////////////////////////////////////////////////
+// FRigPhysicsControlComponent
+////////////////////////////////////////////////////////////////////////////////
+
+void FRigPhysicsControlComponent::Save(FArchive& Ar)
+{
+	FRigBaseComponent::Save(Ar);
+
+	Ar << ParentBody;
+	Ar << ControlData;
+	Ar << ControlTarget;
+}
+
+void FRigPhysicsControlComponent::Load(FArchive& Ar)
+{
+	FRigBaseComponent::Load(Ar);
+
+	Ar << ParentBody;
+	Ar << ControlData;
+	Ar << ControlTarget;
+}
+
diff --git a/Engine/Plugins/Experimental/PhysicsControl/Source/ControlRigPhysics/Private/RigPhysicsControlExecution.cpp b/Engine/Plugins/Experimental/PhysicsControl/Source/ControlRigPhysics/Private/RigPhysicsControlExecution.cpp
new file mode 100644
--- /dev/null
+++ b/Engine/Plugins/Experimental/PhysicsControl/Source/ControlRigPhysics/Private/RigPhysicsControlExecution.cpp
@@ -0,0 +1,67 @@
+// Copyright Epic Games, Inc. All Rights Reserved.
+
+#include "RigPhysicsControlExecution.h"
+#include "RigPhysicsSolver.h"
+#include "RigPhysicsControlComponent.h"
+
+#include "Rigs/RigHierarchyController.h"
+#include "Units/RigUnitContext.h"
+#include "Units/Execution/RigUnit_PrepareForExecution.h"
+
+#include "ControlRig.h"
+
+#include UE_INLINE_GENERATED_CPP_BY_NAME(RigPhysicsControlExecution)
+
+//======================================================================================================================
+FRigUnit_AddPhysicsControl_Execute()
+{
+	Result.Reset();
+
+	if (ExecuteContext.GetEventName() == FRigUnit_PrepareForExecution::EventName)
+	{
+		if (URigHierarchy* Hierarchy = ExecuteContext.Hierarchy)
+		{
+			if (UControlRig* ControlRig = Cast<UControlRig>(Hierarchy->GetOuter()))
+			{
+				if (ControlRig->GetPhysicsSolver(SolverID) == nullptr)
+				{
+					UE_CONTROLRIG_RIGUNIT_REPORT_ERROR(
+						TEXT("Physics Solver with guid '%s' cannot be found."), *SolverID.ToString());
+					return;
+				}
+
+				if (URigHierarchyController* Controller = Hierarchy->GetController())
+				{
+					FRigHierarchyControllerInstructionBracket InstructionBracket(
+						Controller, ExecuteContext.GetInstructionIndex());
+					Result = Controller->AddComponent(FRigPhysicsControlComponent::StaticStruct(), Name, Owner);
+					if (Result.IsValid())
+					{
+						if (FRigPhysicsControlComponent* Component = Cast<FRigPhysicsControlComponent>(
+							Hierarchy->FindComponent(Result)))
+						{
+							Component->ParentBody = ParentBody;
+							Component->ControlData = ControlData;
+							Component->ControlTarget = ControlTarget;
+
+							if (FRigPhysicsSolver* PhysicsSolver = Cast<FRigPhysicsSolver>(
+								ControlRig->GetPhysicsSolver(SolverID)))
+							{
+								// Danny TODO
+								//PhysicsSolver->AddControl(Result, Set);
+							}
+						}
+					}
+				}
+			}
+		}
+	}
+	else
+	{
+		UE_CONTROLRIG_RIGUNIT_REPORT_ERROR(TEXT("AddPhysicsControl can only be used during Setup"));
+	}
+
+
+}
+
+
diff --git a/Engine/Plugins/Experimental/PhysicsControl/Source/ControlRigPhysics/Private/RigPhysicsExecution.cpp b/Engine/Plugins/Experimental/PhysicsControl/Source/ControlRigPhysics/Private/RigPhysicsExecution.cpp
new file mode 100644
--- /dev/null
+++ b/Engine/Plugins/Experimental/PhysicsControl/Source/ControlRigPhysics/Private/RigPhysicsExecution.cpp
@@ -0,0 +1,338 @@
+// Copyright Epic Games, Inc. All Rights Reserved.
+
+#include "RigPhysicsExecution.h"
+#include "RigPhysicsSolver.h"
+#include "RigPhysicsControlComponent.h"
+
+#include "Rigs/RigHierarchyController.h"
+#include "Units/RigUnitContext.h"
+#include "Units/Execution/RigUnit_PrepareForExecution.h"
+
+#include "ControlRig.h"
+
+#include UE_INLINE_GENERATED_CPP_BY_NAME(RigPhysicsExecution)
+
+//======================================================================================================================
+FRigUnit_AddPhysicsSolver_Execute()
+{
+	FString ErrorMessage;
+	if(!FRigUnit_DynamicHierarchyBase::IsValidToRunInContext(ExecuteContext, true, &ErrorMessage))
+	{
+		if(!ErrorMessage.IsEmpty())
+		{
+			UE_CONTROLRIG_RIGUNIT_REPORT_ERROR(TEXT("%s"), *ErrorMessage);
+		}
+		return;
+	}
+
+	SolverID = FRigPhysicsSolverID();
+
+	if (!ExecuteContext.Hierarchy)
+	{
+		return;
+	}
+
+	if(UControlRig* ControlRig = Cast<UControlRig>(ExecuteContext.Hierarchy->GetOuter()))
+	{
+		TSharedPtr<FRigPhysicsSolver> PhysicsSolver = MakeShared<FRigPhysicsSolver>(
+			ControlRig, SolverSettings, SimulationSpaceSettings);
+		SolverID = ControlRig->RegisterPhysicsSolver(PhysicsSolver, Name);
+	}
+}
+
+//======================================================================================================================
+FRigUnit_StepPhysicsSolver_Execute()
+{
+	DECLARE_SCOPE_HIERARCHICAL_COUNTER_RIGUNIT();
+
+	if (ExecuteContext.Hierarchy)
+	{
+		if (UControlRig* ControlRig = Cast<UControlRig>(ExecuteContext.Hierarchy->GetOuter()))
+		{
+			if (FRigPhysicsSolver* PhysicsSolver = Cast<FRigPhysicsSolver>(ControlRig->GetPhysicsSolver(SolverID)))
+			{
+				double DeltaTime = DeltaTimeOverride > 0 ? DeltaTimeOverride : ExecuteContext.GetDeltaTime();
+				PhysicsSolver->StepSimulation(ExecuteContext, DeltaTime);
+				if (bEnableVisualization)
+				{
+					PhysicsSolver->Draw(ExecuteContext.GetDrawInterface(), ExecuteContext.GetWorld());
+				}
+			}
+		}
+	}
+}
+
+//======================================================================================================================
+FRigUnit_InstantiatePhysics_Execute()
+{
+	DECLARE_SCOPE_HIERARCHICAL_COUNTER_RIGUNIT();
+
+	if (ExecuteContext.Hierarchy)
+	{
+		if (UControlRig* ControlRig = Cast<UControlRig>(ExecuteContext.Hierarchy->GetOuter()))
+		{
+			if (FRigPhysicsSolver* PhysicsSolver = Cast<FRigPhysicsSolver>(ControlRig->GetPhysicsSolver(SolverID)))
+			{
+				PhysicsSolver->Instantiate();
+			}
+		}
+	}
+}
+
+//======================================================================================================================
+FRigUnit_AddPhysicsBody_Execute()
+{
+	Result.Reset();
+
+	if (!SolverID.IsValid())
+	{
+		UE_CONTROLRIG_RIGUNIT_REPORT_ERROR(TEXT("Physics Solver Guid is not valid"));
+		return;
+	}
+
+	if (ExecuteContext.GetEventName() == FRigUnit_PrepareForExecution::EventName)
+	{
+		if (URigHierarchy* Hierarchy = ExecuteContext.Hierarchy)
+		{
+			if (UControlRig* ControlRig = Cast<UControlRig>(Hierarchy->GetOuter()))
+			{
+				if (ControlRig->GetPhysicsSolver(SolverID) == nullptr)
+				{
+					UE_CONTROLRIG_RIGUNIT_REPORT_ERROR(
+						TEXT("Physics Solver with guid '%s' cannot be found."), *SolverID.ToString());
+					return;
+				}
+
+				if (URigHierarchyController* Controller = Hierarchy->GetController())
+				{
+					FRigHierarchyControllerInstructionBracket InstructionBracket(
+						Controller, ExecuteContext.GetInstructionIndex());
+					Result = Controller->AddComponent(FRigPhysicsComponent::StaticStruct(), TEXT("PhysicsBody"), Owner);
+					if (Result.IsValid())
+					{
+						if (FRigPhysicsComponent* Component = Cast<FRigPhysicsComponent>(
+							Hierarchy->FindComponent(Result)))
+						{
+							Component->BodySolverSettings.SolverID = SolverID;
+							Component->Dynamics = Dynamics;
+							Component->Collision = Collision;
+							Component->Joint = Joint;
+							Component->BodyData = BodyData;
+
+							if (FRigPhysicsSolver* PhysicsSolver = Cast<FRigPhysicsSolver>(
+								ControlRig->GetPhysicsSolver(SolverID)))
+							{
+								PhysicsSolver->AddBody(Result);
+							}
+						}
+					}
+				}
+			}
+		}
+	}
+	else
+	{
+		UE_CONTROLRIG_RIGUNIT_REPORT_ERROR(TEXT("AddPhysicsBody can only be used during Setup"));
+	}
+}
+
+//======================================================================================================================
+FRigUnit_HierarchySetPhysicsBodySourceBone_Execute()
+{
+	DECLARE_SCOPE_HIERARCHICAL_COUNTER_RIGUNIT();
+
+	if (ExecuteContext.Hierarchy)
+	{
+		if (FRigPhysicsComponent* PhysicsComponent = Cast<FRigPhysicsComponent>(
+			ExecuteContext.Hierarchy->FindComponent(PhysicsBody)))
+		{
+			PhysicsComponent->BodySolverSettings.SourceBone = SourceBone;
+		}
+	}
+}
+
+//======================================================================================================================
+FRigUnit_HierarchySetPhysicsBodyTargetBone_Execute()
+{
+	DECLARE_SCOPE_HIERARCHICAL_COUNTER_RIGUNIT();
+
+	if (ExecuteContext.Hierarchy)
+	{
+		if (FRigPhysicsComponent* PhysicsComponent = Cast<FRigPhysicsComponent>(
+			ExecuteContext.Hierarchy->FindComponent(PhysicsBody)))
+		{
+			PhysicsComponent->BodySolverSettings.TargetBone = TargetBone;
+		}
+	}
+}
+
+//======================================================================================================================
+FRigUnit_HierarchySetPhysicsBodyKinematicTarget_Execute()
+{
+	DECLARE_SCOPE_HIERARCHICAL_COUNTER_RIGUNIT();
+
+	if (ExecuteContext.Hierarchy)
+	{
+		if (FRigPhysicsComponent* PhysicsComponent = Cast<FRigPhysicsComponent>(
+			ExecuteContext.Hierarchy->FindComponent(PhysicsBody)))
+		{
+			if (UControlRig* ControlRig = Cast<UControlRig>(ExecuteContext.Hierarchy->GetOuter()))
+			{
+				if (FRigPhysicsSolver* PhysicsSolver = Cast<FRigPhysicsSolver>(
+					ControlRig->GetPhysicsSolver((PhysicsComponent->BodySolverSettings.SolverID))))
+				{
+					PhysicsSolver->SetBodyKinematicTarget(PhysicsBody, KinematicTargetSpace, KinematicTarget);
+				}
+			}
+		}
+	}
+}
+
+//======================================================================================================================
+FRigUnit_HierarchySetPhysicsBodySparseData_Execute()
+{
+	DECLARE_SCOPE_HIERARCHICAL_COUNTER_RIGUNIT();
+
+	if (ExecuteContext.Hierarchy)
+	{
+		if (FRigPhysicsComponent* PhysicsComponent = Cast<FRigPhysicsComponent>(
+			ExecuteContext.Hierarchy->FindComponent(PhysicsBody)))
+		{
+			if (UControlRig* ControlRig = Cast<UControlRig>(ExecuteContext.Hierarchy->GetOuter()))
+			{
+				if (FRigPhysicsSolver* PhysicsSolver = Cast<FRigPhysicsSolver>(
+					ControlRig->GetPhysicsSolver((PhysicsComponent->BodySolverSettings.SolverID))))
+				{
+					PhysicsSolver->SetBodySparseData(PhysicsBody, Data);
+				}
+			}
+		}
+	}
+}
+
+//======================================================================================================================
+FRigUnit_HierarchySetPhysicsBodyMovementType_Execute()
+{
+	DECLARE_SCOPE_HIERARCHICAL_COUNTER_RIGUNIT();
+	if (ExecuteContext.Hierarchy)
+	{
+		if (FRigPhysicsComponent* PhysicsComponent = Cast<FRigPhysicsComponent>(
+			ExecuteContext.Hierarchy->FindComponent(PhysicsBody)))
+		{
+			if (UControlRig* ControlRig = Cast<UControlRig>(ExecuteContext.Hierarchy->GetOuter()))
+			{
+				if (FRigPhysicsSolver* PhysicsSolver = Cast<FRigPhysicsSolver>(
+					ControlRig->GetPhysicsSolver((PhysicsComponent->BodySolverSettings.SolverID))))
+				{
+					PhysicsSolver->SetBodyMovementType(PhysicsBody, MovementType);
+				}
+			}
+		}
+	}
+}
+
+//======================================================================================================================
+FRigUnit_HierarchySetPhysicsBodyCollisionType_Execute()
+{
+	DECLARE_SCOPE_HIERARCHICAL_COUNTER_RIGUNIT();
+	if (ExecuteContext.Hierarchy)
+	{
+		if (FRigPhysicsComponent* PhysicsComponent = Cast<FRigPhysicsComponent>(
+			ExecuteContext.Hierarchy->FindComponent(PhysicsBody)))
+		{
+			if (UControlRig* ControlRig = Cast<UControlRig>(ExecuteContext.Hierarchy->GetOuter()))
+			{
+				if (FRigPhysicsSolver* PhysicsSolver = Cast<FRigPhysicsSolver>(
+					ControlRig->GetPhysicsSolver((PhysicsComponent->BodySolverSettings.SolverID))))
+				{
+					PhysicsSolver->SetBodyCollisionType(PhysicsBody, CollisionType);
+				}
+			}
+		}
+	}
+}
+
+//======================================================================================================================
+FRigUnit_HierarchySetPhysicsBodyGravityMultiplier_Execute()
+{
+	DECLARE_SCOPE_HIERARCHICAL_COUNTER_RIGUNIT();
+	if (ExecuteContext.Hierarchy)
+	{
+		if (FRigPhysicsComponent* PhysicsComponent = Cast<FRigPhysicsComponent>(
+			ExecuteContext.Hierarchy->FindComponent(PhysicsBody)))
+		{
+			if (UControlRig* ControlRig = Cast<UControlRig>(ExecuteContext.Hierarchy->GetOuter()))
+			{
+				if (FRigPhysicsSolver* PhysicsSolver = Cast<FRigPhysicsSolver>(
+					ControlRig->GetPhysicsSolver((PhysicsComponent->BodySolverSettings.SolverID))))
+				{
+					PhysicsSolver->SetBodyGravityMultiplier(PhysicsBody, GravityMultiplier);
+				}
+			}
+		}
+	}
+}
+
+//======================================================================================================================
+FRigUnit_HierarchySetPhysicsBodyPhysicsBlendWeight_Execute()
+{
+	DECLARE_SCOPE_HIERARCHICAL_COUNTER_RIGUNIT();
+	if (ExecuteContext.Hierarchy)
+	{
+		if (FRigPhysicsComponent* PhysicsComponent = Cast<FRigPhysicsComponent>(
+			ExecuteContext.Hierarchy->FindComponent(PhysicsBody)))
+		{
+			if (UControlRig* ControlRig = Cast<UControlRig>(ExecuteContext.Hierarchy->GetOuter()))
+			{
+				if (FRigPhysicsSolver* PhysicsSolver = Cast<FRigPhysicsSolver>(
+					ControlRig->GetPhysicsSolver((PhysicsComponent->BodySolverSettings.SolverID))))
+				{
+					PhysicsSolver->SetBodyPhysicsBlendWeight(PhysicsBody, PhysicsBlendWeight);
+				}
+			}
+		}
+	}
+}
+
+//======================================================================================================================
+FRigUnit_HierarchySetPhysicsBodyUseSkeletalAnimation_Execute()
+{
+	DECLARE_SCOPE_HIERARCHICAL_COUNTER_RIGUNIT();
+	if (ExecuteContext.Hierarchy)
+	{
+		if (FRigPhysicsComponent* PhysicsComponent = Cast<FRigPhysicsComponent>(
+			ExecuteContext.Hierarchy->FindComponent(PhysicsBody)))
+		{
+			if (UControlRig* ControlRig = Cast<UControlRig>(ExecuteContext.Hierarchy->GetOuter()))
+			{
+				if (FRigPhysicsSolver* PhysicsSolver = Cast<FRigPhysicsSolver>(
+					ControlRig->GetPhysicsSolver((PhysicsComponent->BodySolverSettings.SolverID))))
+				{
+					PhysicsSolver->SetBodyUseSkeletalAnimation(PhysicsBody, bUseSkeletalAnimation);
+				}
+			}
+		}
+	}
+}
+
+//======================================================================================================================
+FRigUnit_HierarchySetPhysicsBodyUpdateKinematicFromSimulation_Execute()
+{
+	DECLARE_SCOPE_HIERARCHICAL_COUNTER_RIGUNIT();
+	if (ExecuteContext.Hierarchy)
+	{
+		if (FRigPhysicsComponent* PhysicsComponent = Cast<FRigPhysicsComponent>(
+			ExecuteContext.Hierarchy->FindComponent(PhysicsBody)))
+		{
+			if (UControlRig* ControlRig = Cast<UControlRig>(ExecuteContext.Hierarchy->GetOuter()))
+			{
+				if (FRigPhysicsSolver* PhysicsSolver = Cast<FRigPhysicsSolver>(
+					ControlRig->GetPhysicsSolver((PhysicsComponent->BodySolverSettings.SolverID))))
+				{
+					PhysicsSolver->SetBodyUpdateKinematicFromSimulation(PhysicsBody, bUpdateKinematicFromSimulation);
+				}
+			}
+		}
+	}
+}
+
diff --git a/Engine/Plugins/Experimental/PhysicsControl/Source/ControlRigPhysics/Private/RigPhysicsSolver.cpp b/Engine/Plugins/Experimental/PhysicsControl/Source/ControlRigPhysics/Private/RigPhysicsSolver.cpp
new file mode 100644
--- /dev/null
+++ b/Engine/Plugins/Experimental/PhysicsControl/Source/ControlRigPhysics/Private/RigPhysicsSolver.cpp
@@ -0,0 +1,685 @@
+// Copyright Epic Games, Inc. All Rights Reserved.
+
+#include "RigPhysicsSolver.h"
+#include "RigPhysicsComponent.h"
+
+#include "AnimNode_RigidBodyWithControl.h"
+#include "PhysicsControlPoseData.h"
+
+#include "ControlRig.h"
+#include "Rigs/RigHierarchy.h"
+#include "Rigs/RigHierarchyController.h"
+#include "Units/RigUnitContext.h"
+#include "RigVMCore/RigVMExecuteContext.h"
+
+#include "Engine/World.h"
+
+#include "Physics/ImmediatePhysics/ImmediatePhysicsActorHandle.h"
+#include "Physics/ImmediatePhysics/ImmediatePhysicsJointHandle.h"
+#include "Physics/ImmediatePhysics/ImmediatePhysicsSimulation.h"
+#include "Physics/ImmediatePhysics/ImmediatePhysicsChaos/ImmediatePhysicsSimulation_Chaos.h"
+
+#include "Physics/Experimental/PhysScene_Chaos.h"
+
+#include "Chaos/ParticleHandle.h"
+#include "Chaos/PBDJointConstraints.h"
+#include "Chaos/ImplicitObject.h"
+#include "Chaos/ShapeInstance.h"
+#include "Chaos/Capsule.h"
+#include "Chaos/ChaosScene.h"
+#include "Chaos/PBDJointConstraintUtilities.h"
+#include "Chaos/Evolution/SimulationSpace.h"
+
+#include "PhysicsEngine/BodyInstance.h"
+#include "PhysicsEngine/ConstraintInstance.h"
+#include "PhysicsEngine/BodySetup.h"
+
+DEFINE_LOG_CATEGORY(LogRigPhysics);
+
+//======================================================================================================================
+FRigPhysicsSolver::FRigPhysicsSolver(
+	TObjectPtr<UControlRig>                     InOwningControlRig,
+	const FRigPhysicsSolverSettings&            InSolverSettings,
+	const FRigPhysicsSimulationSpaceSettings&   InSimulationSpaceSettings)
+	: OwningControlRig(InOwningControlRig)
+	, SolverSettings(InSolverSettings)
+	, SimulationSpaceSettings(InSimulationSpaceSettings)
+{
+	// Danny TODO temporary hackery to allow creation in worker thread
+	Chaos::CVars::bEnableAsyncInitBody = true;
+}
+
+//======================================================================================================================
+void FRigPhysicsSolver::AddBody(FRigComponentKey RigComponentKey)
+{
+	// Just make the record for now. The body details may change, so only actually instantiate later.
+	FRigBodyRecord& Record = BodyRecords.Add(RigComponentKey);
+}
+
+//======================================================================================================================
+void FRigPhysicsSolver::InitialiseSimulation()
+{
+	DestroyPhysicsSimulation();
+
+	Simulation = MakeShared<ImmediatePhysics::FSimulation>();
+}
+
+//======================================================================================================================
+void FRigPhysicsSolver::DestroyPhysicsSimulation()
+{
+	for (TPair<FRigComponentKey, FRigBodyRecord>& BodyRecordPair : BodyRecords)
+	{
+		FRigBodyRecord& Record = BodyRecordPair.Value;
+		if (Record.JointHandle)
+		{
+			Simulation->DestroyJoint(Record.JointHandle);
+			Record.JointHandle = nullptr;
+		}
+		if (Record.ActorHandle)
+		{
+			Simulation->DestroyActor(Record.ActorHandle);
+			Record.ActorHandle = nullptr;
+		}
+	}
+	// We don't clear out the BodyRecords because they tell us what we have been told to make.
+
+	for (ImmediatePhysics::FActorHandle* ActorHandle : AdditionalActorHandles)
+	{
+		Simulation->DestroyActor(ActorHandle);
+	}
+	// Additional actors are made/owned by us on instantiation, so need to be cleared out
+	AdditionalActorHandles.Empty();
+
+	Simulation.Reset();
+}
+
+//======================================================================================================================
+static void SetCommonProperties(const FRigPhysicsCollisionShape& Shape, FKShapeElem& ShapeElem)
+{
+	ShapeElem.RestOffset = Shape.RestOffset;
+	ShapeElem.SetName(Shape.Name);
+	ShapeElem.SetContributeToMass(Shape.bContributeToMass);
+#ifdef PER_SHAPE_COLLISION
+	// Note that FKShapeElem supports enabling/disabling collision per shape, but this is discarded by the immediate
+	// solver.
+	ShapeElem.SetCollisionEnabled(Shape.CollisionEnabled);
+#endif
+}
+
+//======================================================================================================================
+static Chaos::EJointMotionType GetMotionTypeFromLimitValue(const float Value)
+{
+	// Danny TODO This +ve/-ve thing isn't nice for users - use an enum
+	if (Value > 0.0f)
+	{
+		return Chaos::EJointMotionType::Limited;
+	}
+	else if (Value < 0.0f)
+	{
+		return Chaos::EJointMotionType::Free;
+	}
+	else
+	{
+		return Chaos::EJointMotionType::Locked;
+	}
+}
+
+//======================================================================================================================
+ImmediatePhysics::FActorHandle* FRigPhysicsSolver::CreateBody(
+	const FRigPhysicsCollision& Collision, 
+	const FRigPhysicsDynamics*  Dynamics, 
+	const FTransform&           BodyRelSimSpaceTM) const
+{
+	// Danny TODO Use the correct materials (potentially tricky since we may be called on
+	// any thread - see the RBAN code for reference).
+
+	// This appears to be the only way
+	UBodySetup* BodySetup = NewObject<UBodySetup>();
+	BodySetup->CollisionTraceFlag = CTF_UseSimpleAsComplex;
+
+	// Note that the base class (FKShapeElem) for the shapes has got additional
+	// functionality such as enable/disable collision etc.
+
+	bool bAddedCollision = false;
+	for (const FRigPhysicsCollisionBox& Shape : Collision.Boxes)
+	{
+		FKBoxElem Elem(Shape.Extents.X, Shape.Extents.Y, Shape.Extents.Z);
+		SetCommonProperties(Shape, Elem);
+		Elem.Center = Shape.TM.GetTranslation();
+		Elem.Rotation = Shape.TM.Rotator();
+		BodySetup->AggGeom.BoxElems.Add(Elem);
+		bAddedCollision = true;
+	}
+
+	for (const FRigPhysicsCollisionSphere& Shape : Collision.Spheres)
+	{
+		FKSphereElem Elem(Shape.Radius);
+		SetCommonProperties(Shape, Elem);
+		Elem.Center = Shape.TM.GetTranslation();
+		// Note that there is no rotation
+		BodySetup->AggGeom.SphereElems.Add(Elem);
+		bAddedCollision = true;
+	}
+
+	for (const FRigPhysicsCollisionCapsule& Shape : Collision.Capsules)
+	{
+		FKSphylElem Elem(Shape.Radius, Shape.Length);
+		SetCommonProperties(Shape, Elem);
+		Elem.Center = Shape.TM.GetTranslation();
+		Elem.Rotation = Shape.TM.Rotator();
+		BodySetup->AggGeom.SphylElems.Add(Elem);
+		bAddedCollision = true;
+	}
+
+	// This is a bit odd, but in order to create a low-level immediate physics representation, you
+	// have to create a high-level temporary one. Danny TODO avoid using BodyInstance and BodySetup.
+	//
+	// See USkeletalMeshComponent::InstantiatePhysicsAsset_Internal for how it's done in RBAN (which
+	// makes body instances, converts to immediate physics actors, and then discards the body instances).
+	FBodyInstance BodyInstance;
+	if (Dynamics)
+	{
+		BodyInstance.bSimulatePhysics = true;
+		BodyInstance.bEnableGravity = true;
+		BodyInstance.SetMassOverride(Dynamics->Mass);
+		BodyInstance.LinearDamping = 0;
+		BodyInstance.AngularDamping = 0;
+	}
+	else
+	{
+		BodyInstance.bSimulatePhysics = false;
+	}
+	BodyInstance.InitBody(BodySetup, BodyRelSimSpaceTM, nullptr, nullptr);
+
+	ImmediatePhysics::FActorHandle* ActorHandle = Dynamics ?
+		Simulation->CreateDynamicActor(&BodyInstance, BodyRelSimSpaceTM) :
+		Simulation->CreateKinematicActor(&BodyInstance, BodyRelSimSpaceTM);
+
+	// Creation is guaranteed to return something
+	check(ActorHandle);
+
+	Simulation->AddToCollidingPairs(ActorHandle);
+
+	if (bAddedCollision && Dynamics)
+	{
+		// Note that particles are always created disabled. They will simulate when disabled,
+		// but won't collide!
+		ActorHandle->SetEnabled(true);
+	}
+
+	return ActorHandle;
+}
+
+
+//======================================================================================================================
+void FRigPhysicsSolver::Instantiate()
+{
+	if (!bNeedToInstantiate || !OwningControlRig)
+	{
+		return;
+	}
+
+	InitialiseSimulation();
+
+	if (!SolverSettings.Collision.IsEmpty())
+	{
+		// When we make these additional collision shapes, their actors are all considered to be at
+		// the origin, with the offsets being contained in the collision shapes.
+		FTransform BodyRelSimSpaceTM = ConvertCollisionSpaceTransformToSimSpace(FTransform());
+
+		ImmediatePhysics::FActorHandle* EnvironmentActorHandle = CreateBody(
+			SolverSettings.Collision, nullptr, BodyRelSimSpaceTM);
+		AdditionalActorHandles.Add(EnvironmentActorHandle);
+	}
+
+	InstantiatePhysicsBodies();
+
+	InstantiatePhysicsJoints();
+
+	bNeedToInstantiate = false;
+}
+
+//======================================================================================================================
+void FRigPhysicsSolver::InstantiatePhysicsJoints()
+{
+	const URigHierarchy* Hierarchy = OwningControlRig->GetHierarchy();
+
+	// Once all the bodies are created, we can make their physics joints
+	for (TPair<FRigComponentKey, FRigBodyRecord>& BodyRecordPair : BodyRecords)
+	{
+		const FRigComponentKey& ChildBodyComponentKey = BodyRecordPair.Key;
+		FRigBodyRecord& ChildBodyRecord = BodyRecordPair.Value;
+
+		ImmediatePhysics::FActorHandle* ChildActorHandle = ChildBodyRecord.ActorHandle;
+		if (!ChildActorHandle)
+		{
+			continue;
+		}
+
+		if (const FRigPhysicsComponent* ChildPhysicsComponent = Cast<FRigPhysicsComponent>(
+			Hierarchy->FindComponent(ChildBodyComponentKey)))
+		{
+			const FRigPhysicsJoint& PhysicsJoint = ChildPhysicsComponent->Joint;
+			if (!PhysicsJoint.bEnable)
+			{
+				continue;
+			}
+
+			// Now we need to find the parent body (etc) for the physics joint
+			FRigComponentKey ParentBodyComponentKey = PhysicsJoint.ParentBody;
+
+			// And along the way we will discover various other things
+			FRigElementKey ChildBoneKey = ChildBodyComponentKey.ElementKey;
+			if (!ChildBoneKey.IsValid())
+			{
+				continue;
+			}
+
+			if (!ParentBodyComponentKey.IsValid())
+			{
+				// Physics bodies are parented to a bone - we want the parent of the bone
+				FRigElementKey ParentBoneKey = Hierarchy->GetFirstParent(ChildBoneKey);
+				// Now we have the regular bone, we want the physics element that is the child
+				if (ParentBoneKey.IsValid())
+				{
+					TArray<FRigComponentKey> ComponentKeys = Hierarchy->GetComponentKeys(ParentBoneKey);
+					for (const FRigComponentKey& ComponentKey : ComponentKeys)
+					{
+						const FRigBaseComponent* BaseComponent = Hierarchy->FindComponent(ComponentKey);
+						// Danny TODO what's the best way to find/check this?!
+						if (Cast<FRigPhysicsComponent>(BaseComponent))
+						{
+							ParentBodyComponentKey = ComponentKey;
+							break;
+						}
+					}
+				}
+			}
+
+			const FRigPhysicsComponent* ParentPhysicsComponent =
+				Cast<FRigPhysicsComponent>(Hierarchy->FindComponent(ParentBodyComponentKey));
+			if (!ParentPhysicsComponent)
+			{
+				continue;
+			}
+
+			const FRigBodyRecord* ParentRecord = nullptr;
+			ImmediatePhysics::FActorHandle* ParentActorHandle = nullptr;
+
+			// Find the record that corresponds to the parent, so we can see if there is a parent actor
+			for (const TPair<FRigComponentKey, FRigBodyRecord>& OtherBodyRecordPair : BodyRecords)
+			{
+				const FRigComponentKey& OtherBodyComponentKey = OtherBodyRecordPair.Key;
+				const FRigBodyRecord& OtherBodyRecord = OtherBodyRecordPair.Value;
+				if (OtherBodyComponentKey == ParentPhysicsComponent->GetKey())
+				{
+					ParentRecord = &OtherBodyRecord;
+					ParentActorHandle = OtherBodyRecord.ActorHandle;
+					break;
+				}
+			}
+			if (!ParentRecord)
+			{
+				continue;
+			}
+
+			// Make the physics joint (joint constraint). 
+			// UE likes to treat Body1 (index 0) as the child, and Body2 (index 1) as the parent
+
+			const FTransform ParentCoMTransform = ParentActorHandle->GetLocalCoMTransform();
+			const FTransform ChildCoMTransform = ChildActorHandle->GetLocalCoMTransform();
+
+			Chaos::FPBDJointSettings ConstraintSettings;
+			if (PhysicsJoint.bAutoCalculateChildOffset)
+			{
+				// Danny TODO - is the "non-override" actually always simply the identity TM?
+				ConstraintSettings.ConnectorTransforms[0] = FTransform();
+			}
+			ConstraintSettings.ConnectorTransforms[0] =
+				PhysicsJoint.ExtraChildOffset * ConstraintSettings.ConnectorTransforms[0];
+
+			if (PhysicsJoint.bAutoCalculateParentOffset)
+			{
+				// Danny TODO is it right to use the TM of the child bone relative to parent bone?
+				// What about if the body has got a TM too?
+				ConstraintSettings.ConnectorTransforms[1] = Hierarchy->GetLocalTransform(ChildBoneKey, true);
+			}
+			ConstraintSettings.ConnectorTransforms[1] =
+				PhysicsJoint.ExtraParentOffset * ConstraintSettings.ConnectorTransforms[1];
+
+			// Danny TODO The linear limit is likely implemented as a box limit (check) - it's just
+			// not exposed as one. We should change this.
+			ConstraintSettings.LinearLimit = PhysicsJoint.LinearLimit.Length();
+			Chaos::EJointMotionType LinearMotionType = GetMotionTypeFromLimitValue(ConstraintSettings.LinearLimit);
+			ConstraintSettings.LinearMotionTypes = { LinearMotionType, LinearMotionType, LinearMotionType };
+
+			ConstraintSettings.AngularLimits = FMath::DegreesToRadians(PhysicsJoint.AngularLimit);
+			ConstraintSettings.AngularMotionTypes = {
+				GetMotionTypeFromLimitValue(PhysicsJoint.AngularLimit[0]),
+				GetMotionTypeFromLimitValue(PhysicsJoint.AngularLimit[1]),
+				GetMotionTypeFromLimitValue(PhysicsJoint.AngularLimit[2])
+			};
+
+			// The physics setting is backwards, because we can't enable collision on bodies that
+			// are set to not collide for other reasons.
+			ConstraintSettings.bCollisionEnabled = !PhysicsJoint.bDisableCollision;
+			ConstraintSettings.bProjectionEnabled = 
+				PhysicsJoint.LinearProjectionAmount > 0 || PhysicsJoint.AngularProjectionAmount > 0;
+			ConstraintSettings.AngularProjection = PhysicsJoint.AngularProjectionAmount;
+			ConstraintSettings.LinearProjection = PhysicsJoint.LinearProjectionAmount;
+			ConstraintSettings.ParentInvMassScale = PhysicsJoint.ParentInverseMassScale;
+
+			ConstraintSettings.bUseLinearSolver = SolverSettings.bUseLinearJointSolver;
+
+			ChildBodyRecord.JointHandle = Simulation->CreateJoint(
+				ConstraintSettings, ChildActorHandle, ParentActorHandle);
+		}
+	}
+}
+
+//======================================================================================================================
+void FRigPhysicsSolver::InstantiatePhysicsBodies()
+{
+	const URigHierarchy* Hierarchy = OwningControlRig->GetHierarchy();
+
+	for (TPair<FRigComponentKey, FRigBodyRecord>& BodyRecordPair : BodyRecords)
+	{
+		const FRigComponentKey& ComponentKey = BodyRecordPair.Key;
+		FRigBodyRecord& Record = BodyRecordPair.Value;
+
+		if (const FRigPhysicsComponent* PhysicsComponent = Cast<FRigPhysicsComponent>(
+			Hierarchy->FindComponent(ComponentKey)))
+		{
+			const FRigPhysicsCollision& Collision = PhysicsComponent->Collision;
+			const FRigPhysicsDynamics& Dynamics = PhysicsComponent->Dynamics;
+			const FPhysicsControlModifierData& BodyData = PhysicsComponent->BodyData;
+
+			FRigElementKey SourceKey = PhysicsComponent->BodySolverSettings.SourceBone;
+			if (!SourceKey.IsValid())
+			{
+				SourceKey = ComponentKey.ElementKey;
+			}
+
+			FTransform BodyRelComponentTM;
+			// What should we do if the key is not valid? 
+			if (SourceKey.IsValid())
+			{
+				BodyRelComponentTM = Hierarchy->GetGlobalTransform(SourceKey);
+			}
+
+			const FTransform BodyRelSimSpaceTM = ConvertComponentSpaceTransformToSimSpace(BodyRelComponentTM);
+
+			Record.ActorHandle = CreateBody(Collision, &Dynamics, BodyRelSimSpaceTM);
+
+			// Note that these data will be applied in the PrePhysics update
+			Record.BodyData = BodyData;
+		}
+	}
+}
+
+//======================================================================================================================
+void FRigPhysicsSolver::SetBodySparseData(
+	FRigComponentKey                         RigComponentKey,
+	const FPhysicsControlModifierSparseData& SparseData)
+{
+	if (FRigBodyRecord* Record = BodyRecords.Find(RigComponentKey))
+	{
+		Record->BodyData.UpdateFromSparseData(SparseData);
+	}
+}
+
+//======================================================================================================================
+void FRigPhysicsSolver::SetBodyMovementType(FRigComponentKey RigComponentKey, EPhysicsMovementType MovementType)
+{
+	if (FRigBodyRecord* Record = BodyRecords.Find(RigComponentKey))
+	{
+		Record->BodyData.MovementType = MovementType;
+	}
+}
+
+//======================================================================================================================
+void FRigPhysicsSolver::SetBodyCollisionType(FRigComponentKey RigComponentKey, ECollisionEnabled::Type CollisionType)
+{
+	if (FRigBodyRecord* Record = BodyRecords.Find(RigComponentKey))
+	{
+		Record->BodyData.CollisionType = CollisionType;
+	}
+}
+
+//======================================================================================================================
+void FRigPhysicsSolver::SetBodyGravityMultiplier(FRigComponentKey RigComponentKey, float GravityMultiplier)
+{
+	if (FRigBodyRecord* Record = BodyRecords.Find(RigComponentKey))
+	{
+		Record->BodyData.GravityMultiplier = GravityMultiplier;
+	}
+}
+
+//======================================================================================================================
+void FRigPhysicsSolver::SetBodyPhysicsBlendWeight(FRigComponentKey RigComponentKey, float PhysicsBlendWeight)
+{
+	if (FRigBodyRecord* Record = BodyRecords.Find(RigComponentKey))
+	{
+		Record->BodyData.PhysicsBlendWeight = PhysicsBlendWeight;
+	}
+}
+
+//======================================================================================================================
+void FRigPhysicsSolver::SetBodyUseSkeletalAnimation(FRigComponentKey RigComponentKey, bool bUseSkeletalAnimation)
+{
+	if (FRigBodyRecord* Record = BodyRecords.Find(RigComponentKey))
+	{
+		Record->BodyData.bUseSkeletalAnimation = bUseSkeletalAnimation;
+	}
+}
+
+//======================================================================================================================
+void FRigPhysicsSolver::SetBodyUpdateKinematicFromSimulation(
+	FRigComponentKey RigComponentKey, bool bUpdateKinematicFromSimulation)
+{
+	if (FRigBodyRecord* Record = BodyRecords.Find(RigComponentKey))
+	{
+		Record->BodyData.bUpdateKinematicFromSimulation = bUpdateKinematicFromSimulation;
+	}
+}
+
+//======================================================================================================================
+void FRigPhysicsSolver::SetBodyKinematicTarget(
+	FRigComponentKey                RigComponentKey,
+	ERigPhysicsKinematicTargetSpace KinematicTargetSpace,
+	const FTransform&               KinematicTarget)
+{
+	if (FRigBodyRecord* Record = BodyRecords.Find(RigComponentKey))
+	{
+		Record->KinematicTargetSpace = KinematicTargetSpace;
+		Record->KinematicTarget = KinematicTarget;
+	}
+}
+
+//======================================================================================================================
+void FRigPhysicsSolver::UpdatePrePhysics(FRigVMExecuteContext& ExecuteContext, const FSimulationSpaceData& SimulationSpaceData)
+{
+	URigHierarchy* Hierarchy = OwningControlRig->GetHierarchy();
+
+	for (ImmediatePhysics::FActorHandle* ActorHandle : AdditionalActorHandles)
+	{
+		FTransform BodyRelSimSpaceTM = ConvertCollisionSpaceTransformToSimSpace(FTransform());
+		ActorHandle->SetKinematicTarget(BodyRelSimSpaceTM);
+	}
+
+	for (TPair<FRigComponentKey, FRigBodyRecord>& BodyRecordPair : BodyRecords)
+	{
+		const FRigComponentKey& ComponentKey = BodyRecordPair.Key;
+		const FRigBodyRecord& Record = BodyRecordPair.Value;
+
+		if (const FRigPhysicsComponent* PhysicsComponent = Cast<FRigPhysicsComponent>(
+			Hierarchy->FindComponent(ComponentKey)))
+		{
+			if (Record.ActorHandle)
+			{
+				UpdateBodyFromModifierData(
+					Record.ActorHandle, Simulation.Get(), Record.BodyData, SimulationSpaceData.Gravity);
+
+				if (Record.ActorHandle->GetIsKinematic())
+				{
+					// Get the target in component space, and then convert it into sim space if necessary.
+
+					// If the target is already in component space, then that's all we need
+					FTransform KinematicTargetCS;
+					switch (Record.KinematicTargetSpace)
+					{
+						case ERigPhysicsKinematicTargetSpace::Component:
+						{
+							KinematicTargetCS = Record.KinematicTarget;
+							break;
+						}
+						case ERigPhysicsKinematicTargetSpace::World:
+						{
+							// Record.KinematicTarget * SimulationSpaceState.ComponentTM.Inverse();
+							KinematicTargetCS = Record.KinematicTarget.GetRelativeTransform(
+								SimulationSpaceState.ComponentTM);
+							break;
+						}
+						default:
+						{
+							// All the other options are relative to a bone, so the first task is to get
+							// that, which will be in component space
+							FRigElementKey SourceKey = PhysicsComponent->BodySolverSettings.SourceBone;
+							if (!SourceKey.IsValid())
+							{
+								SourceKey = ComponentKey.ElementKey;
+							}
+							if (SourceKey.IsValid())
+							{
+								FTransform SourceComponentSpaceTM = Hierarchy->GetGlobalTransform(SourceKey);
+
+								switch (Record.KinematicTargetSpace)
+								{
+									case ERigPhysicsKinematicTargetSpace::OffsetInBoneSpace:
+									{
+										KinematicTargetCS = Record.KinematicTarget * SourceComponentSpaceTM;
+										break;
+									}
+									case ERigPhysicsKinematicTargetSpace::OffsetInWorldSpace:
+									{
+										// Convert the bone to WS, apply the target, and convert back
+										FTransform BoneWS = SourceComponentSpaceTM * SimulationSpaceState.ComponentTM;
+										FTransform KinematicTargetWS = BoneWS * Record.KinematicTarget;
+										// Danny TODO figure out which of the GetRelativeTransform versions this is
+										KinematicTargetCS = SimulationSpaceState.ComponentTM.Inverse() * KinematicTargetWS;
+										break;
+									}
+									case ERigPhysicsKinematicTargetSpace::OffsetInComponentSpace:
+									{
+										KinematicTargetCS = SourceComponentSpaceTM * Record.KinematicTarget;
+										break;
+									}
+									default:
+									{
+										UE_CONTROLRIG_RIGUNIT_REPORT_ERROR(TEXT("Kinematic target space is not valid"));
+									}
+								}
+							}
+						}
+					}
+					FTransform KinematicTargetTM = ConvertComponentSpaceTransformToSimSpace(KinematicTargetCS);
+					Record.ActorHandle->SetKinematicTarget(KinematicTargetTM);
+				}
+			}
+		}
+	}
+}
+
+//======================================================================================================================
+// Note that we read back into a target bone, which may have been specified explicitly, or will
+// otherwise default to the physics element parent.
+void FRigPhysicsSolver::UpdatePostPhysics()
+{
+	URigHierarchy* Hierarchy = OwningControlRig->GetHierarchy();
+	for (TPair<FRigComponentKey, FRigBodyRecord>& BodyRecordPair : BodyRecords)
+	{
+		const FRigComponentKey& ComponentKey = BodyRecordPair.Key;
+		const FRigBodyRecord& Record = BodyRecordPair.Value;
+
+		if (Record.ActorHandle)
+		{
+			if (const FRigPhysicsComponent* PhysicsComponent = Cast<FRigPhysicsComponent>(
+				Hierarchy->FindComponent(ComponentKey)))
+			{
+				// We will only have created this body record if it's associated with our solver, so
+				// we don't need to check again that our solver is in its list.
+				FRigElementKey TargetElementKey = PhysicsComponent->BodySolverSettings.TargetBone;
+				if (!TargetElementKey.IsValid())
+				{
+					TargetElementKey = ComponentKey.ElementKey;
+				}
+
+				if (TargetElementKey.IsValid())
+				{
+					FTransform TargetSimSpaceTM = Record.ActorHandle->GetWorldTransform();
+					FTransform TargetComponentSpaceTM = ConvertSimSpaceTransformToComponentSpace(TargetSimSpaceTM);
+					Hierarchy->SetGlobalTransform(TargetElementKey, TargetComponentSpaceTM);
+				}
+			}
+		}
+	}
+}
+
+//======================================================================================================================
+void FRigPhysicsSolver::StepSimulation(FRigVMExecuteContext& ExecuteContext, double DeltaTime)
+{
+	// We need to know about the simulation space etc before we can instantiate anything into the right place
+	FSimulationSpaceData SimulationSpaceData = UpdateSimulationSpaceStateAndCalculateData(DeltaTime);
+
+	// We instantiate when we do the first simulation - this makes sure any changes applied by
+	// the user have been made. It also means there is no overhead if physics is never stepped.
+	// However, there may be a hitch due to the creation, so it may also happen during construction.
+	Instantiate(); // There is an early out if it's already been done
+
+	if (!Simulation)
+	{
+		return;
+	}
+
+	// Set settings that might change
+	Simulation->SetSolverSettings(
+		SolverSettings.FixedTimeStep,
+		999999999.0, // Danny TODO cull distance - expose this? Yuck
+		SolverSettings.MaxDepenetrationVelocity,
+		SolverSettings.bUseLinearJointSolver,
+		SolverSettings.PositionIterations,
+		SolverSettings.VelocityIterations,
+		SolverSettings.ProjectionIterations,
+		SolverSettings.bUseManifolds);
+
+	// Other settings - would normally be static (so Danny TODO move this)
+	ChaosJointSolverSettings.bSolvePositionLast = SolverSettings.bSolveJointPositionsLast;
+	ChaosJointSolverSettings.bSortEnabled = true;
+
+	// Simulation space
+#if 0
+	// Danny TODO this requires a tweak to the API to allow the full settings to be passed in and
+	// controlled by us. That needs reviewing by physics.
+	Chaos::FSimulationSpaceSettings ChaosSimulationSpaceSettings;
+	ChaosSimulationSpaceSettings.Alpha = SimulationSpaceSettings.SpaceMovementAmount;
+	ChaosSimulationSpaceSettings.ExternalLinearEtherDrag = SimulationSpaceSettings.ExternalLinearDragV;
+	Simulation->SetSimulationSpaceSettings(ChaosSimulationSpaceSettings);
+#else
+	Simulation->SetSimulationSpaceSettings(
+		SimulationSpaceSettings.SpaceMovementAmount, SimulationSpaceSettings.ExternalLinearDragV);
+#endif
+	Simulation->UpdateSimulationSpace(
+		SimulationSpaceState.SimulationSpaceTM,
+		SimulationSpaceData.LinearVelocity,
+		SimulationSpaceData.AngularVelocity,
+		SimulationSpaceData.LinearAcceleration,
+		SimulationSpaceData.AngularAcceleration);
+
+	UpdatePrePhysics(ExecuteContext, SimulationSpaceData);
+
+	Simulation->Simulate(
+		DeltaTime, SolverSettings.MaxDeltaTime, SolverSettings.MaxTimeSteps, 
+		SimulationSpaceData.Gravity, &ChaosJointSolverSettings);
+
+	UpdatePostPhysics();
+}
diff --git a/Engine/Plugins/Experimental/PhysicsControl/Source/ControlRigPhysics/Private/RigPhysicsSolver.h b/Engine/Plugins/Experimental/PhysicsControl/Source/ControlRigPhysics/Private/RigPhysicsSolver.h
new file mode 100644
--- /dev/null
+++ b/Engine/Plugins/Experimental/PhysicsControl/Source/ControlRigPhysics/Private/RigPhysicsSolver.h
@@ -0,0 +1,172 @@
+// Copyright Epic Games, Inc. All Rights Reserved.
+
+#pragma once
+
+#include "RigPhysicsData.h"
+
+#include "PhysicsControlData.h"
+
+#include "Rigs/RigHierarchyComponents.h"
+
+#include "Physics/ImmediatePhysics/ImmediatePhysicsDeclares.h"
+#include "Chaos/PBDJointConstraintTypes.h"
+
+#include "RigPhysicsSolver.generated.h"
+
+class UControlRig;
+
+struct FRigVMDrawInterface;
+struct FRigVMExecuteContext;
+
+struct FRigPhysicsCollision;
+struct FRigPhysicsDynamics;
+
+#if UE_BUILD_SHIPPING || UE_BUILD_TEST
+DECLARE_LOG_CATEGORY_EXTERN(LogRigPhysics, Log, Warning);
+#else
+DECLARE_LOG_CATEGORY_EXTERN(LogRigPhysics, Log, All);
+#endif
+
+/**
+ * For internal use only, to keep track of objects we have created
+ */
+USTRUCT()
+struct FRigBodyRecord
+{
+	GENERATED_BODY()
+
+	ImmediatePhysics::FActorHandle* ActorHandle = nullptr;
+	ImmediatePhysics::FJointHandle* JointHandle = nullptr;
+	FPhysicsControlModifierData     BodyData;
+
+	ERigPhysicsKinematicTargetSpace KinematicTargetSpace = ERigPhysicsKinematicTargetSpace::OffsetInBoneSpace;
+	FTransform                      KinematicTarget;
+};
+
+
+USTRUCT(BlueprintType)
+struct FRigPhysicsSolver : public FRigPhysicsSolverBase
+{
+public:
+	GENERATED_BODY()
+
+	FRigPhysicsSolver(
+		TObjectPtr<UControlRig>                   InOwningControlRig = nullptr, 
+		const FRigPhysicsSolverSettings&          InSolverSettings = FRigPhysicsSolverSettings(),
+		const FRigPhysicsSimulationSpaceSettings& InSimulationSpaceSettings = FRigPhysicsSimulationSpaceSettings());
+
+	// This adds an entry to our records - actual creation will happen at instantiation time
+	void AddBody(FRigComponentKey RigComponentKey);
+
+	void SetBodySparseData(FRigComponentKey RigComponentKey, const FPhysicsControlModifierSparseData& SparseData);
+	void SetBodyMovementType(FRigComponentKey RigComponentKey, EPhysicsMovementType MovementType);
+	void SetBodyCollisionType(FRigComponentKey RigComponentKey, ECollisionEnabled::Type CollisionType);
+	void SetBodyGravityMultiplier(FRigComponentKey RigComponentKey, float GravityMultiplier);
+	void SetBodyPhysicsBlendWeight(FRigComponentKey RigComponentKey, float PhysicsBlendWeight);
+	void SetBodyUseSkeletalAnimation(FRigComponentKey RigComponentKey, bool bUseSkeletalAnimation);
+	void SetBodyUpdateKinematicFromSimulation(FRigComponentKey RigComponentKey, bool bUpdateKinematicFromSimulation);
+
+	void SetBodyKinematicTarget(
+		FRigComponentKey                RigComponentKey, 
+		ERigPhysicsKinematicTargetSpace KinematicTargetSpace, 
+		const FTransform&               KinematicTarget);
+
+	// This will initialise/create the simulation and then create everything we need in it. 
+	void Instantiate();
+
+	// Integrates the simulation forwards by DeltaTime
+	void StepSimulation(FRigVMExecuteContext& ExecuteContext, double DeltaTime);
+
+	// Draws shapes etc AND (potentially) enables the low-level Chaos debug draw
+	void Draw(FRigVMDrawInterface* DI, const UWorld* DebugWorld) const;
+
+private:
+	// Used by the world-space to simulation-space motion transfer system in Component- or
+	// Bone-Space sims, and preserved between updates.
+	struct FSimulationSpaceState
+	{
+		FTransform ComponentTM;
+		FTransform BoneRelComponentTM;
+		FVector    ComponentLinearVelocity;
+		FVector    ComponentAngularVelocity;
+		FVector    BoneLinearVelocity;
+		FVector    BoneAngularVelocity;
+
+		// The world transform of the simulation space
+		FTransform SimulationSpaceTM;
+		FTransform PrevSimulationSpaceTM;
+		FTransform PrevPrevSimulationSpaceTM;
+		// The time between SimulationSpaceTM and PrevSimulationSpaceTM
+		float      Dt = 1.0f; 
+		// The time between PrevSimulationSpaceTM and PrevPrevSimulationSpaceTM
+		float      PrevDt = 1.0f;
+	};
+
+	// Represents the properties of the simulation space - calculated near the beginning of the update.
+	// Note that all these are specified in the simulation space itself
+	struct FSimulationSpaceData
+	{
+		FVector LinearVelocity;
+		FVector AngularVelocity;
+		FVector LinearAcceleration;
+		FVector AngularAcceleration;
+		FVector Gravity;
+	};
+
+	void InitialiseSimulation();
+
+	void InstantiatePhysicsJoints();
+
+	void InstantiatePhysicsBodies();
+
+	// Creates an actor with collision. This will by dynamic if Dynamics is valid, or otherwise kinematic
+	ImmediatePhysics::FActorHandle* CreateBody(
+		const FRigPhysicsCollision& Collision,
+		const FRigPhysicsDynamics*  Dynamics,
+		const FTransform&           BodyRelSimSpaceTM) const;
+
+	// This releases the actors etc in the simulation, and destroys the simulation, but it doesn't
+	// change our records (apart from resetting pointers to actors etc).
+	void DestroyPhysicsSimulation();
+
+	void UpdatePrePhysics(FRigVMExecuteContext& ExecuteContext, const FSimulationSpaceData& SimulationSpaceData);
+	void UpdatePostPhysics();
+
+	void InitSimulationSpace(const FTransform& ComponentToWorld, const FTransform& BoneToComponent);
+
+	FSimulationSpaceData UpdateSimulationSpaceStateAndCalculateData(const float Dt);
+
+	// Returns the simulation space transform, in world space
+	FTransform GetSimulationSpaceTransform() const;
+
+	// Converts a transform from component space (e.g. coming from the owning control rig) into the
+	// simulation space
+	FTransform ConvertComponentSpaceTransformToSimSpace(const FTransform& TM) const;
+
+	// Converts a transform from the simulation space to component space (e.g. for writing back to
+	// the owning control rig)
+	FTransform ConvertSimSpaceTransformToComponentSpace(const FTransform& TM) const;
+
+	// Converts a vector specified in world space to the simulation space (e.g. converting gravity)
+	FVector ConvertWorldVectorToSimSpaceNoScale(const FVector& WorldVector) const;
+
+	FTransform ConvertCollisionSpaceTransformToSimSpace(const FTransform& TM) const;
+
+private:
+	TObjectPtr<UControlRig> OwningControlRig;
+
+	TMap<FRigComponentKey, FRigBodyRecord> BodyRecords;
+	TSharedPtr<ImmediatePhysics::FSimulation> Simulation;
+
+	FRigPhysicsSolverSettings SolverSettings;
+	FRigPhysicsSimulationSpaceSettings SimulationSpaceSettings;
+
+	TArray<ImmediatePhysics::FActorHandle*> AdditionalActorHandles;
+
+	Chaos::FPBDJointSolverSettings ChaosJointSolverSettings;
+
+	FSimulationSpaceState SimulationSpaceState;
+
+	bool bNeedToInstantiate = true;
+};
+
diff --git a/Engine/Plugins/Experimental/PhysicsControl/Source/ControlRigPhysics/Private/RigPhysicsSolverSimulationSpace.cpp b/Engine/Plugins/Experimental/PhysicsControl/Source/ControlRigPhysics/Private/RigPhysicsSolverSimulationSpace.cpp
new file mode 100644
--- /dev/null
+++ b/Engine/Plugins/Experimental/PhysicsControl/Source/ControlRigPhysics/Private/RigPhysicsSolverSimulationSpace.cpp
@@ -0,0 +1,254 @@
+// Copyright Epic Games, Inc. All Rights Reserved.
+
+#include "RigPhysicsSolver.h"
+
+#include "Rigs/RigHierarchy.h"
+#include "Rigs/RigHierarchyController.h"
+#include "ControlRig.h"
+
+#include "Physics/ImmediatePhysics/ImmediatePhysicsActorHandle.h"
+#include "Physics/ImmediatePhysics/ImmediatePhysicsJointHandle.h"
+#include "Physics/ImmediatePhysics/ImmediatePhysicsSimulation.h"
+
+#include "Physics/Experimental/PhysScene_Chaos.h"
+#include "Chaos/ParticleHandle.h"
+#include "Chaos/PBDJointConstraints.h"
+#include "Chaos/ImplicitObject.h"
+#include "Chaos/ShapeInstance.h"
+#include "Chaos/Capsule.h"
+#include "Chaos/PBDJointConstraintUtilities.h"
+
+#include "PhysicsEngine/BodyInstance.h"
+#include "PhysicsEngine/ConstraintInstance.h"
+#include "PhysicsEngine/BodySetup.h"
+
+//======================================================================================================================
+FORCEINLINE_DEBUGGABLE FTransform GetSpaceTransform(
+	ERigPhysicsSimulationSpace Space, const FTransform& ComponentTM, const FTransform& BoneTM)
+{
+	switch (Space)
+	{
+	case ERigPhysicsSimulationSpace::Component: return ComponentTM;
+	case ERigPhysicsSimulationSpace::World: return FTransform::Identity;
+	case ERigPhysicsSimulationSpace::SpaceBone: return BoneTM * ComponentTM;
+	default: ensureMsgf(false, TEXT("Unsupported Simulation Space")); return FTransform::Identity;
+	}
+}
+
+//======================================================================================================================
+FTransform FRigPhysicsSolver::GetSimulationSpaceTransform() const
+{
+	switch (SolverSettings.SimulationSpace)
+	{
+	case ERigPhysicsSimulationSpace::Component: return SimulationSpaceState.ComponentTM;
+	case ERigPhysicsSimulationSpace::World: return FTransform::Identity;
+	case ERigPhysicsSimulationSpace::SpaceBone: return SimulationSpaceState.BoneRelComponentTM * SimulationSpaceState.ComponentTM;
+	default: ensureMsgf(false, TEXT("Unsupported Simulation Space")); return FTransform::Identity;
+	}
+}
+
+//======================================================================================================================
+FTransform FRigPhysicsSolver::ConvertComponentSpaceTransformToSimSpace(const FTransform& TM) const
+{
+	switch (SolverSettings.SimulationSpace)
+	{
+	case ERigPhysicsSimulationSpace::Component: return TM;
+	case ERigPhysicsSimulationSpace::World:  return TM * SimulationSpaceState.ComponentTM;
+	case ERigPhysicsSimulationSpace::SpaceBone: return TM.GetRelativeTransform(SimulationSpaceState.BoneRelComponentTM);
+	default: ensureMsgf(false, TEXT("Unsupported Simulation Space")); return TM;
+	}
+}
+
+//======================================================================================================================
+FTransform FRigPhysicsSolver::ConvertSimSpaceTransformToComponentSpace(const FTransform& TM) const
+{
+	switch (SolverSettings.SimulationSpace)
+	{
+	case ERigPhysicsSimulationSpace::Component: return TM;
+	case ERigPhysicsSimulationSpace::World:  return TM.GetRelativeTransform(SimulationSpaceState.ComponentTM);
+	case ERigPhysicsSimulationSpace::SpaceBone: return TM * SimulationSpaceState.BoneRelComponentTM;
+	default: ensureMsgf(false, TEXT("Unsupported Simulation Space")); return TM;
+	}
+}
+
+//======================================================================================================================
+FORCEINLINE_DEBUGGABLE FVector ConvertWorldVectorToSimSpaceNoScale(
+	ERigPhysicsSimulationSpace Space, const FVector& WorldVector, 
+	const FTransform& ComponentTM, const FTransform& BoneTM)
+{
+	switch (Space)
+	{
+	case ERigPhysicsSimulationSpace::Component: return ComponentTM.InverseTransformVectorNoScale(WorldVector);
+	case ERigPhysicsSimulationSpace::World: return WorldVector;
+	case ERigPhysicsSimulationSpace::SpaceBone:
+		return BoneTM.InverseTransformVectorNoScale(ComponentTM.InverseTransformVectorNoScale(WorldVector));
+	default: ensureMsgf(false, TEXT("Unsupported Simulation Space")); return WorldVector;
+	}
+}
+
+//======================================================================================================================
+FVector FRigPhysicsSolver::ConvertWorldVectorToSimSpaceNoScale(const FVector& WorldVector) const
+{
+	switch (SolverSettings.SimulationSpace)
+	{
+	case ERigPhysicsSimulationSpace::Component: return SimulationSpaceState.ComponentTM.InverseTransformVectorNoScale(WorldVector);
+	case ERigPhysicsSimulationSpace::World: return WorldVector;
+	case ERigPhysicsSimulationSpace::SpaceBone:
+		return SimulationSpaceState.BoneRelComponentTM.InverseTransformVectorNoScale(
+			SimulationSpaceState.ComponentTM.InverseTransformVectorNoScale(WorldVector));
+	default: ensureMsgf(false, TEXT("Unsupported Simulation Space")); return WorldVector;
+	}
+}
+
+//======================================================================================================================
+FTransform FRigPhysicsSolver::ConvertCollisionSpaceTransformToSimSpace(const FTransform& TM) const
+{
+	FTransform SimSpaceTM = GetSpaceTransform(
+		SolverSettings.SimulationSpace, SimulationSpaceState.ComponentTM, SimulationSpaceState.BoneRelComponentTM);
+	FTransform CollisionSpaceTM = GetSpaceTransform(
+		SolverSettings.CollisionSpace, SimulationSpaceState.ComponentTM, SimulationSpaceState.BoneRelComponentTM);
+
+	FTransform WorldSpaceTM = TM * CollisionSpaceTM;
+	return WorldSpaceTM.GetRelativeTransform(SimSpaceTM);
+}
+
+//======================================================================================================================
+void FRigPhysicsSolver::InitSimulationSpace(
+	const FTransform& ComponentTM,
+	const FTransform& BoneRelComponentTM)
+{
+	SimulationSpaceState.ComponentTM = ComponentTM;
+	SimulationSpaceState.BoneRelComponentTM = BoneRelComponentTM;
+	SimulationSpaceState.ComponentLinearVelocity = FVector::ZeroVector;
+	SimulationSpaceState.ComponentAngularVelocity = FVector::ZeroVector;
+	SimulationSpaceState.BoneLinearVelocity = FVector::ZeroVector;
+	SimulationSpaceState.BoneAngularVelocity = FVector::ZeroVector;
+}
+
+//======================================================================================================================
+FVector CalculateAngularVelocity(const FQuat& PrevQ, const FQuat& CurrentQ, float Dt)
+{
+	const FQuat DeltaQ = (CurrentQ * PrevQ.Inverse()).GetShortestArcWith(FQuat::Identity);
+	return DeltaQ.ToRotationVector() / Dt;
+}
+
+//======================================================================================================================
+// Note - don't use the space conversion functions here as the state won't have been set yet.
+FRigPhysicsSolver::FSimulationSpaceData FRigPhysicsSolver::UpdateSimulationSpaceStateAndCalculateData(const float Dt)
+{
+	URigHierarchy* Hierarchy = OwningControlRig->GetHierarchy();
+
+	if (const USceneComponent* OwningSceneComponent = OwningControlRig->GetOwningSceneComponent())
+	{
+		SimulationSpaceState.ComponentTM = OwningSceneComponent->GetComponentTransform();
+	}
+	else 
+	{
+		SimulationSpaceState.ComponentTM.SetIdentity();
+	}
+
+	if (SolverSettings.SimulationSpace == ERigPhysicsSimulationSpace::SpaceBone &&  SolverSettings.SpaceBone.IsValid())
+	{
+		SimulationSpaceState.BoneRelComponentTM = Hierarchy->GetGlobalTransform(SolverSettings.SpaceBone);
+	}
+
+	SimulationSpaceState.PrevDt = SimulationSpaceState.Dt;
+	SimulationSpaceState.Dt = Dt;
+
+	SimulationSpaceState.PrevPrevSimulationSpaceTM = SimulationSpaceState.PrevSimulationSpaceTM;
+	SimulationSpaceState.PrevSimulationSpaceTM = SimulationSpaceState.SimulationSpaceTM;
+	SimulationSpaceState.SimulationSpaceTM = GetSpaceTransform(
+		SolverSettings.SimulationSpace, SimulationSpaceState.ComponentTM, SimulationSpaceState.BoneRelComponentTM);
+
+	FSimulationSpaceData SimulationSpaceData;
+	SimulationSpaceData.Gravity = ::ConvertWorldVectorToSimSpaceNoScale(
+		SolverSettings.SimulationSpace, SolverSettings.Gravity, 
+		SimulationSpaceState.ComponentTM, SimulationSpaceState.BoneRelComponentTM);
+
+	if (SolverSettings.SimulationSpace == ERigPhysicsSimulationSpace::World)
+	{
+		// Danny TODO This is probably redundant unless we support runtime switching of the space
+		InitSimulationSpace(SimulationSpaceState.ComponentTM, SimulationSpaceState.BoneRelComponentTM);
+		SimulationSpaceData.LinearVelocity = SimulationSpaceSettings.ExternalLinearVelocity;
+		SimulationSpaceData.AngularVelocity = SimulationSpaceSettings.ExternalAngularVelocity;
+		return SimulationSpaceData;
+	}
+
+	// Don't estimate velocities if the timestep is effectively zero - just preserve them
+	if (SimulationSpaceState.Dt < SMALL_NUMBER || SimulationSpaceState.PrevDt < SMALL_NUMBER)
+	{
+		return SimulationSpaceData;
+	}
+
+	// We calculate velocities etc in world space first, and then subsequently convert them into
+	// simulation space.
+
+	// Note that the velocity/accel calculations are intended to track the world/simulation behavior
+	// - not necessarily be the most accurate calculations! For example, we could use one-sided
+	// finite difference approximations, but this wouldn't necessarily be correct.
+
+	// World-space component linear velocity and acceleration
+	const FVector PrevSpaceLinearVel = Chaos::FVec3::CalculateVelocity(
+		SimulationSpaceState.PrevPrevSimulationSpaceTM.GetTranslation(),
+		SimulationSpaceState.PrevSimulationSpaceTM.GetTranslation(), SimulationSpaceState.PrevDt);
+	SimulationSpaceData.LinearVelocity = Chaos::FVec3::CalculateVelocity(
+		SimulationSpaceState.PrevSimulationSpaceTM.GetTranslation(),
+		SimulationSpaceState.SimulationSpaceTM.GetTranslation(), SimulationSpaceState.Dt);
+	SimulationSpaceData.LinearAcceleration = 
+		(SimulationSpaceData.LinearVelocity - PrevSpaceLinearVel) / SimulationSpaceState.Dt;
+
+	// World-space component angular velocity and acceleration
+	const FVector PrevSpaceAngularVel = CalculateAngularVelocity(
+		SimulationSpaceState.PrevPrevSimulationSpaceTM.GetRotation(),
+		SimulationSpaceState.PrevSimulationSpaceTM.GetRotation(), SimulationSpaceState.PrevDt);
+	SimulationSpaceData.AngularVelocity = CalculateAngularVelocity(
+		SimulationSpaceState.PrevSimulationSpaceTM.GetRotation(),
+		SimulationSpaceState.SimulationSpaceTM.GetRotation(), SimulationSpaceState.Dt);
+	SimulationSpaceData.AngularAcceleration = 
+		(SimulationSpaceData.AngularVelocity - PrevSpaceAngularVel) / SimulationSpaceState.Dt;
+
+	// Apply Z scale
+	SimulationSpaceData.LinearVelocity.Z *= SimulationSpaceSettings.VelocityScaleZ;
+	SimulationSpaceData.LinearAcceleration.Z *= SimulationSpaceSettings.VelocityScaleZ;
+
+	// Clamped world-space motion of the simulation space
+	if (SimulationSpaceSettings.bClampLinearVelocity)
+	{
+		SimulationSpaceData.LinearVelocity = 
+			SimulationSpaceData.LinearVelocity.GetClampedToMaxSize(SimulationSpaceSettings.MaxLinearVelocity);
+	}
+	SimulationSpaceData.LinearVelocity += SimulationSpaceSettings.ExternalLinearVelocity;
+
+	if (SimulationSpaceSettings.bClampAngularVelocity)
+	{
+		SimulationSpaceData.AngularVelocity =
+		SimulationSpaceData.AngularVelocity.GetClampedToMaxSize(SimulationSpaceSettings.MaxAngularVelocity);
+	}
+	SimulationSpaceData.AngularVelocity += SimulationSpaceSettings.ExternalAngularVelocity;
+
+	if (SimulationSpaceSettings.bClampLinearAcceleration)
+	{
+		SimulationSpaceData.LinearAcceleration = 
+			SimulationSpaceData.LinearAcceleration.GetClampedToMaxSize(SimulationSpaceSettings.MaxLinearAcceleration);
+	}
+
+	if (SimulationSpaceSettings.bClampAngularAcceleration)
+	{
+		SimulationSpaceData.AngularAcceleration = 
+			SimulationSpaceData.AngularAcceleration.GetClampedToMaxSize(SimulationSpaceSettings.MaxAngularAcceleration);
+	}
+
+	// Transform world-space motion into simulation space Danny TODO note that this matches the code
+	// in RBAN, and is doing what the interface requires (i.e. movement of the space in the space of
+	// the space!). However, the resultant motion is incorrect for rotation of the space.
+	SimulationSpaceData.LinearVelocity =
+		SimulationSpaceState.SimulationSpaceTM.InverseTransformVector(SimulationSpaceData.LinearVelocity);
+	SimulationSpaceData.AngularVelocity =
+		SimulationSpaceState.SimulationSpaceTM.InverseTransformVector(SimulationSpaceData.AngularVelocity);
+	SimulationSpaceData.LinearAcceleration =
+		SimulationSpaceState.SimulationSpaceTM.InverseTransformVector(SimulationSpaceData.LinearAcceleration);
+	SimulationSpaceData.AngularAcceleration =
+		SimulationSpaceState.SimulationSpaceTM.InverseTransformVector(SimulationSpaceData.AngularAcceleration);
+
+	return SimulationSpaceData;
+}
diff --git a/Engine/Plugins/Experimental/PhysicsControl/Source/ControlRigPhysics/Private/RigPhysicsSolverVisualization.cpp b/Engine/Plugins/Experimental/PhysicsControl/Source/ControlRigPhysics/Private/RigPhysicsSolverVisualization.cpp
new file mode 100644
--- /dev/null
+++ b/Engine/Plugins/Experimental/PhysicsControl/Source/ControlRigPhysics/Private/RigPhysicsSolverVisualization.cpp
@@ -0,0 +1,343 @@
+// Copyright Epic Games, Inc. All Rights Reserved.
+
+#include "RigPhysicsSolver.h"
+
+#include "Rigs/RigHierarchy.h"
+#include "Rigs/RigHierarchyController.h"
+#include "ControlRig.h"
+
+#include "RigVMCore/RigVMDrawInterface.h"
+
+#include "Physics/ImmediatePhysics/ImmediatePhysicsActorHandle.h"
+#include "Physics/ImmediatePhysics/ImmediatePhysicsJointHandle.h"
+#include "Physics/ImmediatePhysics/ImmediatePhysicsSimulation.h"
+
+#include "Physics/Experimental/PhysScene_Chaos.h"
+#include "Chaos/ParticleHandle.h"
+#include "Chaos/PBDJointConstraints.h"
+#include "Chaos/ImplicitObject.h"
+#include "Chaos/ShapeInstance.h"
+#include "Chaos/Capsule.h"
+#include "Chaos/PBDJointConstraintUtilities.h"
+
+#include "Engine/Engine.h"
+#include "Materials/MaterialInstanceDynamic.h"
+
+#include "Math/Color.h"
+
+namespace RigPhysicsSolverDraw
+{
+	float Thickness = 0.3f;
+	FLinearColor DynamicColor = FColor::Yellow;
+	FLinearColor KinematicColor = FColor::Blue;
+	int32 Detail = 16;
+}
+
+//======================================================================================================================
+static void DrawShapes(
+	FRigVMDrawInterface*                  DI,
+	const bool                            bIsKinematic,
+	const Chaos::FGeometryParticleHandle* Particle,
+	const FTransform&                     ShapeTransform, 
+	const Chaos::FImplicitObject*         ImplicitObject,
+	const Chaos::FPerShapeData*           Shape)
+{
+	using namespace Chaos;
+
+	if (!ImplicitObject)
+	{
+		return;
+	}
+	const EImplicitObjectType PackedType = ImplicitObject->GetType(); // Type includes scaling and instancing data
+	const EImplicitObjectType InnerType = GetInnerType(ImplicitObject->GetType());
+
+	// For simplicity, we're going to assume no scaling, instancing, transforming etc
+
+	if (!ImplicitObject)
+	{
+		return;
+	}
+
+	switch (InnerType)
+	{
+	case ImplicitObjectType::Transformed:
+	{
+		const TImplicitObjectTransformed<FReal, 3>* Transformed = 
+			ImplicitObject->GetObject<TImplicitObjectTransformed<FReal, 3>>();
+		FTransform TransformedTransform = FTransform(
+			ShapeTransform.GetRotation() * Transformed->GetTransform().GetRotation(),
+			ShapeTransform.TransformPosition(Transformed->GetTransform().GetLocation()));
+		DrawShapes(DI, bIsKinematic, Particle, TransformedTransform, Transformed->GetTransformedObject(), Shape);
+		return;
+	}
+	case ImplicitObjectType::Sphere:
+	{
+		const Chaos::FSphere* Sphere = ImplicitObject->GetObject<Chaos::FSphere>();
+		DI->DrawSphere(ShapeTransform, FTransform(FVec3(Sphere->GetCenterOfMass())), Sphere->GetRadiusf(),
+			bIsKinematic ? RigPhysicsSolverDraw::KinematicColor : RigPhysicsSolverDraw::DynamicColor,
+			RigPhysicsSolverDraw::Thickness, RigPhysicsSolverDraw::Detail);
+		break;
+	}
+	case ImplicitObjectType::Box:
+	{
+		const TBox<FReal, 3>* Box = ImplicitObject->GetObject<TBox<FReal, 3>>();
+		DI->DrawBox(ShapeTransform, FTransform(FQuat::Identity, Box->GetCenter(), Box->Extents()),
+			bIsKinematic ? RigPhysicsSolverDraw::KinematicColor : RigPhysicsSolverDraw::DynamicColor,
+			RigPhysicsSolverDraw::Thickness);
+		break;
+	}
+	case ImplicitObjectType::Capsule:
+	{
+		const Chaos::FCapsule* Capsule = ImplicitObject->GetObject<Chaos::FCapsule>();
+		const FQuat Q = FRotationMatrix::MakeFromZ(Capsule->GetAxis()).ToQuat();
+		DI->DrawCapsule(ShapeTransform, FTransform(Q, Capsule->GetCenterOfMass()),
+			Capsule->GetRadiusf(), Capsule->GetHeightf(),
+			bIsKinematic ? RigPhysicsSolverDraw::KinematicColor : RigPhysicsSolverDraw::DynamicColor,
+			RigPhysicsSolverDraw::Thickness);
+		break;
+	}
+	default:
+		// If we don't know what it is, don't draw it.
+		break;
+	}
+
+}
+
+//======================================================================================================================
+static void DrawActor(
+	FRigVMDrawInterface*                  DI,
+	const FTransform&                     SpaceTransform,
+	const ImmediatePhysics::FActorHandle& ActorHandle)
+{
+	bool bIsKinematic = ActorHandle.GetIsKinematic();
+	const Chaos::FGeometryParticleHandle* GeometryParticleHandle = ActorHandle.GetParticle();
+	const FTransform ParticleTransform = GeometryParticleHandle->GetTransformXR() * SpaceTransform;
+	for (const Chaos::FShapeInstancePtr& ShapeInstance : GeometryParticleHandle->ShapeInstances())
+	{
+		const Chaos::FImplicitObject* ImplicitObject = ShapeInstance->GetGeometry();
+		DrawShapes(
+			DI, bIsKinematic, GeometryParticleHandle, 
+			ParticleTransform, ImplicitObject, ShapeInstance.Get());
+	}
+}
+
+enum
+{
+	TwistIndex = 0,
+	Swing1Index = 1,
+	Swing2Index = 2
+};
+
+// Returns a number to use for the limit for limit based on limit type
+float GetLimitAngleRadians(float LimitAngle, Chaos::EJointMotionType LimitType)
+{
+	switch (LimitType)
+	{
+	case Chaos::EJointMotionType::Free: return UE_PI;
+	case Chaos::EJointMotionType::Locked: return 0.f;
+	default: return LimitAngle;
+	}
+}
+
+
+//======================================================================================================================
+static void DrawJoint(
+	FRigVMDrawInterface*                  DI,
+	const FTransform&                     SpaceTransform,
+	const ImmediatePhysics::FJointHandle& JointHandle)
+{
+	const Chaos::FPBDJointConstraintHandle* ConstraintHandle = JointHandle.GetConstraint();
+	if (!ConstraintHandle)
+	{
+		return;
+	}
+
+	const Chaos::TVec2<const ImmediatePhysics_Chaos::FActorHandle*>& ActorHandles = JointHandle.GetActorHandles();
+	const ImmediatePhysics_Chaos::FActorHandle* ChildActor = ActorHandles[0];
+	const ImmediatePhysics_Chaos::FActorHandle* ParentActor = ActorHandles[1];
+
+	FTransform ChildActorTM = (ChildActor ? ChildActor->GetWorldTransform() : FTransform()) * SpaceTransform;
+	FTransform ParentActorTM = (ParentActor ? ParentActor->GetWorldTransform() : FTransform()) * SpaceTransform;
+
+	bool bChildIsKinematic = ChildActor ? ChildActor->GetIsKinematic() : true;
+	bool bParentIsKinematic = ParentActor ? ParentActor->GetIsKinematic() : true;
+
+	float Size = 5.0f;
+
+	const Chaos::FPBDJointSettings& JointSettings = ConstraintHandle->GetSettings();
+
+	FTransform DialFrame = JointSettings.ConnectorTransforms[0] * ChildActorTM;
+	FTransform LimitFrame = JointSettings.ConnectorTransforms[1] * ParentActorTM;
+
+	static bool bDrawAsAxes = false;
+	if (bDrawAsAxes)
+	{
+		DI->DrawAxes(LimitFrame, FTransform(), Size, RigPhysicsSolverDraw::Thickness * 2);
+		DI->DrawAxes(DialFrame, FTransform(), Size, RigPhysicsSolverDraw::Thickness);
+	}
+
+	{
+		// See FConstraintInstance::DrawConstraintImp for inspiration
+
+		// There seems to be a swap between swing1 and swing2 compared to GetSwingTwistAngles
+		FVector LimitAngleRadians(
+			GetLimitAngleRadians(JointSettings.AngularLimits[0], JointSettings.AngularMotionTypes[0]),
+			GetLimitAngleRadians(JointSettings.AngularLimits[2], JointSettings.AngularMotionTypes[2]),
+			GetLimitAngleRadians(JointSettings.AngularLimits[1], JointSettings.AngularMotionTypes[1])
+		);
+
+		Chaos::FReal TwistAngle, Swing1Angle, Swing2Angle;
+		const FQuat LimitQ = LimitFrame.GetRotation();
+		FQuat DialQ = DialFrame.GetRotation();
+		DialQ.EnforceShortestArcWith(LimitQ);
+		Chaos::FPBDJointUtilities::GetSwingTwistAngles(LimitQ, DialQ, TwistAngle, Swing1Angle, Swing2Angle);
+
+		bool bDrawViolatedLimits = true;
+
+		bool bTwistViolated = false;
+		bool bSwing1Violated = false;
+		bool bSwing2Violated = false;
+
+		if (bDrawViolatedLimits)
+		{
+			bTwistViolated = JointSettings.AngularMotionTypes[TwistIndex] == Chaos::EJointMotionType::Limited &&
+				FMath::Abs(TwistAngle) > LimitAngleRadians[0];
+
+			bSwing1Violated = JointSettings.AngularMotionTypes[Swing1Index] == Chaos::EJointMotionType::Limited &&
+				FMath::Abs(Swing1Angle) > LimitAngleRadians[1];
+
+			bSwing2Violated = JointSettings.AngularMotionTypes[Swing2Index] == Chaos::EJointMotionType::Limited &&
+				FMath::Abs(Swing2Angle) > LimitAngleRadians[2];
+		}
+
+		const bool bLockSwing1 = JointSettings.AngularMotionTypes[Swing1Index] == Chaos::EJointMotionType::Locked;
+		const bool bLockSwing2 = JointSettings.AngularMotionTypes[Swing2Index] == Chaos::EJointMotionType::Locked;
+		const bool bLockAllSwing = bLockSwing1 && bLockSwing2;
+
+		check(GEngine->ConstraintLimitMaterialX && GEngine->ConstraintLimitMaterialY && GEngine->ConstraintLimitMaterialZ);
+		static UMaterialInterface* LimitMaterialX = GEngine->ConstraintLimitMaterialX;
+		static UMaterialInterface* LimitMaterialXAxis = GEngine->ConstraintLimitMaterialXAxis;
+		static UMaterialInterface* LimitMaterialY = GEngine->ConstraintLimitMaterialY;
+		static UMaterialInterface* LimitMaterialYAxis = GEngine->ConstraintLimitMaterialYAxis;
+		static UMaterialInterface* LimitMaterialZ = GEngine->ConstraintLimitMaterialZ;
+		static UMaterialInterface* LimitMaterialZAxis = GEngine->ConstraintLimitMaterialZAxis;
+
+		float ArrowLength = Size * 1.05f; // stick out a little bit 
+
+		// If swing is limited (but not locked) - draw the swing limit cone.
+		if (!bLockAllSwing)
+		{
+			if (JointSettings.AngularMotionTypes[Swing1Index] == Chaos::EJointMotionType::Free && 
+				JointSettings.AngularMotionTypes[Swing2Index] == Chaos::EJointMotionType::Free)
+			{
+				DI->DrawSphere(LimitFrame, FTransform(), Size * 0.2f, FColor::White, 
+					RigPhysicsSolverDraw::Thickness, RigPhysicsSolverDraw::Detail);
+			}
+			else
+			{
+				FTransform ConeLimitTM = LimitFrame;
+				ConeLimitTM.SetScale3D(FVector(Size));
+				const float Swing1LimitAngle = LimitAngleRadians[Swing1Index];
+				const float Swing2LimitAngle = LimitAngleRadians[Swing2Index];
+				DI->DrawCone(
+					ConeLimitTM, FTransform(), Swing1LimitAngle, Swing2LimitAngle, RigPhysicsSolverDraw::Detail, 
+					true, FColor::Green, LimitMaterialX->GetRenderProxy(), RigPhysicsSolverDraw::Thickness);
+			}
+			// Draw the swing Dial indicator - shows the current orientation of the child frame
+			// relative to the parent frame on the swing axis.
+			// Start the arrow at the limit position as it can be confusing if there is joint separation
+			DI->DrawArrow(FTransform(LimitFrame.GetTranslation()), 
+				DialFrame.GetUnitAxis(EAxis::X) * ArrowLength,
+				DialFrame.GetUnitAxis(EAxis::Y), FColor::Red, RigPhysicsSolverDraw::Thickness);
+
+			if (bSwing1Violated || bSwing2Violated)
+			{
+				// Drawing a sphere seems a bit heavy... but it appears that if DrawPoint is used then they get
+				// culled when you get close!
+				DI->DrawSphere(FTransform(
+					DialFrame.GetRotation(),
+					LimitFrame.GetTranslation() + DialFrame.GetUnitAxis(EAxis::X) * ArrowLength),
+					FTransform(),
+					ArrowLength * 0.01f, FColor::Orange, RigPhysicsSolverDraw::Thickness * 4);
+			}
+
+		}
+
+		// Draw the twist limit
+		if (JointSettings.AngularMotionTypes[TwistIndex] != Chaos::EJointMotionType::Locked)
+		{
+			// Draw as a flat cone
+			FTransform ConeLimitTM = LimitFrame;
+			ConeLimitTM.SetScale3D(FVector(Size));
+			const float TwistLimitAngle = LimitAngleRadians[TwistIndex];
+			DI->DrawCone(ConeLimitTM, FTransform(FQuat::MakeFromRotationVector(FVector(0, 1, 0) * -HALF_PI)),
+				TwistLimitAngle, 0.0f, RigPhysicsSolverDraw::Detail, true, FColor::Green,
+				LimitMaterialYAxis->GetRenderProxy(), RigPhysicsSolverDraw::Thickness);
+
+			// Draw the twist Dial indicator - shows the current orientation of the child frame
+			// relative to the parent frame on the twist axis.
+			FQuat Rot(LimitFrame.GetUnitAxis(EAxis::X), TwistAngle);
+			FVector TwistArrow = Rot * LimitFrame.GetUnitAxis(EAxis::Z);
+			FVector TwistSideDir = FVector::CrossProduct(LimitFrame.GetUnitAxis(EAxis::X), TwistArrow);
+
+			DI->DrawArrow(FTransform(LimitFrame.GetTranslation()),
+				TwistArrow * ArrowLength,
+				TwistSideDir, FColor::Blue, RigPhysicsSolverDraw::Thickness);
+
+			if (bTwistViolated)
+			{
+				DI->DrawSphere(FTransform(
+					DialFrame.GetRotation(),
+					LimitFrame.GetTranslation() + TwistArrow * ArrowLength),
+					FTransform(),
+					ArrowLength * 0.01f, FColor::Orange, RigPhysicsSolverDraw::Thickness * 4);
+			}
+		}
+	}
+}
+
+//======================================================================================================================
+void FRigPhysicsSolver::Draw(FRigVMDrawInterface* DI, const UWorld* DebugWorld) const
+{
+#if CHAOS_DEBUG_DRAW
+	// Danny TODO enable/disable chaos debug draw with a cvar or flag
+	if (DebugWorld && DebugWorld->GetPhysicsScene() && DebugWorld->GetPhysicsScene()->GetDebugDrawScene())
+	{
+		Simulation->SetDebugDrawScene(FString("ControlRig"), DebugWorld->GetPhysicsScene()->GetDebugDrawScene());
+		Simulation->DebugDraw();
+	}
+#endif
+
+	const URigHierarchy* Hierarchy = OwningControlRig->GetHierarchy();
+	if (!DI || !Hierarchy)
+	{
+		return;
+	}
+
+	// All rendering is done relative to the component, so convert the sim space (identity) into the
+	// component space.
+	const FTransform SpaceTransform = ConvertSimSpaceTransformToComponentSpace(FTransform());
+
+	for (const ImmediatePhysics::FActorHandle* ActorHandle : AdditionalActorHandles)
+	{
+		DrawActor(DI, SpaceTransform, *ActorHandle);
+	}
+
+	for (const TPair<FRigComponentKey, FRigBodyRecord>& BodyRecordPair : BodyRecords)
+	{
+		const FRigBodyRecord& Record = BodyRecordPair.Value;
+		const ImmediatePhysics::FActorHandle* ActorHandle = Record.ActorHandle;
+		const ImmediatePhysics::FJointHandle* JointHandle = Record.JointHandle;
+
+		if (JointHandle)
+		{
+			DrawJoint(DI, SpaceTransform, *JointHandle);
+		}
+
+		if (ActorHandle)
+		{
+			DrawActor(DI, SpaceTransform, *ActorHandle);
+		}
+	}
+}
+
diff --git a/Engine/Plugins/Experimental/PhysicsControl/Source/ControlRigPhysics/Public/ControlRigPhysicsModule.h b/Engine/Plugins/Experimental/PhysicsControl/Source/ControlRigPhysics/Public/ControlRigPhysicsModule.h
new file mode 100644
--- /dev/null
+++ b/Engine/Plugins/Experimental/PhysicsControl/Source/ControlRigPhysics/Public/ControlRigPhysicsModule.h
@@ -0,0 +1,14 @@
+// Copyright Epic Games, Inc. All Rights Reserved.
+
+#pragma once
+
+#include "Modules/ModuleManager.h"
+
+class FControlRigPhysicsModule : public IModuleInterface
+{
+public:
+
+	/** IModuleInterface implementation */
+	virtual void StartupModule() override;
+	virtual void ShutdownModule() override;
+};
diff --git a/Engine/Plugins/Experimental/PhysicsControl/Source/ControlRigPhysics/Public/RigPhysicsComponent.h b/Engine/Plugins/Experimental/PhysicsControl/Source/ControlRigPhysics/Public/RigPhysicsComponent.h
new file mode 100644
--- /dev/null
+++ b/Engine/Plugins/Experimental/PhysicsControl/Source/ControlRigPhysics/Public/RigPhysicsComponent.h
@@ -0,0 +1,50 @@
+// Copyright Epic Games, Inc. All Rights Reserved.
+
+#pragma once
+
+#include "RigPhysicsData.h"
+
+#include "PhysicsControlData.h"
+
+#include "Rigs/RigPhysics.h"
+#include "Rigs/RigHierarchyDefines.h"
+
+#include "RigPhysicsComponent.generated.h"
+
+struct FRigVMExecuteContext;
+struct FRigBaseElement;
+struct FRigControlElement;
+class URigHierarchy;
+class FRigElementKeyRedirector;
+
+/**
+ * A component that can be added to a joint/element that defines how a physical body can be "attached" to it.
+ * The body supports dynamic movement, collision, and a physics joint with this body's parent in the hierarchy.
+ */
+USTRUCT(BlueprintType)
+struct CONTROLRIGPHYSICS_API FRigPhysicsComponent : public FRigBaseComponent
+{
+public:
+
+	GENERATED_BODY()
+	DECLARE_RIG_COMPONENT_METHODS(FRigPhysicsComponent)
+
+	UPROPERTY(BlueprintReadOnly, EditAnywhere, Category = Physics, meta = (ShowOnlyInnerProperties))
+	FRigPhysicsBodySolverSettings BodySolverSettings;
+
+	UPROPERTY(BlueprintReadOnly, EditAnywhere, Category = Physics, meta = (ShowOnlyInnerProperties))
+	FRigPhysicsDynamics Dynamics;
+
+	UPROPERTY(BlueprintReadOnly, EditAnywhere, Category = Physics, meta = (ShowOnlyInnerProperties))
+	FRigPhysicsCollision Collision;
+
+	UPROPERTY(BlueprintReadOnly, EditAnywhere, Category = Physics, meta = (ShowOnlyInnerProperties))
+	FRigPhysicsJoint Joint;
+
+	UPROPERTY(BlueprintReadOnly, EditAnywhere, Category = Physics)
+	FPhysicsControlModifierData BodyData;
+
+	virtual void Save(FArchive& A) override;
+	virtual void Load(FArchive& Ar) override;
+};
+
diff --git a/Engine/Plugins/Experimental/PhysicsControl/Source/ControlRigPhysics/Public/RigPhysicsControlComponent.h b/Engine/Plugins/Experimental/PhysicsControl/Source/ControlRigPhysics/Public/RigPhysicsControlComponent.h
new file mode 100644
--- /dev/null
+++ b/Engine/Plugins/Experimental/PhysicsControl/Source/ControlRigPhysics/Public/RigPhysicsControlComponent.h
@@ -0,0 +1,36 @@
+// Copyright Epic Games, Inc. All Rights Reserved.
+
+#pragma once
+
+#include "PhysicsControlData.h"
+
+#include "Rigs/RigPhysics.h"
+#include "RigPhysicsControlComponent.generated.h"
+
+/**
+ * A component that can be added to hierarchy elements (joints) to add the data required to control 
+ * the simulation of them
+ */
+USTRUCT(BlueprintType)
+struct CONTROLRIGPHYSICS_API FRigPhysicsControlComponent : public FRigBaseComponent
+{
+public:
+
+	GENERATED_BODY()
+	DECLARE_RIG_COMPONENT_METHODS(FRigPhysicsControlComponent)
+
+	UPROPERTY(BlueprintReadOnly, EditAnywhere, Category = Control)
+	FRigComponentKey ParentBody;
+
+	/** Describes the initial strength etc of the new control */
+	UPROPERTY(BlueprintReadWrite, EditAnywhere, Category = Control)
+	FPhysicsControlData ControlData;
+
+	/** Describes the initial target for the new control */
+	UPROPERTY(BlueprintReadWrite, EditAnywhere, Category = Control)
+	FPhysicsControlTarget ControlTarget;
+
+	virtual void Save(FArchive& A) override;
+	virtual void Load(FArchive& Ar) override;
+};
+
diff --git a/Engine/Plugins/Experimental/PhysicsControl/Source/ControlRigPhysics/Public/RigPhysicsControlExecution.h b/Engine/Plugins/Experimental/PhysicsControl/Source/ControlRigPhysics/Public/RigPhysicsControlExecution.h
new file mode 100644
--- /dev/null
+++ b/Engine/Plugins/Experimental/PhysicsControl/Source/ControlRigPhysics/Public/RigPhysicsControlExecution.h
@@ -0,0 +1,65 @@
+// Copyright Epic Games, Inc. All Rights Reserved.
+
+#pragma once
+
+#include "RigPhysicsComponent.h"
+#include "RigPhysicsSolver.h"
+#include "PhysicsControlData.h"
+
+#include "Units/Execution/RigUnit_DynamicHierarchy.h"
+
+#include "RigPhysicsControlExecution.generated.h"
+
+/**
+ * Adds a new physics control as a component on the owner element.
+ * Note: This node only runs as part of the construction event.
+ */
+USTRUCT(meta = (DisplayName = "Add Physics Control Component", Keywords = "Construction,Create,New,Joint,Control,Physics", Varying))
+struct CONTROLRIGPHYSICS_API FRigUnit_AddPhysicsControl : public FRigUnitMutable
+{
+	GENERATED_BODY()
+
+	FRigUnit_AddPhysicsControl()
+	{
+		Owner = FRigElementKey();
+		Result = FRigComponentKey();
+	}
+
+	RIGVM_METHOD()
+	virtual void Execute() override;
+
+	/**
+	 * The owner of the newly created component (must be set/valid)
+	 */
+	UPROPERTY(meta = (Input, BoneName))
+	FRigElementKey Owner;
+
+	UPROPERTY(meta = (Output))
+	FRigComponentKey Result;
+
+	/*
+	 * The solver to relate this new physics element to
+	 */
+	UPROPERTY(meta = (Input))
+	FRigPhysicsSolverID SolverID;
+
+	UPROPERTY(meta = (Input))
+	FRigComponentKey ParentBody;
+
+	/** Describes the initial strength etc of the new control */
+	UPROPERTY(meta = (Input))
+	FPhysicsControlData ControlData;
+
+	/** Describes the initial target for the new control */
+	UPROPERTY(meta = (Input))
+	FPhysicsControlTarget ControlTarget;
+
+	/** Which set to include the control in (optional). Note that it automatically gets added to the set "All" */
+	UPROPERTY(meta = (Input))
+	FName Set;
+
+	/** NamePrefix Name of the control component. */
+	UPROPERTY(meta = (Input))
+	FName Name = TEXT("Physics Control");
+};
+
diff --git a/Engine/Plugins/Experimental/PhysicsControl/Source/ControlRigPhysics/Public/RigPhysicsData.h b/Engine/Plugins/Experimental/PhysicsControl/Source/ControlRigPhysics/Public/RigPhysicsData.h
new file mode 100644
--- /dev/null
+++ b/Engine/Plugins/Experimental/PhysicsControl/Source/ControlRigPhysics/Public/RigPhysicsData.h
@@ -0,0 +1,536 @@
+// Copyright Epic Games, Inc. All Rights Reserved.
+
+#pragma once
+
+#include "PhysicsControlData.h"
+
+#include "Rigs/RigPhysics.h"
+#include "Rigs/RigHierarchyDefines.h"
+
+#include "RigPhysicsData.generated.h"
+
+struct FRigVMExecuteContext;
+struct FRigBaseElement;
+struct FRigControlElement;
+class URigHierarchy;
+class FRigElementKeyRedirector;
+
+/** 
+ * Specifies what space a kinematic target is defined in 
+ * Danny TODO this is much better than the similar mechanism in PCC and RBWC (which has a 
+ * bUseSkeletalAnimation) - switch them to use this.
+ */
+UENUM()
+enum class ERigPhysicsKinematicTargetSpace : uint8
+{
+	/** Target is in world space */
+	World,
+	/** Target is in the space of the component */
+	Component,
+	/**
+	 * Target is specified relative to the bone that controls the (kinematic) body, with the offset applied 
+	 * in the space of that body. 
+	 */
+	OffsetInBoneSpace,
+	/**
+	 * Target is specified relative to the bone that controls the (kinematic) body, with the offset applied
+	 * in world space. 
+	 */
+	OffsetInWorldSpace,
+	/**
+	 * Target is specified relative to the bone that controls the (kinematic) body, with the offset applied
+	 * in component space. 
+	 */
+	OffsetInComponentSpace,
+};
+
+/** Specifies what space the simulation should run in */
+UENUM()
+enum class ERigPhysicsSimulationSpace : uint8
+{
+	/** Simulate in world space. Moving the skeletal mesh will generate velocity changes */
+	World,
+	/** 
+	 * Simulate in component space. Moving the entire skeletal mesh will have no affect on velocities 
+	 * unless this is explicitly enabled in the simulation space settings */
+	Component,
+	/**
+	 * Simulate in the space of the specified bone. Moving the entire skeletal mesh and/or the Bone
+	 * will have no affect on velocities unless this is explicitly enabled in the simulation space settings.
+	 */
+	SpaceBone,
+};
+
+/**
+ * Properties common to all collision shape types
+ */
+USTRUCT(BlueprintType)
+struct CONTROLRIGPHYSICS_API FRigPhysicsCollisionShape
+{
+	GENERATED_BODY()
+
+	/** Offset used when generating contact points. This allows you to smooth out
+	the Minkowski sum by radius R. Useful for making objects slide smoothly
+	on top of irregularities  */
+	UPROPERTY(Category = Shape, EditAnywhere)
+	float RestOffset = 0.0f;
+
+	UPROPERTY(Category = Shape, EditAnywhere)
+	FName Name;
+
+	/** True if this shape should contribute to the overall mass of the body it
+		belongs to. This lets you create extra collision volumes which do not affect
+		the mass properties of an object. */
+	UPROPERTY(Category = Shape, EditAnywhere)
+	bool bContributeToMass = true;
+};
+
+/**
+ * Box collision
+ */
+USTRUCT(BlueprintType)
+struct CONTROLRIGPHYSICS_API FRigPhysicsCollisionBox : public FRigPhysicsCollisionShape
+{
+	GENERATED_BODY()
+
+	friend FArchive& operator <<(FArchive& Ar, FRigPhysicsCollisionBox& Data)
+	{
+		Ar << Data.TM;
+		Ar << Data.Extents;
+		return Ar;
+	}
+
+	UPROPERTY(BlueprintReadWrite, EditAnywhere, Category = Shapes)
+	FTransform TM;
+
+	// These are the full extents of the box in each axis
+	UPROPERTY(BlueprintReadWrite, EditAnywhere, Category = Shapes)
+	FVector Extents;
+};
+
+/**
+ * Sphere collision
+ */
+USTRUCT(BlueprintType)
+struct CONTROLRIGPHYSICS_API FRigPhysicsCollisionSphere : public FRigPhysicsCollisionShape
+{
+	GENERATED_BODY()
+
+	friend FArchive& operator <<(FArchive& Ar, FRigPhysicsCollisionSphere& Data)
+	{
+		Ar << Data.TM;
+		Ar << Data.Radius;
+		return Ar;
+	}
+
+	UPROPERTY(BlueprintReadWrite, EditAnywhere, Category = Shapes)
+	FTransform TM;
+
+	UPROPERTY(BlueprintReadWrite, EditAnywhere, Category = Shapes)
+	float Radius;
+};
+
+/**
+ * Capsule collision
+ */
+USTRUCT(BlueprintType)
+struct CONTROLRIGPHYSICS_API FRigPhysicsCollisionCapsule : public FRigPhysicsCollisionShape
+{
+	GENERATED_BODY()
+
+	friend FArchive& operator <<(FArchive& Ar, FRigPhysicsCollisionCapsule& Data)
+	{
+		Ar << Data.TM;
+		Ar << Data.Radius;
+		Ar << Data.Length;
+		return Ar;
+	}
+
+	UPROPERTY(BlueprintReadWrite, EditAnywhere, Category = Shapes)
+	FTransform TM;
+
+	UPROPERTY(BlueprintReadWrite, EditAnywhere, Category = Shapes)
+	float Radius;
+
+	// This is the length of the core part of the capsule. The total length will be Length + 2 * Radius
+	UPROPERTY(BlueprintReadWrite, EditAnywhere, Category = Shapes)
+	float Length;
+};
+
+/**
+ * Collection of shapes that define the collision and (optionally) mass distribution of the body
+ */
+USTRUCT(BlueprintType)
+struct CONTROLRIGPHYSICS_API FRigPhysicsCollision
+{
+	GENERATED_BODY()
+
+	friend FArchive& operator <<(FArchive& Ar, FRigPhysicsCollision& Data)
+	{
+		Ar << Data.Boxes;
+		Ar << Data.Spheres;
+		Ar << Data.Capsules;
+		return Ar;
+	}
+
+	UPROPERTY(BlueprintReadWrite, EditAnywhere, Category = Shapes)
+	TArray<FRigPhysicsCollisionBox> Boxes;
+
+	UPROPERTY(BlueprintReadWrite, EditAnywhere, Category = Shapes)
+	TArray<FRigPhysicsCollisionSphere> Spheres;
+
+	UPROPERTY(BlueprintReadWrite, EditAnywhere, Category = Shapes)
+	TArray<FRigPhysicsCollisionCapsule> Capsules;
+
+	bool IsEmpty() const { return Boxes.IsEmpty() && Spheres.IsEmpty() && Capsules.IsEmpty(); }
+};
+
+/**
+ * These settings are used when the simulation is not done in world space. They allow the 
+ * movement of the simulation space (e.g. the component) to be added to the local simulation in a controlled way.
+ */
+USTRUCT(BlueprintType)
+struct FRigPhysicsSimulationSpaceSettings
+{
+	GENERATED_BODY()
+
+	// Overall multiplier on the effects of simulation space movement on the simulation
+	UPROPERTY(EditAnywhere, BlueprintReadWrite, Category = Settings)
+	float SpaceMovementAmount = 1.0f;
+
+	// Multiplier on the components of vertical velocity of the simulation space that is passed to the simulation.
+	// Usually from 0.0 to 1.0 to reduce the effects of jumping and crouching on the simulation, but
+	// it can be higher than 1.0 if you need to exaggerate this motion for some reason.
+	UPROPERTY(EditAnywhere, BlueprintReadWrite, Category = Settings)
+	float VelocityScaleZ = 1.0f;
+
+	UPROPERTY(EditAnywhere, BlueprintReadWrite, Category = Settings, meta = (InlineEditConditionToggle))
+	bool bClampLinearVelocity = false;
+
+	// A clamp on the effective world-space velocity that is passed to the simulation. The default
+	// value effectively means "unlimited". It is not usually required to change this but you would
+	// reduce this to limit the effects of drag on the bodies in the simulation (if you have bodies
+	// that have LinearDrag set to non-zero in the physics asset). Expected values in this case
+	// would be somewhat less than the usual velocities of your object which is commonly a few
+	// hundred for a character.
+	UPROPERTY(EditAnywhere, BlueprintReadWrite, Category = Settings, meta = (ClampMin = "0.0", EditCondition = "bClampLinearVelocity"))
+	float MaxLinearVelocity = 10000;
+
+	UPROPERTY(EditAnywhere, BlueprintReadWrite, Category = Settings, meta = (InlineEditConditionToggle))
+	bool bClampAngularVelocity = false;
+
+	// A clamp on the effective world-space angular velocity that is passed to the simulation. Units
+	// are radian/s, so a value of about 6.0 is one rotation per second. The default value
+	// effectively means "unlimited". You would reduce this (and MaxAngularAcceleration) to limit
+	// how much bodies "fly out" when the actor spins on the spot. This is especially useful if you
+	// have characters than can rotate very quickly and you would probably want values around or
+	// less than 10 in this case.
+	UPROPERTY(EditAnywhere, BlueprintReadWrite, Category = Settings, meta = (ClampMin = "0.0", EditCondition = "bClampAngularVelocity"))
+	float MaxAngularVelocity = 10000;
+
+	UPROPERTY(EditAnywhere, BlueprintReadWrite, Category = Settings, meta = (InlineEditConditionToggle))
+	bool bClampLinearAcceleration = false;
+
+	// A clamp on the effective world-space acceleration that is passed to the simulation. The
+	// default value effectively means "unlimited". This property is used to stop the bodies of the
+	// simulation flying out when suddenly changing linear speed. It is useful when you have
+	// characters than can changes from stationary to running very quickly such as in an FPS. A
+	// common value for a character might be in the few hundreds.
+	UPROPERTY(EditAnywhere, BlueprintReadWrite, Category = Settings, meta = (ClampMin = "0.0", EditCondition = "bClampLinearAcceleration"))
+	float MaxLinearAcceleration = 10000;
+
+	UPROPERTY(EditAnywhere, BlueprintReadWrite, Category = Settings, meta = (InlineEditConditionToggle))
+	bool bClampAngularAcceleration = false;
+
+	// A clamp on the effective world-space angular acceleration that is passed to the simulation.
+	// Units are radian/s/s. The default value effectively means "unlimited". This has a similar
+	// effect to MaxAngularVelocity, except that it is related to the flying out of bodies when the
+	// rotation speed suddenly changes. A typical value for a character might be around 100.
+	UPROPERTY(EditAnywhere, BlueprintReadWrite, Category = Settings, meta = (ClampMin = "0.0", EditCondition = "bClampAngularAcceleration"))
+	float MaxAngularAcceleration = 10000;
+
+	// Additional linear drag applied to every body in addition to linear drag specified on them in
+	// the physics asset. When combined with ExternalLinearVelocity, this can be used to add a
+	// temporary wind-blown effect without having to tune linear drag on all the bodies in the
+	// physics asset. The result is that each body has a force equal to -ExternalLinearDragV *
+	// ExternalLinearVelocity applied to it, in addition to all other forces. The vector is in
+	// simulation local space.
+	UPROPERTY(EditAnywhere, BlueprintReadWrite, Category = Settings)
+	FVector ExternalLinearDragV = FVector::ZeroVector;
+
+	// Additional velocity that is added to the component velocity so the simulation acts as if the
+	// actor is moving at speed, even when stationary. The vector is in world space. This could be
+	// used for wind effects etc. Typical values are similar to the velocity of the object or
+	// effect, and usually around or less than 1000 for characters/wind.
+	UPROPERTY(EditAnywhere, BlueprintReadWrite, Category = Settings)
+	FVector ExternalLinearVelocity = FVector::ZeroVector;
+
+	// Additional angular velocity that is added to the component angular velocity. This can be used
+	// to make the simulation act as if the actor is rotating even when it is not. E.g., to apply
+	// physics to a character on a podium as the camera rotates around it, to emulate the podium
+	// itself rotating. Vector is in world space. Units are rad/s.
+	UPROPERTY(EditAnywhere, BlueprintReadWrite, Category = Settings)
+	FVector ExternalAngularVelocity = FVector::ZeroVector;
+};
+
+/**
+ * Configuration/settings for the physics solver
+ */
+USTRUCT(BlueprintType)
+struct CONTROLRIGPHYSICS_API FRigPhysicsSolverSettings
+{
+	GENERATED_BODY()
+
+	UPROPERTY(BlueprintReadOnly, EditAnywhere, Category = SolverSettings)
+	ERigPhysicsSimulationSpace SimulationSpace = ERigPhysicsSimulationSpace::Component;
+
+	// The space in which collision shapes are defined
+	UPROPERTY(BlueprintReadOnly, EditAnywhere, Category = SolverSettings)
+	ERigPhysicsSimulationSpace CollisionSpace = ERigPhysicsSimulationSpace::Component;
+
+	/** The bone to use for the collision or simulation space (only relevant if one is set to BoneSpace) */
+	UPROPERTY(BlueprintReadOnly, EditAnywhere, Category = SolverSettings)
+	FRigElementKey SpaceBone;
+
+	// The collision shapes defined in the collision space - e.g. for representing a ground etc
+	UPROPERTY(BlueprintReadOnly, EditAnywhere, Category = SolverSettings, meta = (ShowOnlyInnerProperties))
+	FRigPhysicsCollision Collision;
+
+	UPROPERTY(EditAnywhere, BlueprintReadWrite, Category = SolverSettings)
+	FVector Gravity = FVector(0, 0, -981.0);
+
+	/**
+	 * The number of position iterations to run. The position solve is responsible for de-penetration.
+	 * Increasing this will improve simulation stability, but increase the cost.
+	*/
+	UPROPERTY(EditAnywhere, BlueprintReadWrite, Category = SolverSettings, meta = (ClampMin = 0))
+	int32 PositionIterations = 6;
+
+	/**
+	 * The number of velocity iterations to run. The velocity solve is responsible for restitution (bounce) and friction.
+	 * This should usually be 1, but could be 0 if you don't care about friction and restitution.
+	*/
+	UPROPERTY(EditAnywhere, BlueprintReadWrite, Category = SolverSettings, meta = (ClampMin = 0))
+	int32 VelocityIterations = 1;
+
+	/**
+	 * The number of projection iterations to run. The projection phase is a final pass over the constraints, applying
+	 * a semi-physical correction to any joint errors remaining after the position and velocity solves. It can be
+	 * very helpful to stabilize joint chains, but can cause issues with collision response. The projection magnitude
+	 * can be controlled per-constraint in the constraint settings (assuming ProjectionIteration is not zero).
+	 * This should be left as 1 in almost all cases.
+	*/
+	UPROPERTY(EditAnywhere, BlueprintReadWrite, Category = SolverSettings, meta = (ClampMin = 0))
+	int32 ProjectionIterations = 1;
+
+	/**
+	 * When bodies are penetrating, this is the maximum velocity delta that can be applied in one frame.
+	*/
+	UPROPERTY(EditAnywhere, BlueprintReadWrite, Category = SolverSettings, meta = (ClampMin = 0))
+	float MaxDepenetrationVelocity = 0.0f;
+
+	/**
+	 * The recommended fixed timestep for the RBAN solver. Set to 0 to run with variable timestep (default).
+	 * NOTE: If this value is non-zero and less than the current frame time, the simulation will step multiple times
+	 * which increases the cost.
+	 */
+	UPROPERTY(EditAnywhere, BlueprintReadWrite, Category = SolverSettings, meta = (ClampMin = 0))
+	float FixedTimeStep = 0.0f;
+
+	// The maximum number of solver steps that can be made
+	UPROPERTY(EditAnywhere, BlueprintReadWrite, Category = SolverSettings, meta = (ClampMin = 1))
+	int32 MaxTimeSteps = 10;
+
+	/**
+	 * If a variable timestep is in use, this is the maximum delta time that can be used. If this is smaller than the
+	 * requested total delta time then multiple steps will be used.
+	 */
+	UPROPERTY(EditAnywhere, BlueprintReadWrite, Category = SolverSettings, meta = (ClampMin = 0))
+	float MaxDeltaTime = 1.f / 30.f;
+
+	/**
+	 * Whether to use the linear or non-linear solver for RBAN Joints. The linear solver is significantly cheaper than
+	 * the non-linear solver when you are running multiple iterations, but is more likely to suffer from jitter.
+	 * In general you should try to use the linear solver and increase the PositionIterations to improve stability if
+	 * possible, only using the non-linear solver as a last resort.
+	*/
+	UPROPERTY(EditAnywhere, BlueprintReadWrite, Category = SolverSettings)
+	bool bUseLinearJointSolver = true;
+
+	/**
+	 * When solving joints, whether to solve the positions last (as opposed to the orientations)
+	 */
+	UPROPERTY(EditAnywhere, BlueprintReadWrite, Category = SolverSettings)
+	bool bSolveJointPositionsLast = true;
+
+	/**
+	 * It enables the use of multi-point contact manifolds, which are created only once at the start of each tick.
+	 * When disabled, a single-point contact is generated in each solver iteration which is more expensive.
+	*/
+	UPROPERTY(EditAnywhere, BlueprintReadWrite, Category = CollisionSettings)
+	bool bUseManifolds = true;
+};
+
+/**
+ * Properties of a body affecting its dynamics (movement without regard to collision).
+ */
+USTRUCT(BlueprintType)
+struct CONTROLRIGPHYSICS_API FRigPhysicsDynamics
+{
+	GENERATED_BODY()
+
+	friend FArchive& operator <<(FArchive& Ar, FRigPhysicsDynamics& Data)
+	{
+		Ar << Data.Mass;
+		Ar << Data.bOverrideCentreOfMass;
+		Ar << Data.CentreOfMassOverride;
+		Ar << Data.bOverrideMomentsOfInertia;
+		Ar << Data.MomentsOfInertiaOverride;
+		return Ar;
+	}
+
+	// The total mass of this body.
+	UPROPERTY(BlueprintReadWrite, EditAnywhere, Category = Dynamics)
+	float Mass = 1.0f;
+
+	// Use this to specify the location of the centre of mass
+	UPROPERTY(BlueprintReadWrite, EditAnywhere, Category = Dynamics)
+	bool bOverrideCentreOfMass = false;
+
+	UPROPERTY(BlueprintReadWrite, EditAnywhere, Category = Dynamics, meta = (EditCondition = "bOverrideCentreOfMass"))
+	FVector CentreOfMassOverride;
+
+	// Use this to specify the moments of inertia
+	UPROPERTY(BlueprintReadWrite, EditAnywhere, Category = Dynamics)
+	bool bOverrideMomentsOfInertia = false;
+
+	UPROPERTY(BlueprintReadWrite, EditAnywhere, Category = Dynamics, meta = (EditCondition = "bOverrideMomentsOfInertia"))
+	FVector MomentsOfInertiaOverride = FVector(1.0f, 1.0f, 1.0f);
+};
+
+/**
+ * This represents an "articulation" - a general purpose "character joint" consisting of
+ *   * A linear limit (a linear constraint), which is normally used to pin two bones together
+ *   * An angular limit (an angular constraint), which is normally used to allow a limited range of rotational movement.
+ * 
+ * The joint is defined by a parent frame, which is attached to the parent body, and a child frame which is 
+ * attached to the child body. If the limits are all zero, then these two bodies will be held in a fixed pose. 
+ * Normally the angular limit will be relaxed to allow rotational movement around the joint position.
+ * 
+ * Most bodies in a character will have one Physics Joint which connects them to their parent in the hierarchy. 
+ * However
+ *   * The root of the hierarchy will not have an enabled Physics Joint
+ *   * There may be times when a body has additional Physics Joints.
+ */
+USTRUCT(BlueprintType)
+struct CONTROLRIGPHYSICS_API FRigPhysicsJoint
+{
+	GENERATED_BODY()
+
+	friend FArchive& operator <<(FArchive& Ar, FRigPhysicsJoint& Data)
+	{
+		Ar << Data.bEnable;
+		Ar << Data.ParentBody;
+		Ar << Data.bAutoCalculateParentOffset;
+		Ar << Data.bAutoCalculateChildOffset;
+		Ar << Data.ExtraParentOffset;
+		Ar << Data.ExtraChildOffset;
+		Ar << Data.LinearLimit;
+		Ar << Data.AngularLimit;
+		Ar << Data.bDisableCollision;
+		Ar << Data.LinearProjectionAmount;
+		Ar << Data.AngularProjectionAmount;
+		Ar << Data.ParentInverseMassScale;
+		return Ar;
+	}
+
+	UPROPERTY(BlueprintReadWrite, EditAnywhere, Category = PhysicsJoint)
+	bool bEnable = true;
+
+	// The body element that will be used as the parent for the physics joint (the child will always
+	// be whatever body this physics joint is owned by. If it is unset, then a suitable body will be
+	// searched for (by assuming this body is a component of a bone which is the child of a parent
+	// bone which has a body. If it is otherwise invalid, or the parent bone is not physical, then
+	// no physics joint will be created.
+	UPROPERTY(BlueprintReadOnly, EditAnywhere, Category = PhysicsJoint)
+	FRigComponentKey ParentBody;
+
+	// The auto-calculated offset of the parent frame from the parent body places it at the 
+	// location of the child body relative to the parent body in the initial pose.
+	UPROPERTY(BlueprintReadWrite, EditAnywhere, Category = PhysicsJoint)
+	bool bAutoCalculateParentOffset = true;
+
+	// The parent frame offset applied after any auto-calculation
+	UPROPERTY(BlueprintReadWrite, EditAnywhere, Category = PhysicsJoint)
+	FTransform ExtraParentOffset;
+
+	// The auto-calculated offset of the child frame is for it to be co-located with the child body.
+	UPROPERTY(BlueprintReadWrite, EditAnywhere, Category = PhysicsJoint)
+	bool bAutoCalculateChildOffset = true;
+
+	// The child frame offset applied after any auto-calculation
+	UPROPERTY(BlueprintReadWrite, EditAnywhere, Category = PhysicsJoint)
+	FTransform ExtraChildOffset;
+
+	UPROPERTY(BlueprintReadWrite, EditAnywhere, Category = PhysicsJoint)
+	FVector LinearLimit = FVector::ZeroVector;
+
+	// Angular limit in degrees (twist, swing1, swing2)
+	UPROPERTY(BlueprintReadWrite, EditAnywhere, Category = PhysicsJoint)
+	FVector AngularLimit = FVector::ZeroVector;
+
+	// Disable collisions between the parent and child bodies
+	UPROPERTY(BlueprintReadWrite, EditAnywhere, Category = PhysicsJoint)
+	bool bDisableCollision = true;
+
+	// The amount of linear projection to reduce joint separation when the linear constraint is
+	// locked or has a hard limit.
+	UPROPERTY(BlueprintReadWrite, EditAnywhere, Category = PhysicsJoint, meta = (ClampMin = 0, ClampMax = 1))
+	float LinearProjectionAmount = 1.0f;
+
+	// The amount of angular projection to reduce joint separation when the angular constraint is
+	// locked or has a hard limit. Note that projecting back to angular limits will often break the
+	// position projection, so it is normally better to disable this.
+	UPROPERTY(BlueprintReadWrite, EditAnywhere, Category = PhysicsJoint, meta = (ClampMin = 0, ClampMax = 1))
+	float AngularProjectionAmount = 0.0f;
+
+	// As this is reduced to 0, the parent body becomes less affected by the motion of the child
+	// body, making the joint behave in one direction.
+	UPROPERTY(BlueprintReadWrite, EditAnywhere, Category = PhysicsJoint, meta = (ClampMin = 0))
+	float ParentInverseMassScale = 1.0f;
+};
+
+/**
+ * Basic settings of the body in relation to the solver that are not covered elsewhere
+ */
+USTRUCT(BlueprintType)
+struct CONTROLRIGPHYSICS_API FRigPhysicsBodySolverSettings
+{
+	GENERATED_BODY()
+
+	FRigPhysicsBodySolverSettings(
+		FRigPhysicsSolverID InSolverID = FRigPhysicsSolverID(), FRigElementKey InTargetBone = FRigElementKey())
+		: SolverID(InSolverID), TargetBone(InTargetBone) {}
+
+	friend FArchive& operator <<(FArchive& Ar, FRigPhysicsBodySolverSettings& Data)
+	{
+		Ar << Data.SolverID;
+		Ar << Data.TargetBone;
+		Ar << Data.SourceBone;
+		return Ar;
+	}
+
+	UPROPERTY(BlueprintReadWrite, EditAnywhere, Category = Solver)
+	FRigPhysicsSolverID SolverID;
+
+	/// The bone that is used to initialize physics, as well as what to track when the body is set to be kinematic.
+	/// Note that if this is unset, the it will default to the parent of our owner.
+	UPROPERTY(BlueprintReadWrite, EditAnywhere, Category = Solver)
+	FRigElementKey SourceBone;
+
+	/// The bone that is written to following simulation. Note that if this is unset, the it will
+	/// default to the parent of our owner.
+	UPROPERTY(BlueprintReadWrite, EditAnywhere, Category = Solver)
+	FRigElementKey TargetBone;
+};
+
diff --git a/Engine/Plugins/Experimental/PhysicsControl/Source/ControlRigPhysics/Public/RigPhysicsExecution.h b/Engine/Plugins/Experimental/PhysicsControl/Source/ControlRigPhysics/Public/RigPhysicsExecution.h
new file mode 100644
--- /dev/null
+++ b/Engine/Plugins/Experimental/PhysicsControl/Source/ControlRigPhysics/Public/RigPhysicsExecution.h
@@ -0,0 +1,356 @@
+// Copyright Epic Games, Inc. All Rights Reserved.
+
+#pragma once
+
+#include "RigPhysicsComponent.h"
+#include "RigPhysicsSolver.h"
+#include "PhysicsControlData.h"
+
+#include "Units/Execution/RigUnit_DynamicHierarchy.h"
+
+#include "RigPhysicsExecution.generated.h"
+
+/**
+ * Adds a new physics solver
+ * Note: This node only runs as part of the construction event.
+ */
+USTRUCT(meta=(DisplayName="Spawn Physics Solver", Keywords="Construction,Create,New,Simulation,Physics", Varying))
+struct CONTROLRIGPHYSICS_API FRigUnit_AddPhysicsSolver : public FRigUnitMutable
+{
+	GENERATED_BODY()
+
+	FRigUnit_AddPhysicsSolver()
+	{
+		Name = TEXT("Solver");
+		SolverID = FRigPhysicsSolverID();
+	}
+
+	/*
+	 * The name of the new solver to spawn
+	 */
+	UPROPERTY(meta = (Input))
+	FName Name;
+
+	UPROPERTY(meta = (Input))
+	FRigPhysicsSolverSettings SolverSettings;
+
+	UPROPERTY(meta = (Input))
+	FRigPhysicsSimulationSpaceSettings SimulationSpaceSettings;
+
+	/*
+	 * The identifier of the solver
+	 */
+	UPROPERTY(meta = (Output))
+	FRigPhysicsSolverID SolverID;
+
+	RIGVM_METHOD()
+	virtual void Execute() override;
+};
+
+/**
+ * Instantiates all the objects in the physics world. Some properties can't be modified after this happens.
+ * Note that it will happen automatically during the first simulation step if it hasn't been explicitly
+ * requested. Explicit instantiation allows the timing to be controlled, as allocations etc may cause some
+ * delays.
+ */
+USTRUCT(meta = (DisplayName = "Instantiate physics", Keywords = "Physics"))
+struct CONTROLRIGPHYSICS_API FRigUnit_InstantiatePhysics : public FRigUnitMutable
+{
+	GENERATED_BODY()
+
+	RIGVM_METHOD()
+		virtual void Execute() override;
+
+	/*
+	 * The solver to relate this new physics element to
+	 */
+	UPROPERTY(meta = (Input))
+	FRigPhysicsSolverID SolverID;
+};
+
+
+/**
+ * Steps the specified physics solver
+ */
+USTRUCT(meta = (DisplayName = "Step Physics Solver", Keywords = "Simulate,Physics"))
+struct CONTROLRIGPHYSICS_API FRigUnit_StepPhysicsSolver : public FRigUnitMutable
+{
+	GENERATED_BODY()
+
+	FRigUnit_StepPhysicsSolver()
+	{
+		DeltaTimeOverride = -1.0f;
+	}
+
+	RIGVM_METHOD()
+		virtual void Execute() override;
+
+	/*
+	 * The solver to relate this new physics element to
+	 */
+	UPROPERTY(meta = (Input))
+	FRigPhysicsSolverID SolverID;
+
+	/*
+	 * Normally the execute context time will be used, but if this is positive then it will override it.
+	 */
+	UPROPERTY(meta = (Input))
+	float DeltaTimeOverride;
+
+	/* Enable visualization (includes enabling low-level chaos debug draw) */
+	UPROPERTY(meta = (Input))
+	bool bEnableVisualization = true;
+};
+
+/**
+ * Adds a new physics body as a component on the owner element.
+ * Note: This node only runs as part of the construction event.
+ */
+USTRUCT(meta = (DisplayName = "Add Physics Body Component", Keywords = "Construction,Create,New,Joint,Body,Skeleton,Physics", Varying))
+struct CONTROLRIGPHYSICS_API FRigUnit_AddPhysicsBody : public FRigUnitMutable
+{
+	GENERATED_BODY()
+
+	FRigUnit_AddPhysicsBody()
+	{
+		Owner = FRigElementKey();
+		Result = FRigComponentKey();
+		SolverID = FRigPhysicsSolverID();
+	}
+
+	RIGVM_METHOD()
+		virtual void Execute() override;
+
+	/**
+	 * The owner of the newly created component (must be set/valid)
+	 */
+	UPROPERTY(meta = (Input, BoneName))
+	FRigElementKey Owner;
+
+	UPROPERTY(meta = (Output))
+	FRigComponentKey Result;
+
+	/*
+	 * The solver to relate this new physics element to
+	 */
+	UPROPERTY(meta = (Input))
+	FRigPhysicsSolverID SolverID;
+
+	/*
+	 * The dynamics properties of the new physics element
+	 */
+	UPROPERTY(meta = (Input))
+	FRigPhysicsDynamics Dynamics;
+
+	/*
+	 * The collision properties of the new physics element
+	 */
+	UPROPERTY(meta = (Input))
+	FRigPhysicsCollision Collision;
+
+	/**
+	 * The physics joint (typically to its parent) properties of the new physics element
+	 */
+	UPROPERTY(meta = (Input))
+	FRigPhysicsJoint Joint;
+
+	/**
+	 * The runtime modifiable data
+	 */
+	UPROPERTY(meta = (Input))
+	FPhysicsControlModifierData BodyData;
+};
+
+/**
+ * Sets what bone is used as a source transform for the physics body. This is used as a kinematic target, and when
+ * initializing the simulation.
+ */
+USTRUCT(meta = (DisplayName = "Set Source Bone", Keywords = "Physics", Varying))
+struct CONTROLRIGPHYSICS_API FRigUnit_HierarchySetPhysicsBodySourceBone : public FRigUnitMutable
+{
+	GENERATED_BODY()
+
+	RIGVM_METHOD()
+	virtual void Execute() override;
+
+	UPROPERTY(meta = (Input))
+	FRigComponentKey PhysicsBody;
+
+	UPROPERTY(meta = (Input))
+	FRigElementKey SourceBone;
+};
+
+/**
+ * Sets what bone is targeted by the simulation - i.e. where the simulation output is written to.
+ */
+USTRUCT(meta = (DisplayName = "Set Target Bone", Keywords = "Physics", Varying))
+struct CONTROLRIGPHYSICS_API FRigUnit_HierarchySetPhysicsBodyTargetBone : public FRigUnitMutable
+{
+	GENERATED_BODY()
+
+	RIGVM_METHOD()
+	virtual void Execute() override;
+
+	UPROPERTY(meta = (Input))
+	FRigComponentKey PhysicsBody;
+
+	UPROPERTY(meta = (Input))
+	FRigElementKey TargetBone;
+};
+
+
+/**
+ * Sets all the data on a body - but in a sparse way so you can decide which parameters get applied.
+ * Danny TODO - Note that the sparse data does not get displayed correctly, so this is largely 
+ * unusable - the flags that enable/disable all end up getting reset if the user attempts to change them.
+ */
+USTRUCT(meta = (DisplayName = "Set Physics Body Data", Keywords = "Physics", Varying))
+struct CONTROLRIGPHYSICS_API FRigUnit_HierarchySetPhysicsBodySparseData : public FRigUnitMutable
+{
+	GENERATED_BODY()
+
+	RIGVM_METHOD()
+	virtual void Execute() override;
+
+	UPROPERTY(meta = (Input))
+	FRigComponentKey PhysicsBody;
+
+	UPROPERTY(meta = (Input))
+	FPhysicsControlModifierSparseData Data;
+};
+
+/**
+ * Sets the kinematic target for a body - note that this won't actually make the body kinematic
+ */
+USTRUCT(meta = (DisplayName = "Set Kinematic Target", Keywords = "Physics", Varying))
+struct CONTROLRIGPHYSICS_API FRigUnit_HierarchySetPhysicsBodyKinematicTarget : public FRigUnitMutable
+{
+	GENERATED_BODY()
+
+	RIGVM_METHOD()
+	virtual void Execute() override;
+
+	UPROPERTY(meta = (Input))
+	FRigComponentKey PhysicsBody;
+
+	UPROPERTY(meta = (Input))
+	ERigPhysicsKinematicTargetSpace KinematicTargetSpace = ERigPhysicsKinematicTargetSpace::OffsetInBoneSpace;
+
+	UPROPERTY(meta = (Input))
+	FTransform KinematicTarget;
+};
+
+/**
+ * Sets the movement type for this body. 
+ * Danny TODO explain the different between kinematic and static
+ */
+USTRUCT(meta = (DisplayName = "Set Physics Body Movement Type", Keywords = "Physics", Varying))
+struct CONTROLRIGPHYSICS_API FRigUnit_HierarchySetPhysicsBodyMovementType : public FRigUnitMutable
+{
+	GENERATED_BODY()
+
+	RIGVM_METHOD()
+	virtual void Execute() override;
+
+	UPROPERTY(meta = (Input))
+	FRigComponentKey PhysicsBody;
+
+	UPROPERTY(meta = (Input))
+	EPhysicsMovementType MovementType = EPhysicsMovementType::Simulated;
+};
+
+/**
+ * Sets what collision is used for this body
+ */
+USTRUCT(meta = (DisplayName = "Set Physics Body Collision Type", Keywords = "Physics", Varying))
+struct CONTROLRIGPHYSICS_API FRigUnit_HierarchySetPhysicsBodyCollisionType : public FRigUnitMutable
+{
+	GENERATED_BODY()
+
+	RIGVM_METHOD()
+	virtual void Execute() override;
+
+	UPROPERTY(meta = (Input))
+	FRigComponentKey PhysicsBody;
+
+	UPROPERTY(meta = (Input))
+	TEnumAsByte<ECollisionEnabled::Type> CollisionType = ECollisionEnabled::QueryAndPhysics;
+};
+
+/**
+ * Sets the multiplier on gravity that should be applied to the body.
+ */
+USTRUCT(meta = (DisplayName = "Set Physics Body Gravity Multiplier", Keywords = "Physics", Varying))
+struct CONTROLRIGPHYSICS_API FRigUnit_HierarchySetPhysicsBodyGravityMultiplier : public FRigUnitMutable
+{
+	GENERATED_BODY()
+
+	RIGVM_METHOD()
+	virtual void Execute() override;
+
+	UPROPERTY(meta = (Input))
+	FRigComponentKey PhysicsBody;
+
+	UPROPERTY(meta = (Input))
+	float GravityMultiplier = 1.0f;
+};
+
+/**
+ * Controls the amount that the simulation is blended back into the target bones. 
+ * Danny TODO implement
+ */
+USTRUCT(meta = (DisplayName = "Set Physics Body Physics Blend Weight", Keywords = "Physics", Varying))
+struct CONTROLRIGPHYSICS_API FRigUnit_HierarchySetPhysicsBodyPhysicsBlendWeight : public FRigUnitMutable
+{
+	GENERATED_BODY()
+
+	RIGVM_METHOD()
+	virtual void Execute() override;
+
+	UPROPERTY(meta = (Input))
+	FRigComponentKey PhysicsBody;
+
+	UPROPERTY(meta = (Input))
+	float PhysicsBlendWeight = 1.0f;
+};
+
+/**
+ * Danny TODO I think this is redundant given the ability to specify the space when kinematic 
+ * targets are set. Investigate propagating that into the rest of Physics Control.
+ */
+USTRUCT(meta = (DisplayName = "Set Physics Body Use Skeletal Animation", Keywords = "Physics", Varying))
+struct CONTROLRIGPHYSICS_API FRigUnit_HierarchySetPhysicsBodyUseSkeletalAnimation : public FRigUnitMutable
+{
+	GENERATED_BODY()
+
+	RIGVM_METHOD()
+	virtual void Execute() override;
+
+	UPROPERTY(meta = (Input))
+	FRigComponentKey PhysicsBody;
+
+	UPROPERTY(meta = (Input))
+	bool bUseSkeletalAnimation = true;
+};
+
+/**
+ * If true, then kinematic objects will be written back from simulation to the bones. This only 
+ * necessary when either kinematic targets are being used, or when the target bone differs from the source bone.
+ * Danny TODO - check the implementation, and also can we auto-calculate this at runtime?
+ */
+USTRUCT(meta = (DisplayName = "Set Physics Body Update Kinematic From Simulation", Keywords = "Physics", Varying))
+struct CONTROLRIGPHYSICS_API FRigUnit_HierarchySetPhysicsBodyUpdateKinematicFromSimulation : public FRigUnitMutable
+{
+	GENERATED_BODY()
+
+	RIGVM_METHOD()
+	virtual void Execute() override;
+
+	UPROPERTY(meta = (Input))
+	FRigComponentKey PhysicsBody;
+
+	UPROPERTY(meta = (Input))
+	bool bUpdateKinematicFromSimulation = true;
+};
+
+
diff --git a/Engine/Plugins/Experimental/PhysicsControl/Source/PhysicsControl/Private/RigidBodyWithControl.cpp b/Engine/Plugins/Experimental/PhysicsControl/Source/PhysicsControl/Private/RigidBodyWithControl.cpp
--- a/Engine/Plugins/Experimental/PhysicsControl/Source/PhysicsControl/Private/RigidBodyWithControl.cpp
+++ b/Engine/Plugins/Experimental/PhysicsControl/Source/PhysicsControl/Private/RigidBodyWithControl.cpp
@@ -721,35 +721,46 @@ void FAnimNode_RigidBodyWithControl::ApplyControl(FRigidBodyControlRecord& Contr
 }
 
 //======================================================================================================================
-void FAnimNode_RigidBodyWithControl::ApplyModifier(
-	const FRigidBodyModifierRecord& BodyModifierRecord, 
-	const FVector&                  SimSpaceGravity)
+void UpdateBodyFromModifierData(
+	ImmediatePhysics::FActorHandle*    ActorHandle, 
+	ImmediatePhysics::FSimulation*     PhysicsSimulation,
+	const FPhysicsControlModifierData& ModifierData,
+	const FVector&                     SimSpaceGravity)
 {
-	if (BodyModifierRecord.ActorHandle)
+	if (ModifierData.MovementType != EPhysicsMovementType::Default)
 	{
-		if (BodyModifierRecord.ModifierData.MovementType != EPhysicsMovementType::Default)
-		{
-			// Note that there's an early out if there's no change needed, so this should be OK.
-			PhysicsSimulation->SetIsKinematic(
-				BodyModifierRecord.ActorHandle,
-				BodyModifierRecord.ModifierData.MovementType != EPhysicsMovementType::Simulated);
-		}
+		// Note that there's an early out if there's no change needed, so this should be OK.
+		PhysicsSimulation->SetIsKinematic(
+			ActorHandle,
+			ModifierData.MovementType != EPhysicsMovementType::Simulated);
+	}
 
-		// Note that the actual kinematic targets will be set separately, since they need to be set
-		// for all kinematics whether or not they were under a modifier.
+	// Note that the actual kinematic targets will be set separately, since they need to be set
+	// for all kinematics whether or not they were under a modifier.
 
-		// Scale gravity
-		float GravityMultiplier = BodyModifierRecord.ModifierData.GravityMultiplier;
-		if (GravityMultiplier != 1 && BodyModifierRecord.ActorHandle->IsGravityEnabled())
-		{
-			float Mass = (float) BodyModifierRecord.ActorHandle->GetMass();
-			FVector AntiGravityForce = SimSpaceGravity * (-Mass * (1.0f - GravityMultiplier));
-			BodyModifierRecord.ActorHandle->AddForce(AntiGravityForce);
-		}
+	// Scale gravity
+	float GravityMultiplier = ModifierData.GravityMultiplier;
+	if (GravityMultiplier != 1 && ActorHandle->IsGravityEnabled())
+	{
+		float Mass = (float) ActorHandle->GetMass();
+		FVector AntiGravityForce = SimSpaceGravity * (-Mass * (1.0f - GravityMultiplier));
+		ActorHandle->AddForce(AntiGravityForce);
+	}
+
+	// Set collision
+	PhysicsSimulation->SetHasCollision(
+		ActorHandle, CollisionEnabledHasPhysics(ModifierData.CollisionType));
+}
 
-		// Set collision
-		PhysicsSimulation->SetHasCollision(
-			BodyModifierRecord.ActorHandle, CollisionEnabledHasPhysics(BodyModifierRecord.ModifierData.CollisionType));
+//======================================================================================================================
+void FAnimNode_RigidBodyWithControl::ApplyModifier(
+	const FRigidBodyModifierRecord& BodyModifierRecord, 
+	const FVector&                  SimSpaceGravity)
+{
+	if (BodyModifierRecord.ActorHandle && PhysicsSimulation)
+	{
+		UpdateBodyFromModifierData(
+			BodyModifierRecord.ActorHandle, PhysicsSimulation, BodyModifierRecord.ModifierData, SimSpaceGravity);
 	}
 }
 
diff --git a/Engine/Plugins/Experimental/PhysicsControl/Source/PhysicsControl/Public/AnimNode_RigidBodyWithControl.h b/Engine/Plugins/Experimental/PhysicsControl/Source/PhysicsControl/Public/AnimNode_RigidBodyWithControl.h
--- a/Engine/Plugins/Experimental/PhysicsControl/Source/PhysicsControl/Public/AnimNode_RigidBodyWithControl.h
+++ b/Engine/Plugins/Experimental/PhysicsControl/Source/PhysicsControl/Public/AnimNode_RigidBodyWithControl.h
@@ -50,6 +50,16 @@ enum class MapConstraintsBehaviorType : uint8
 	None UMETA(DisplayName = "Nothing", ToolTip = "Nothing"),
 };
 
+/**
+ * A helper function that lets any immediate simulation user update their bodies in the same way that
+ * RigidBodyWithControl does.
+ */
+PHYSICSCONTROL_API void UpdateBodyFromModifierData(
+	ImmediatePhysics::FActorHandle* ActorHandle,
+	ImmediatePhysics::FSimulation* PhysicsSimulation,
+	const FPhysicsControlModifierData& ModifierData,
+	const FVector& SimSpaceGravity);
+
 /**
  * Controller that simulates physics based on the physics asset of the skeletal mesh component
  */
diff --git a/Engine/Plugins/Experimental/PhysicsControl/Source/PhysicsControl/Public/PhysicsControlData.h b/Engine/Plugins/Experimental/PhysicsControl/Source/PhysicsControl/Public/PhysicsControlData.h
--- a/Engine/Plugins/Experimental/PhysicsControl/Source/PhysicsControl/Public/PhysicsControlData.h
+++ b/Engine/Plugins/Experimental/PhysicsControl/Source/PhysicsControl/Public/PhysicsControlData.h
@@ -559,6 +559,28 @@ struct PHYSICSCONTROL_API FPhysicsControlData
 
 	bool operator==(const FPhysicsControlData& other) const = default;
 
+	friend FArchive& operator<<(FArchive& Ar, FPhysicsControlData& Data)
+	{
+		Ar << Data.LinearStrength;
+		Ar << Data.LinearDampingRatio;
+		Ar << Data.LinearExtraDamping;
+		Ar << Data.MaxForce;
+		Ar << Data.AngularStrength;
+		Ar << Data.AngularDampingRatio;
+		Ar << Data.AngularExtraDamping;
+		Ar << Data.MaxTorque;
+		Ar << Data.LinearTargetVelocityMultiplier;
+		Ar << Data.AngularTargetVelocityMultiplier;
+		Ar << Data.SkeletalAnimationVelocityMultiplier;
+		Ar << Data.CustomControlPoint;
+		Ar << Data.bEnabled;
+		Ar << Data.bUseCustomControlPoint;
+		Ar << Data.bUseSkeletalAnimation;
+		Ar << Data.bDisableCollision;
+		Ar << Data.bOnlyControlChildObject;
+		return Ar;
+	}
+
 	/** Applies the values that have been flagged as enabled from the sparse data */
 	void UpdateFromSparseData(const FPhysicsControlSparseData& SparseData);
 
@@ -639,26 +661,26 @@ struct PHYSICSCONTROL_API FPhysicsControlData
 
 	/** Whether this control should be enabled */
 	UPROPERTY(EditAnywhere, BlueprintReadWrite, Category = PhysicsControl)
-	uint8 bEnabled : 1;
+	bool bEnabled;
 
 	/** Whether or not to use the custom control point position */
 	UPROPERTY(EditAnywhere, BlueprintReadWrite, Category = PhysicsControl)
-	uint8 bUseCustomControlPoint : 1;
+	bool bUseCustomControlPoint;
 
 	/** If true then the target will be applied on top of the skeletal animation (if there is any) */
 	UPROPERTY(EditAnywhere, BlueprintReadWrite, Category = PhysicsControl)
-	uint8 bUseSkeletalAnimation : 1;
+	bool bUseSkeletalAnimation;
 
 	/**
 	 * Whether or not this control should disable collision between the parent and child bodies (only
 	 * has an effect if there is a parent body)
 	 */
 	UPROPERTY(EditAnywhere, BlueprintReadWrite, Category = PhysicsControl)
-	uint8 bDisableCollision : 1;
+	bool bDisableCollision;
 
 	/** If true then the control will only affect the child object, not the parent */
 	UPROPERTY(EditAnywhere, BlueprintReadWrite, Category = PhysicsControl)
-	uint8 bOnlyControlChildObject : 1;
+	bool bOnlyControlChildObject;
 };
 
 /**
@@ -842,6 +864,16 @@ struct PHYSICSCONTROL_API FPhysicsControlTarget
 	{
 	}
 
+	friend FArchive& operator<<(FArchive& Ar, FPhysicsControlTarget& Data)
+	{
+		Ar << Data.TargetPosition;
+		Ar << Data.TargetVelocity;
+		Ar << Data.TargetOrientation;
+		Ar << Data.TargetAngularVelocity;
+		Ar << Data.bApplyControlPointToTarget;
+		return Ar;
+	}
+
 	/** The target position of the child body, relative to the parent body */
 	UPROPERTY(EditAnywhere, BlueprintReadWrite, Category = PhysicsControl)
 	FVector TargetPosition;
