Physics and Control Rig. Catch/handle bad simulation data. Exposes solver timestep smoothing.

diff --git a/Engine/Plugins/Experimental/PhysicsControl/Source/ControlRigPhysics/Private/RigPhysicsSolver.cpp b/Engine/Plugins/Experimental/PhysicsControl/Source/ControlRigPhysics/Private/RigPhysicsSolver.cpp
--- a/Engine/Plugins/Experimental/PhysicsControl/Source/ControlRigPhysics/Private/RigPhysicsSolver.cpp
+++ b/Engine/Plugins/Experimental/PhysicsControl/Source/ControlRigPhysics/Private/RigPhysicsSolver.cpp
@@ -400,19 +400,17 @@ void FRigPhysicsSolver::InstantiatePhysicsBodies()
 				SourceKey = ComponentKey.ElementKey;
 			}
 
-			FTransform BodyRelComponentTM;
 			// What should we do if the key is not valid? 
 			if (SourceKey.IsValid())
 			{
-				BodyRelComponentTM = Hierarchy->GetGlobalTransform(SourceKey);
-			}
-
-			const FTransform BodyRelSimSpaceTM = ConvertComponentSpaceTransformToSimSpace(BodyRelComponentTM);
-
-			Record.ActorHandle = CreateBody(Collision, &Dynamics, BodyRelSimSpaceTM);
+				FTransform SourceComponentSpaceTM = Hierarchy->GetGlobalTransform(SourceKey);
+				const FTransform SourceSimulationSpaceTM =
+					ConvertComponentSpaceTransformToSimSpace(SourceComponentSpaceTM);
+				Record.ActorHandle = CreateBody(Collision, &Dynamics, SourceSimulationSpaceTM);
 
-			// Note that these data will be applied in the PrePhysics update
-			Record.BodyData = BodyData;
+				// Note that these data will be applied in the PrePhysics update
+				Record.BodyData = BodyData;
+			}
 		}
 	}
 }
@@ -595,6 +593,7 @@ void FRigPhysicsSolver::UpdatePrePhysics(FRigVMExecuteContext& ExecuteContext, c
 // otherwise default to the physics element parent.
 void FRigPhysicsSolver::UpdatePostPhysics()
 {
+	bool bGotInvalidSimulationData = false;
 	URigHierarchy* Hierarchy = OwningControlRig->GetHierarchy();
 	for (TPair<FRigComponentKey, FRigBodyRecord>& BodyRecordPair : BodyRecords)
 	{
@@ -617,8 +616,49 @@ void FRigPhysicsSolver::UpdatePostPhysics()
 				if (TargetElementKey.IsValid())
 				{
 					FTransform TargetSimSpaceTM = Record.ActorHandle->GetWorldTransform();
-					FTransform TargetComponentSpaceTM = ConvertSimSpaceTransformToComponentSpace(TargetSimSpaceTM);
-					Hierarchy->SetGlobalTransform(TargetElementKey, TargetComponentSpaceTM);
+					double DistSq = TargetSimSpaceTM.GetTranslation().SquaredLength();
+					if (!TargetSimSpaceTM.IsValid() || (SolverSettings.DistanceThresholdForReset > 0 &&
+						DistSq > FMath::Square(SolverSettings.DistanceThresholdForReset)))
+					{
+						bGotInvalidSimulationData = true;
+					}
+					else
+					{
+						FTransform TargetComponentSpaceTM = ConvertSimSpaceTransformToComponentSpace(TargetSimSpaceTM);
+						Hierarchy->SetGlobalTransform(TargetElementKey, TargetComponentSpaceTM);
+					}
+				}
+			}
+		}
+	}
+
+	if (bGotInvalidSimulationData)
+	{
+		for (TPair<FRigComponentKey, FRigBodyRecord>& BodyRecordPair : BodyRecords)
+		{
+			const FRigComponentKey& ComponentKey = BodyRecordPair.Key;
+			const FRigBodyRecord& Record = BodyRecordPair.Value;
+
+			if (const FRigPhysicsComponent* PhysicsComponent = Cast<FRigPhysicsComponent>(
+				Hierarchy->FindComponent(ComponentKey)))
+			{
+				if (Record.ActorHandle && !Record.ActorHandle->GetIsKinematic())
+				{
+					// Get the TM in component space, and then convert it into sim space.
+					FRigElementKey SourceKey = PhysicsComponent->BodySolverSettings.SourceBone;
+					if (!SourceKey.IsValid())
+					{
+						SourceKey = ComponentKey.ElementKey;
+					}
+					if (SourceKey.IsValid())
+					{
+						FTransform SourceComponentSpaceTM = Hierarchy->GetGlobalTransform(SourceKey);
+						const FTransform SourceSimulationSpaceTM =
+							ConvertComponentSpaceTransformToSimSpace(SourceComponentSpaceTM);
+						Record.ActorHandle->SetWorldTransform(SourceSimulationSpaceTM);
+						Record.ActorHandle->SetLinearVelocity(FVector::ZeroVector);
+						Record.ActorHandle->SetAngularVelocity(FVector::ZeroVector);
+					}
 				}
 			}
 		}
@@ -651,7 +691,7 @@ void FRigPhysicsSolver::StepSimulation(FRigVMExecuteContext& ExecuteContext, dou
 		SolverSettings.VelocityIterations,
 		SolverSettings.ProjectionIterations,
 		SolverSettings.bUseManifolds);
-	Simulation->SetMaxNumRollingAverageStepTimes(1);
+	Simulation->SetMaxNumRollingAverageStepTimes(SolverSettings.MaxNumRollingAverageStepTimes);
 
 	// Other settings - would normally be static (so Danny TODO move this)
 	ChaosJointSolverSettings.bSolvePositionLast = SolverSettings.bSolveJointPositionsLast;
diff --git a/Engine/Plugins/Experimental/PhysicsControl/Source/ControlRigPhysics/Private/RigPhysicsSolverSimulationSpace.cpp b/Engine/Plugins/Experimental/PhysicsControl/Source/ControlRigPhysics/Private/RigPhysicsSolverSimulationSpace.cpp
--- a/Engine/Plugins/Experimental/PhysicsControl/Source/ControlRigPhysics/Private/RigPhysicsSolverSimulationSpace.cpp
+++ b/Engine/Plugins/Experimental/PhysicsControl/Source/ControlRigPhysics/Private/RigPhysicsSolverSimulationSpace.cpp
@@ -212,31 +212,46 @@ FRigPhysicsSolver::FSimulationSpaceData FRigPhysicsSolver::UpdateSimulationSpace
 	SimulationSpaceData.LinearAcceleration.Z *= SimulationSpaceSettings.VelocityScaleZ;
 
 	// Clamped world-space motion of the simulation space
-	if (SimulationSpaceSettings.bClampLinearVelocity)
+	if (
+		(SimulationSpaceSettings.LinearAccelerationThresholdForTeleport > 0 &&
+			SimulationSpaceData.LinearAcceleration.SquaredLength() >
+			FMath::Square(SimulationSpaceSettings.LinearAccelerationThresholdForTeleport)) ||
+		(SimulationSpaceSettings.AngularAccelerationThresholdForTeleport > 0 &&
+			SimulationSpaceData.AngularAcceleration.SquaredLength() >
+			FMath::Square(SimulationSpaceSettings.AngularAccelerationThresholdForTeleport))
+		)
 	{
-		SimulationSpaceData.LinearVelocity = 
-			SimulationSpaceData.LinearVelocity.GetClampedToMaxSize(SimulationSpaceSettings.MaxLinearVelocity);
+		SimulationSpaceData.LinearVelocity = FVector::ZeroVector;
+		SimulationSpaceData.AngularVelocity = FVector::ZeroVector;
+		SimulationSpaceData.LinearAcceleration = FVector::ZeroVector;
+		SimulationSpaceData.AngularAcceleration = FVector::ZeroVector;
 	}
-	SimulationSpaceData.LinearVelocity += SimulationSpaceSettings.ExternalLinearVelocity;
-
-	if (SimulationSpaceSettings.bClampAngularVelocity)
-	{
-		SimulationSpaceData.AngularVelocity =
-		SimulationSpaceData.AngularVelocity.GetClampedToMaxSize(SimulationSpaceSettings.MaxAngularVelocity);
-	}
-	SimulationSpaceData.AngularVelocity += SimulationSpaceSettings.ExternalAngularVelocity;
-
-	if (SimulationSpaceSettings.bClampLinearAcceleration)
+	else
 	{
-		SimulationSpaceData.LinearAcceleration = 
-			SimulationSpaceData.LinearAcceleration.GetClampedToMaxSize(SimulationSpaceSettings.MaxLinearAcceleration);
+		if (SimulationSpaceSettings.bClampLinearVelocity)
+		{
+			SimulationSpaceData.LinearVelocity =
+				SimulationSpaceData.LinearVelocity.GetClampedToMaxSize(SimulationSpaceSettings.MaxLinearVelocity);
+		}
+		if (SimulationSpaceSettings.bClampAngularVelocity)
+		{
+			SimulationSpaceData.AngularVelocity =
+				SimulationSpaceData.AngularVelocity.GetClampedToMaxSize(SimulationSpaceSettings.MaxAngularVelocity);
+		}
+		if (SimulationSpaceSettings.bClampLinearAcceleration)
+		{
+			SimulationSpaceData.LinearAcceleration =
+				SimulationSpaceData.LinearAcceleration.GetClampedToMaxSize(SimulationSpaceSettings.MaxLinearAcceleration);
+		}
+		if (SimulationSpaceSettings.bClampAngularAcceleration)
+		{
+			SimulationSpaceData.AngularAcceleration =
+				SimulationSpaceData.AngularAcceleration.GetClampedToMaxSize(SimulationSpaceSettings.MaxAngularAcceleration);
+		}
 	}
 
-	if (SimulationSpaceSettings.bClampAngularAcceleration)
-	{
-		SimulationSpaceData.AngularAcceleration = 
-			SimulationSpaceData.AngularAcceleration.GetClampedToMaxSize(SimulationSpaceSettings.MaxAngularAcceleration);
-	}
+	SimulationSpaceData.LinearVelocity += SimulationSpaceSettings.ExternalLinearVelocity;
+	SimulationSpaceData.AngularVelocity += SimulationSpaceSettings.ExternalAngularVelocity;
 
 	// Transform world-space motion into simulation space Danny TODO note that this matches the code
 	// in RBAN, and is doing what the interface requires (i.e. movement of the space in the space of
diff --git a/Engine/Plugins/Experimental/PhysicsControl/Source/ControlRigPhysics/Public/RigPhysicsData.h b/Engine/Plugins/Experimental/PhysicsControl/Source/ControlRigPhysics/Public/RigPhysicsData.h
--- a/Engine/Plugins/Experimental/PhysicsControl/Source/ControlRigPhysics/Public/RigPhysicsData.h
+++ b/Engine/Plugins/Experimental/PhysicsControl/Source/ControlRigPhysics/Public/RigPhysicsData.h
@@ -15,6 +15,8 @@ struct FRigControlElement;
 class URigHierarchy;
 class FRigElementKeyRedirector;
 
+const float RigPhysicsDefaultShapeSize = 10.0f;
+
 /** 
  * Specifies what space a kinematic target is defined in 
  * Danny TODO this is much better than the similar mechanism in PCC and RBWC (which has a 
@@ -105,7 +107,7 @@ struct CONTROLRIGPHYSICS_API FRigPhysicsCollisionBox : public FRigPhysicsCollisi
 
 	// These are the full extents of the box in each axis
 	UPROPERTY(BlueprintReadWrite, EditAnywhere, Category = Shapes)
-	FVector Extents;
+	FVector Extents = FVector::OneVector * RigPhysicsDefaultShapeSize;
 };
 
 /**
@@ -127,7 +129,7 @@ struct CONTROLRIGPHYSICS_API FRigPhysicsCollisionSphere : public FRigPhysicsColl
 	FTransform TM;
 
 	UPROPERTY(BlueprintReadWrite, EditAnywhere, Category = Shapes)
-	float Radius;
+	float Radius = RigPhysicsDefaultShapeSize;
 };
 
 /**
@@ -150,11 +152,11 @@ struct CONTROLRIGPHYSICS_API FRigPhysicsCollisionCapsule : public FRigPhysicsCol
 	FTransform TM;
 
 	UPROPERTY(BlueprintReadWrite, EditAnywhere, Category = Shapes)
-	float Radius;
+	float Radius = RigPhysicsDefaultShapeSize;
 
 	// This is the length of the core part of the capsule. The total length will be Length + 2 * Radius
 	UPROPERTY(BlueprintReadWrite, EditAnywhere, Category = Shapes)
-	float Length;
+	float Length = RigPhysicsDefaultShapeSize;
 };
 
 /**
@@ -249,6 +251,20 @@ struct FRigPhysicsSimulationSpaceSettings
 	UPROPERTY(EditAnywhere, BlueprintReadWrite, Category = Settings, meta = (ClampMin = "0.0", EditCondition = "bClampAngularAcceleration"))
 	float MaxAngularAcceleration = 10000;
 
+	// If the linear acceleration is above this threshold, the movement is treated as a teleport.
+	// The calculated accelerations and velocities will be set to zero, without resetting the
+	// simulation state.
+	// A value of zero will disable this detection
+	UPROPERTY(EditAnywhere, BlueprintReadWrite, Category = Settings, meta = (ClampMin = "0.0"))
+	float LinearAccelerationThresholdForTeleport = 100000;
+
+	// If the angular acceleration (rand/s/s) is above this threshold, the movement is treated as a teleport.
+	// The calculated accelerations and velocities will be set to zero, without resetting the
+	// simulation state.
+	// A value of zero will disable this detection
+	UPROPERTY(EditAnywhere, BlueprintReadWrite, Category = Settings, meta = (ClampMin = "0.0"))
+	float AngularAccelerationThresholdForTeleport = 100000;
+
 	// Additional linear drag applied to every body in addition to linear drag specified on them in
 	// the physics asset. When combined with ExternalLinearVelocity, this can be used to add a
 	// temporary wind-blown effect without having to tune linear drag on all the bodies in the
@@ -323,9 +339,14 @@ struct CONTROLRIGPHYSICS_API FRigPhysicsSolverSettings
 	UPROPERTY(EditAnywhere, BlueprintReadWrite, Category = SolverSettings, meta = (ClampMin = 0))
 	int32 ProjectionIterations = 1;
 
+	// This sets how or if the step size should be smoothed. A value of one will disable smoothing,
+	// so that the physics simulation will match the Control Rig delta time. 
+	UPROPERTY(EditAnywhere, BlueprintReadWrite, Category = SolverSettings, meta = (ClampMin = 1))
+	int32 MaxNumRollingAverageStepTimes = 1;
+
 	/**
 	 * When bodies are penetrating, this is the maximum velocity delta that can be applied in one frame.
-	*/
+	 */
 	UPROPERTY(EditAnywhere, BlueprintReadWrite, Category = SolverSettings, meta = (ClampMin = 0))
 	float MaxDepenetrationVelocity = 0.0f;
 
@@ -353,7 +374,7 @@ struct CONTROLRIGPHYSICS_API FRigPhysicsSolverSettings
 	 * the non-linear solver when you are running multiple iterations, but is more likely to suffer from jitter.
 	 * In general you should try to use the linear solver and increase the PositionIterations to improve stability if
 	 * possible, only using the non-linear solver as a last resort.
-	*/
+	 */
 	UPROPERTY(EditAnywhere, BlueprintReadWrite, Category = SolverSettings)
 	bool bUseLinearJointSolver = true;
 
@@ -366,9 +387,15 @@ struct CONTROLRIGPHYSICS_API FRigPhysicsSolverSettings
 	/**
 	 * It enables the use of multi-point contact manifolds, which are created only once at the start of each tick.
 	 * When disabled, a single-point contact is generated in each solver iteration which is more expensive.
-	*/
+	 */
 	UPROPERTY(EditAnywhere, BlueprintReadWrite, Category = CollisionSettings)
 	bool bUseManifolds = true;
+
+	// If any object in the simulation exceeds this distance from the simulation origin, we will
+	// reset the entire simulation. This is to detect problems/explosions.
+	// Disabled when the value is zero.
+	UPROPERTY(EditAnywhere, BlueprintReadWrite, Category = Settings, meta = (ClampMin = "0.0"))
+	float DistanceThresholdForReset = 0;
 };
 
 /**
@@ -398,7 +425,7 @@ struct CONTROLRIGPHYSICS_API FRigPhysicsDynamics
 	bool bOverrideCentreOfMass = false;
 
 	UPROPERTY(BlueprintReadWrite, EditAnywhere, Category = Dynamics, meta = (EditCondition = "bOverrideCentreOfMass"))
-	FVector CentreOfMassOverride;
+	FVector CentreOfMassOverride = FVector::ZeroVector;
 
 	// Use this to specify the moments of inertia
 	UPROPERTY(BlueprintReadWrite, EditAnywhere, Category = Dynamics)
