Control Rig Physics - ensure that body changes prior to instantiation are used

diff --git a/Engine/Plugins/Experimental/PhysicsControl/Source/ControlRigPhysics/Private/RigPhysicsSolver.cpp b/Engine/Plugins/Experimental/PhysicsControl/Source/ControlRigPhysics/Private/RigPhysicsSolver.cpp
--- a/Engine/Plugins/Experimental/PhysicsControl/Source/ControlRigPhysics/Private/RigPhysicsSolver.cpp
+++ b/Engine/Plugins/Experimental/PhysicsControl/Source/ControlRigPhysics/Private/RigPhysicsSolver.cpp
@@ -54,6 +54,19 @@ void FRigPhysicsSolver::AddBody(FRigComponentKey RigComponentKey)
 {
 	// Just make the record for now. The body details may change, so only actually instantiate later.
 	FRigBodyRecord& Record = BodyRecords.Add(RigComponentKey);
+
+	// We keep track of the desired runtime state, so copy it here
+	if (const URigHierarchy* Hierarchy = OwningControlRig->GetHierarchy())
+	{
+		const FRigPhysicsComponent* PhysicsComponent = Cast<FRigPhysicsComponent>(
+			Hierarchy->FindComponent(RigComponentKey));
+
+		if (PhysicsComponent)
+		{
+			Record.BodyData = PhysicsComponent->BodyData;
+		}
+	}
+
 }
 
 //======================================================================================================================
@@ -392,7 +405,8 @@ void FRigPhysicsSolver::InstantiatePhysicsBodies()
 		{
 			const FRigPhysicsCollision& Collision = PhysicsComponent->Collision;
 			const FRigPhysicsDynamics& Dynamics = PhysicsComponent->Dynamics;
-			const FPhysicsControlModifierData& BodyData = PhysicsComponent->BodyData;
+			// Use the Record's BodyData as it may have been modified since the body was added
+			const FPhysicsControlModifierData& BodyData = Record.BodyData;
 
 			FRigElementKey SourceKey = PhysicsComponent->BodySolverSettings.SourceBone;
 			if (!SourceKey.IsValid())
@@ -407,9 +421,6 @@ void FRigPhysicsSolver::InstantiatePhysicsBodies()
 				const FTransform SourceSimulationSpaceTM =
 					ConvertComponentSpaceTransformToSimSpace(SourceComponentSpaceTM);
 				Record.ActorHandle = CreateBody(Collision, &Dynamics, SourceSimulationSpaceTM);
-
-				// Note that these data will be applied in the PrePhysics update
-				Record.BodyData = BodyData;
 			}
 		}
 	}
