Control Rig Physics - basic support for drives with position and velocity control.

diff --git a/Engine/Plugins/Experimental/PhysicsControl/Source/ControlRigPhysics/Private/RigPhysicsExecution.cpp b/Engine/Plugins/Experimental/PhysicsControl/Source/ControlRigPhysics/Private/RigPhysicsExecution.cpp
--- a/Engine/Plugins/Experimental/PhysicsControl/Source/ControlRigPhysics/Private/RigPhysicsExecution.cpp
+++ b/Engine/Plugins/Experimental/PhysicsControl/Source/ControlRigPhysics/Private/RigPhysicsExecution.cpp
@@ -53,9 +53,10 @@ FRigUnit_StepPhysicsSolver_Execute()
 			{
 				double DeltaTime = DeltaTimeOverride > 0 ? DeltaTimeOverride : ExecuteContext.GetDeltaTime();
 				PhysicsSolver->StepSimulation(ExecuteContext, DeltaTime);
-				if (bEnableVisualization)
+				if (VisualizationSettings.bEnableVisualization)
 				{
-					PhysicsSolver->Draw(ExecuteContext.GetDrawInterface(), ExecuteContext.GetWorld());
+					PhysicsSolver->Draw(
+						ExecuteContext.GetDrawInterface(), VisualizationSettings, ExecuteContext.GetWorld());
 				}
 			}
 		}
diff --git a/Engine/Plugins/Experimental/PhysicsControl/Source/ControlRigPhysics/Private/RigPhysicsSolver.cpp b/Engine/Plugins/Experimental/PhysicsControl/Source/ControlRigPhysics/Private/RigPhysicsSolver.cpp
--- a/Engine/Plugins/Experimental/PhysicsControl/Source/ControlRigPhysics/Private/RigPhysicsSolver.cpp
+++ b/Engine/Plugins/Experimental/PhysicsControl/Source/ControlRigPhysics/Private/RigPhysicsSolver.cpp
@@ -5,6 +5,7 @@
 
 #include "AnimNode_RigidBodyWithControl.h"
 #include "PhysicsControlPoseData.h"
+#include "PhysicsControlHelpers.h"
 
 #include "ControlRig.h"
 #include "Rigs/RigHierarchy.h"
@@ -308,7 +309,6 @@ void FRigPhysicsSolver::InstantiatePhysicsJoints()
 					}
 				}
 			}
-
 			const FRigPhysicsComponent* ParentPhysicsComponent =
 				Cast<FRigPhysicsComponent>(Hierarchy->FindComponent(ParentBodyComponentKey));
 			if (!ParentPhysicsComponent)
@@ -316,6 +316,8 @@ void FRigPhysicsSolver::InstantiatePhysicsJoints()
 				continue;
 			}
 
+			ChildBodyRecord.ParentBodyComponentKey = ParentBodyComponentKey;
+
 			const FRigBodyRecord* ParentRecord = nullptr;
 			ImmediatePhysics::FActorHandle* ParentActorHandle = nullptr;
 
@@ -362,7 +364,7 @@ void FRigPhysicsSolver::InstantiatePhysicsJoints()
 
 			// Danny TODO The linear limit is likely implemented as a box limit (check) - it's just
 			// not exposed as one. We should change this.
-			ConstraintSettings.LinearLimit = PhysicsJoint.LinearLimit.Length();
+			ConstraintSettings.LinearLimit = PhysicsJoint.LinearLimit;
 			Chaos::EJointMotionType LinearMotionType = GetMotionTypeFromLimitValue(ConstraintSettings.LinearLimit);
 			ConstraintSettings.LinearMotionTypes = { LinearMotionType, LinearMotionType, LinearMotionType };
 
@@ -386,6 +388,65 @@ void FRigPhysicsSolver::InstantiatePhysicsJoints()
 
 			ChildBodyRecord.JointHandle = Simulation->CreateJoint(
 				ConstraintSettings, ChildActorHandle, ParentActorHandle);
+
+			// Now set up the drive, if requested
+			const FRigPhysicsDrive& Drive = PhysicsJoint.Drive;
+			if (ChildBodyRecord.JointHandle && Drive.bEnable)
+			{
+				float AngularSpring;
+				float AngularDamping;
+				const float MaxTorque = Drive.MaxTorque;
+
+				float LinearSpring;
+				float LinearDamping;
+				const float MaxForce = Drive.MaxForce;
+
+				UE::PhysicsControl::ConvertStrengthToSpringParams(
+					AngularSpring, AngularDamping,
+					Drive.AngularStrength, Drive.AngularDampingRatio, Drive.AngularExtraDamping);
+				UE::PhysicsControl::ConvertStrengthToSpringParams(
+					LinearSpring, LinearDamping,
+					Drive.LinearStrength, Drive.LinearDampingRatio, Drive.LinearExtraDamping);
+
+				if (Chaos::FPBDJointConstraintHandle* Constraint = ChildBodyRecord.JointHandle->GetConstraint())
+				{
+					Chaos::FPBDJointSettings Settings = Constraint->GetSettings();
+
+					// Linear
+					Settings.LinearDriveForceMode = Chaos::EJointForceMode::Acceleration;
+					const bool bLinearPositionDriveEnabled = LinearSpring > 0;
+					const bool bLinearVelocityDriveEnabled = LinearDamping > 0;
+					Settings.bLinearPositionDriveEnabled = { 
+						bLinearPositionDriveEnabled, bLinearPositionDriveEnabled, bLinearPositionDriveEnabled};
+					Settings.bLinearVelocityDriveEnabled = {
+						bLinearVelocityDriveEnabled, bLinearVelocityDriveEnabled, bLinearVelocityDriveEnabled};
+
+					// Angular
+					Settings.AngularDriveForceMode = Chaos::EJointForceMode::Acceleration;
+					// SLERP is generally better, but doesn't work if anything is locked
+					if (ConstraintSettings.AngularMotionTypes[0] == Chaos::EJointMotionType::Locked ||
+						ConstraintSettings.AngularMotionTypes[0] == Chaos::EJointMotionType::Locked ||
+						ConstraintSettings.AngularMotionTypes[0] == Chaos::EJointMotionType::Locked)
+					{
+						Settings.bAngularSwingPositionDriveEnabled = AngularSpring > 0;
+						Settings.bAngularTwistPositionDriveEnabled = AngularSpring > 0;
+						Settings.bAngularSwingVelocityDriveEnabled = AngularDamping > 0;
+						Settings.bAngularTwistVelocityDriveEnabled = AngularDamping > 0;
+					}
+					else
+					{
+						Settings.bAngularSLerpPositionDriveEnabled = AngularSpring > 0;
+						Settings.bAngularSLerpVelocityDriveEnabled = AngularDamping > 0;
+					}
+
+					Constraint->SetSettings(Settings);
+					Constraint->SetDriveParams(
+						Chaos::FVec3(LinearSpring), Chaos::FVec3(LinearDamping), Chaos::FVec3(MaxForce),
+						Chaos::FVec3(AngularSpring), Chaos::FVec3(AngularDamping), Chaos::FVec3(MaxTorque));
+
+					// The target will be set in the pre-physics update
+				}
+			}
 		}
 	}
 }
@@ -506,7 +567,8 @@ void FRigPhysicsSolver::SetBodyKinematicTarget(
 }
 
 //======================================================================================================================
-void FRigPhysicsSolver::UpdatePrePhysics(FRigVMExecuteContext& ExecuteContext, const FSimulationSpaceData& SimulationSpaceData)
+void FRigPhysicsSolver::UpdatePrePhysics(
+	FRigVMExecuteContext& ExecuteContext, const double DeltaTime, const FSimulationSpaceData& SimulationSpaceData)
 {
 	URigHierarchy* Hierarchy = OwningControlRig->GetHierarchy();
 
@@ -594,6 +656,95 @@ void FRigPhysicsSolver::UpdatePrePhysics(FRigVMExecuteContext& ExecuteContext, c
 					FTransform KinematicTargetTM = ConvertComponentSpaceTransformToSimSpace(KinematicTargetCS);
 					Record.ActorHandle->SetKinematicTarget(KinematicTargetTM);
 				}
+
+				// Now update the joint targets
+				if (Record.JointHandle)
+				{
+					if (Chaos::FPBDJointConstraintHandle* Constraint = Record.JointHandle->GetConstraint())
+					{
+						const Chaos::FPBDJointSettings& Settings = Constraint->GetSettings();
+
+						if (Settings.AngularDriveStiffness.SquaredLength() > 0 ||
+							Settings.AngularDriveDamping.SquaredLength() > 0 ||
+							Settings.LinearDriveStiffness.SquaredLength() > 0 ||
+							Settings.LinearDriveDamping.SquaredLength() > 0)
+						{
+							// Multiplier on the velocity calculated from the current and previous target
+							float TargetVelocityMultiplier = 1.0f;
+
+							FRigElementKey ChildSourceKey;
+							if (const FRigPhysicsComponent* ChildPhysicsComponent = Cast<FRigPhysicsComponent>(
+								Hierarchy->FindComponent(ComponentKey)))
+							{
+								ChildSourceKey = ChildPhysicsComponent->BodySolverSettings.SourceBone;
+								if (!ChildSourceKey.IsValid())
+								{
+									ChildSourceKey = ComponentKey.ElementKey;
+								}
+								TargetVelocityMultiplier = 
+									ChildPhysicsComponent->Joint.Drive.SkeletalAnimationVelocityMultiplier;
+							}
+
+							FRigElementKey ParentSourceKey;
+							if (const FRigPhysicsComponent* ParentPhysicsComponent = Cast<FRigPhysicsComponent>(
+								Hierarchy->FindComponent(Record.ParentBodyComponentKey)))
+							{
+								ParentSourceKey = ParentPhysicsComponent->BodySolverSettings.SourceBone;
+								if (!ParentSourceKey.IsValid())
+								{
+									ParentSourceKey = Record.ParentBodyComponentKey.ElementKey;
+								}
+							}
+
+							if (!ChildSourceKey.IsValid() || !ParentSourceKey.IsValid())
+							{
+								continue;
+							}
+
+							// Note that the drive operates between a parent and child part, so we
+							// don't need to worry about global/component (etc) space.
+							const FTransform ChildTM = Hierarchy->GetGlobalTransform(ChildSourceKey);
+							const FTransform ParentTM = Hierarchy->GetGlobalTransform(ParentSourceKey);
+
+							const FTransform ComponentSpaceParentFrameTM = Settings.ConnectorTransforms[1] * ParentTM;
+							const FTransform ComponentSpaceChildFrameTM = Settings.ConnectorTransforms[0] * ChildTM;
+
+							const FTransform DriveTargetTM = 
+								ComponentSpaceChildFrameTM * ComponentSpaceParentFrameTM.Inverse();
+
+							Constraint->SetLinearDrivePositionTarget(DriveTargetTM.GetTranslation());
+							Constraint->SetAngularDrivePositionTarget(DriveTargetTM.GetRotation());
+
+							if (Record.PreviousDriveTargetUpdateCounter + 1 == UpdateCounter && 
+								TargetVelocityMultiplier > 0)
+							{
+								if (DeltaTime > SMALL_NUMBER)
+								{
+									const FTransform DriveTargetTMDelta = 
+										Record.PreviousDriveTargetTM.Inverse() * DriveTargetTM;
+									const FVector Velocity = DriveTargetTMDelta.GetTranslation() / DeltaTime;
+									const FVector AngularVelocity = 
+										DriveTargetTMDelta.GetRotation().GetShortestArcWith(
+											FQuat::Identity).ToRotationVector() / DeltaTime;
+									if (!Velocity.ContainsNaN() && !AngularVelocity.ContainsNaN())
+									{
+										Constraint->SetLinearDriveVelocityTarget(
+											Velocity * TargetVelocityMultiplier);
+										Constraint->SetAngularDriveVelocityTarget(
+											AngularVelocity * TargetVelocityMultiplier);
+									}
+								}
+							}
+							else
+							{
+								Constraint->SetLinearDriveVelocityTarget(Chaos::FVec3(0));
+								Constraint->SetAngularDriveVelocityTarget(Chaos::FVec3(0));
+							}
+							BodyRecordPair.Value.PreviousDriveTargetUpdateCounter = UpdateCounter;
+							BodyRecordPair.Value.PreviousDriveTargetTM = DriveTargetTM;
+						}
+					}
+				}
 			}
 		}
 	}
@@ -636,6 +787,7 @@ void FRigPhysicsSolver::UpdatePostPhysics()
 					else
 					{
 						FTransform TargetComponentSpaceTM = ConvertSimSpaceTransformToComponentSpace(TargetSimSpaceTM);
+						// Danny TODO should we be "updating" the children here or not? 
 						Hierarchy->SetGlobalTransform(TargetElementKey, TargetComponentSpaceTM);
 					}
 				}
@@ -720,11 +872,13 @@ void FRigPhysicsSolver::StepSimulation(FRigVMExecuteContext& ExecuteContext, dou
 		SimulationSpaceData.LinearAcceleration,
 		SimulationSpaceData.AngularAcceleration);
 
-	UpdatePrePhysics(ExecuteContext, SimulationSpaceData);
+	UpdatePrePhysics(ExecuteContext, DeltaTime, SimulationSpaceData);
 
 	Simulation->Simulate(
 		DeltaTime, SolverSettings.MaxDeltaTime, SolverSettings.MaxTimeSteps, 
 		SimulationSpaceData.Gravity, &ChaosJointSolverSettings);
 
 	UpdatePostPhysics();
+
+	++UpdateCounter;
 }
diff --git a/Engine/Plugins/Experimental/PhysicsControl/Source/ControlRigPhysics/Private/RigPhysicsSolver.h b/Engine/Plugins/Experimental/PhysicsControl/Source/ControlRigPhysics/Private/RigPhysicsSolver.h
--- a/Engine/Plugins/Experimental/PhysicsControl/Source/ControlRigPhysics/Private/RigPhysicsSolver.h
+++ b/Engine/Plugins/Experimental/PhysicsControl/Source/ControlRigPhysics/Private/RigPhysicsSolver.h
@@ -35,12 +35,26 @@ struct FRigBodyRecord
 {
 	GENERATED_BODY()
 
+	// Things that are set during instantiation
 	ImmediatePhysics::FActorHandle* ActorHandle = nullptr;
 	ImmediatePhysics::FJointHandle* JointHandle = nullptr;
-	FPhysicsControlModifierData     BodyData;
+	// When we have a joint, this is the key for the parent body. Note that the child body component
+	// key is stored as the lookup key for the records, so we don't need to duplicate it here.
+	FRigComponentKey                ParentBodyComponentKey;
 
-	ERigPhysicsKinematicTargetSpace KinematicTargetSpace = ERigPhysicsKinematicTargetSpace::OffsetInBoneSpace;
+	// Things that are updated as the simulation progresses
+
+	// The target for when this body is kinematic
 	FTransform                      KinematicTarget;
+	ERigPhysicsKinematicTargetSpace KinematicTargetSpace = ERigPhysicsKinematicTargetSpace::OffsetInBoneSpace;
+
+	// Desired state of the physics body - this will be written to the simulation during UpdatePrePhysics
+	FPhysicsControlModifierData     BodyData;
+
+	// The drive works with velocities so we store the previous target transform, and when it was stored.
+	FTransform                      PreviousDriveTargetTM;
+	// This is stored from the main solver update counter, marking when the previous drive TM was valid.
+	uint64                          PreviousDriveTargetUpdateCounter = 0;
 };
 
 
@@ -78,7 +92,10 @@ public:
 	void StepSimulation(FRigVMExecuteContext& ExecuteContext, double DeltaTime);
 
 	// Draws shapes etc AND (potentially) enables the low-level Chaos debug draw
-	void Draw(FRigVMDrawInterface* DI, const UWorld* DebugWorld) const;
+	void Draw(
+		FRigVMDrawInterface*                    DI,
+		const FRigPhysicsVisualizationSettings& VisualizationSettings,
+		const UWorld*                           DebugWorld) const;
 
 private:
 	// Used by the world-space to simulation-space motion transfer system in Component- or
@@ -129,7 +146,8 @@ private:
 	// change our records (apart from resetting pointers to actors etc).
 	void DestroyPhysicsSimulation();
 
-	void UpdatePrePhysics(FRigVMExecuteContext& ExecuteContext, const FSimulationSpaceData& SimulationSpaceData);
+	void UpdatePrePhysics(
+		FRigVMExecuteContext& ExecuteContext, const double DeltaTime, const FSimulationSpaceData& SimulationSpaceData);
 	void UpdatePostPhysics();
 
 	void InitSimulationSpace(const FTransform& ComponentToWorld, const FTransform& BoneToComponent);
@@ -161,12 +179,20 @@ private:
 	FRigPhysicsSolverSettings SolverSettings;
 	FRigPhysicsSimulationSpaceSettings SimulationSpaceSettings;
 
+	// Used to store things like collision shapes.
 	TArray<ImmediatePhysics::FActorHandle*> AdditionalActorHandles;
 
 	Chaos::FPBDJointSolverSettings ChaosJointSolverSettings;
 
 	FSimulationSpaceState SimulationSpaceState;
 
+	// This is incremented at the end of each simulation step, used to identify when previously
+	// calculated values are valid
+	// Note that the universe will roll over before a uint64 does
+	uint64 UpdateCounter = 0;
+
+	// This will be set to false after instantiation, which can happen manually during construction,
+	// or be postponed until there is a simulation step.
 	bool bNeedToInstantiate = true;
 };
 
diff --git a/Engine/Plugins/Experimental/PhysicsControl/Source/ControlRigPhysics/Private/RigPhysicsSolverVisualization.cpp b/Engine/Plugins/Experimental/PhysicsControl/Source/ControlRigPhysics/Private/RigPhysicsSolverVisualization.cpp
--- a/Engine/Plugins/Experimental/PhysicsControl/Source/ControlRigPhysics/Private/RigPhysicsSolverVisualization.cpp
+++ b/Engine/Plugins/Experimental/PhysicsControl/Source/ControlRigPhysics/Private/RigPhysicsSolverVisualization.cpp
@@ -27,20 +27,19 @@
 
 namespace RigPhysicsSolverDraw
 {
-	float Thickness = 0.3f;
 	FLinearColor DynamicColor = FColor::Yellow;
 	FLinearColor KinematicColor = FColor::Blue;
-	int32 Detail = 16;
 }
 
 //======================================================================================================================
 static void DrawShapes(
-	FRigVMDrawInterface*                  DI,
-	const bool                            bIsKinematic,
-	const Chaos::FGeometryParticleHandle* Particle,
-	const FTransform&                     ShapeTransform, 
-	const Chaos::FImplicitObject*         ImplicitObject,
-	const Chaos::FPerShapeData*           Shape)
+	FRigVMDrawInterface*                    DI,
+	const FRigPhysicsVisualizationSettings& VisualizationSettings,
+	const bool                              bIsKinematic,
+	const Chaos::FGeometryParticleHandle*   Particle,
+	const FTransform&                       ShapeTransform, 
+	const Chaos::FImplicitObject*           ImplicitObject,
+	const Chaos::FPerShapeData*             Shape)
 {
 	using namespace Chaos;
 
@@ -67,7 +66,9 @@ static void DrawShapes(
 		FTransform TransformedTransform = FTransform(
 			ShapeTransform.GetRotation() * Transformed->GetTransform().GetRotation(),
 			ShapeTransform.TransformPosition(Transformed->GetTransform().GetLocation()));
-		DrawShapes(DI, bIsKinematic, Particle, TransformedTransform, Transformed->GetTransformedObject(), Shape);
+		DrawShapes(
+			DI, VisualizationSettings, bIsKinematic, Particle, 
+			TransformedTransform, Transformed->GetTransformedObject(), Shape);
 		return;
 	}
 	case ImplicitObjectType::Sphere:
@@ -75,7 +76,7 @@ static void DrawShapes(
 		const Chaos::FSphere* Sphere = ImplicitObject->GetObject<Chaos::FSphere>();
 		DI->DrawSphere(ShapeTransform, FTransform(FVec3(Sphere->GetCenterOfMass())), Sphere->GetRadiusf(),
 			bIsKinematic ? RigPhysicsSolverDraw::KinematicColor : RigPhysicsSolverDraw::DynamicColor,
-			RigPhysicsSolverDraw::Thickness, RigPhysicsSolverDraw::Detail);
+			VisualizationSettings.LineThickness, VisualizationSettings.ShapeDetail);
 		break;
 	}
 	case ImplicitObjectType::Box:
@@ -83,7 +84,7 @@ static void DrawShapes(
 		const TBox<FReal, 3>* Box = ImplicitObject->GetObject<TBox<FReal, 3>>();
 		DI->DrawBox(ShapeTransform, FTransform(FQuat::Identity, Box->GetCenter(), Box->Extents()),
 			bIsKinematic ? RigPhysicsSolverDraw::KinematicColor : RigPhysicsSolverDraw::DynamicColor,
-			RigPhysicsSolverDraw::Thickness);
+			VisualizationSettings.LineThickness);
 		break;
 	}
 	case ImplicitObjectType::Capsule:
@@ -93,7 +94,7 @@ static void DrawShapes(
 		DI->DrawCapsule(ShapeTransform, FTransform(Q, Capsule->GetCenterOfMass()),
 			Capsule->GetRadiusf(), Capsule->GetHeightf(),
 			bIsKinematic ? RigPhysicsSolverDraw::KinematicColor : RigPhysicsSolverDraw::DynamicColor,
-			RigPhysicsSolverDraw::Thickness);
+			VisualizationSettings.LineThickness);
 		break;
 	}
 	default:
@@ -105,9 +106,10 @@ static void DrawShapes(
 
 //======================================================================================================================
 static void DrawActor(
-	FRigVMDrawInterface*                  DI,
-	const FTransform&                     SpaceTransform,
-	const ImmediatePhysics::FActorHandle& ActorHandle)
+	FRigVMDrawInterface*                    DI,
+	const FRigPhysicsVisualizationSettings& VisualizationSettings,
+	const FTransform&                       SpaceTransform,
+	const ImmediatePhysics::FActorHandle&   ActorHandle)
 {
 	bool bIsKinematic = ActorHandle.GetIsKinematic();
 	const Chaos::FGeometryParticleHandle* GeometryParticleHandle = ActorHandle.GetParticle();
@@ -116,7 +118,7 @@ static void DrawActor(
 	{
 		const Chaos::FImplicitObject* ImplicitObject = ShapeInstance->GetGeometry();
 		DrawShapes(
-			DI, bIsKinematic, GeometryParticleHandle, 
+			DI, VisualizationSettings, bIsKinematic, GeometryParticleHandle,
 			ParticleTransform, ImplicitObject, ShapeInstance.Get());
 	}
 }
@@ -142,9 +144,10 @@ float GetLimitAngleRadians(float LimitAngle, Chaos::EJointMotionType LimitType)
 
 //======================================================================================================================
 static void DrawJoint(
-	FRigVMDrawInterface*                  DI,
-	const FTransform&                     SpaceTransform,
-	const ImmediatePhysics::FJointHandle& JointHandle)
+	FRigVMDrawInterface*                    DI,
+	const FRigPhysicsVisualizationSettings& VisualizationSettings,
+	const FTransform&                       SpaceTransform,
+	const ImmediatePhysics::FJointHandle&   JointHandle)
 {
 	const Chaos::FPBDJointConstraintHandle* ConstraintHandle = JointHandle.GetConstraint();
 	if (!ConstraintHandle)
@@ -172,8 +175,8 @@ static void DrawJoint(
 	static bool bDrawAsAxes = false;
 	if (bDrawAsAxes)
 	{
-		DI->DrawAxes(LimitFrame, FTransform(), Size, RigPhysicsSolverDraw::Thickness * 2);
-		DI->DrawAxes(DialFrame, FTransform(), Size, RigPhysicsSolverDraw::Thickness);
+		DI->DrawAxes(LimitFrame, FTransform(), Size, VisualizationSettings.LineThickness * 2);
+		DI->DrawAxes(DialFrame, FTransform(), Size, VisualizationSettings.LineThickness);
 	}
 
 	{
@@ -231,7 +234,7 @@ static void DrawJoint(
 				JointSettings.AngularMotionTypes[Swing2Index] == Chaos::EJointMotionType::Free)
 			{
 				DI->DrawSphere(LimitFrame, FTransform(), Size * 0.2f, FColor::White, 
-					RigPhysicsSolverDraw::Thickness, RigPhysicsSolverDraw::Detail);
+					VisualizationSettings.LineThickness, VisualizationSettings.ShapeDetail);
 			}
 			else
 			{
@@ -240,15 +243,15 @@ static void DrawJoint(
 				const float Swing1LimitAngle = LimitAngleRadians[Swing1Index];
 				const float Swing2LimitAngle = LimitAngleRadians[Swing2Index];
 				DI->DrawCone(
-					ConeLimitTM, FTransform(), Swing1LimitAngle, Swing2LimitAngle, RigPhysicsSolverDraw::Detail, 
-					true, FColor::Green, LimitMaterialX->GetRenderProxy(), RigPhysicsSolverDraw::Thickness);
+					ConeLimitTM, FTransform(), Swing1LimitAngle, Swing2LimitAngle, VisualizationSettings.ShapeDetail, 
+					true, FColor::Green, LimitMaterialX->GetRenderProxy(), VisualizationSettings.LineThickness);
 			}
 			// Draw the swing Dial indicator - shows the current orientation of the child frame
 			// relative to the parent frame on the swing axis.
 			// Start the arrow at the limit position as it can be confusing if there is joint separation
 			DI->DrawArrow(FTransform(LimitFrame.GetTranslation()), 
 				DialFrame.GetUnitAxis(EAxis::X) * ArrowLength,
-				DialFrame.GetUnitAxis(EAxis::Y), FColor::Red, RigPhysicsSolverDraw::Thickness);
+				DialFrame.GetUnitAxis(EAxis::Y), FColor::Red, VisualizationSettings.LineThickness);
 
 			if (bSwing1Violated || bSwing2Violated)
 			{
@@ -258,7 +261,7 @@ static void DrawJoint(
 					DialFrame.GetRotation(),
 					LimitFrame.GetTranslation() + DialFrame.GetUnitAxis(EAxis::X) * ArrowLength),
 					FTransform(),
-					ArrowLength * 0.01f, FColor::Orange, RigPhysicsSolverDraw::Thickness * 4);
+					ArrowLength * 0.01f, FColor::Orange, VisualizationSettings.LineThickness * 4);
 			}
 
 		}
@@ -271,8 +274,8 @@ static void DrawJoint(
 			ConeLimitTM.SetScale3D(FVector(Size));
 			const float TwistLimitAngle = LimitAngleRadians[TwistIndex];
 			DI->DrawCone(ConeLimitTM, FTransform(FQuat::MakeFromRotationVector(FVector(0, 1, 0) * -HALF_PI)),
-				TwistLimitAngle, 0.0f, RigPhysicsSolverDraw::Detail, true, FColor::Green,
-				LimitMaterialYAxis->GetRenderProxy(), RigPhysicsSolverDraw::Thickness);
+				TwistLimitAngle, 0.0f, VisualizationSettings.ShapeDetail, true, FColor::Green,
+				LimitMaterialYAxis->GetRenderProxy(), VisualizationSettings.LineThickness);
 
 			// Draw the twist Dial indicator - shows the current orientation of the child frame
 			// relative to the parent frame on the twist axis.
@@ -282,7 +285,7 @@ static void DrawJoint(
 
 			DI->DrawArrow(FTransform(LimitFrame.GetTranslation()),
 				TwistArrow * ArrowLength,
-				TwistSideDir, FColor::Blue, RigPhysicsSolverDraw::Thickness);
+				TwistSideDir, FColor::Blue, VisualizationSettings.LineThickness);
 
 			if (bTwistViolated)
 			{
@@ -290,14 +293,17 @@ static void DrawJoint(
 					DialFrame.GetRotation(),
 					LimitFrame.GetTranslation() + TwistArrow * ArrowLength),
 					FTransform(),
-					ArrowLength * 0.01f, FColor::Orange, RigPhysicsSolverDraw::Thickness * 4);
+					ArrowLength * 0.01f, FColor::Orange, VisualizationSettings.LineThickness * 4);
 			}
 		}
 	}
 }
 
 //======================================================================================================================
-void FRigPhysicsSolver::Draw(FRigVMDrawInterface* DI, const UWorld* DebugWorld) const
+void FRigPhysicsSolver::Draw(
+	FRigVMDrawInterface*                    DI,
+	const FRigPhysicsVisualizationSettings& VisualizationSettings,
+	const UWorld*                           DebugWorld) const
 {
 #if CHAOS_DEBUG_DRAW
 	// Danny TODO enable/disable chaos debug draw with a cvar or flag
@@ -320,7 +326,7 @@ void FRigPhysicsSolver::Draw(FRigVMDrawInterface* DI, const UWorld* DebugWorld)
 
 	for (const ImmediatePhysics::FActorHandle* ActorHandle : AdditionalActorHandles)
 	{
-		DrawActor(DI, SpaceTransform, *ActorHandle);
+		DrawActor(DI, VisualizationSettings, SpaceTransform, *ActorHandle);
 	}
 
 	for (const TPair<FRigComponentKey, FRigBodyRecord>& BodyRecordPair : BodyRecords)
@@ -331,12 +337,12 @@ void FRigPhysicsSolver::Draw(FRigVMDrawInterface* DI, const UWorld* DebugWorld)
 
 		if (JointHandle)
 		{
-			DrawJoint(DI, SpaceTransform, *JointHandle);
+			DrawJoint(DI, VisualizationSettings, SpaceTransform, *JointHandle);
 		}
 
 		if (ActorHandle)
 		{
-			DrawActor(DI, SpaceTransform, *ActorHandle);
+			DrawActor(DI, VisualizationSettings, SpaceTransform, *ActorHandle);
 		}
 	}
 }
diff --git a/Engine/Plugins/Experimental/PhysicsControl/Source/ControlRigPhysics/Public/RigPhysicsData.h b/Engine/Plugins/Experimental/PhysicsControl/Source/ControlRigPhysics/Public/RigPhysicsData.h
--- a/Engine/Plugins/Experimental/PhysicsControl/Source/ControlRigPhysics/Public/RigPhysicsData.h
+++ b/Engine/Plugins/Experimental/PhysicsControl/Source/ControlRigPhysics/Public/RigPhysicsData.h
@@ -17,6 +17,22 @@ class FRigElementKeyRedirector;
 
 const float RigPhysicsDefaultShapeSize = 10.0f;
 
+USTRUCT(BlueprintType)
+struct FRigPhysicsVisualizationSettings
+{
+	GENERATED_BODY()
+
+	/* Enable visualization (includes enabling low-level chaos debug draw) */
+	UPROPERTY(BlueprintReadWrite, EditAnywhere, Category = Visualization)
+	bool bEnableVisualization = true;
+
+	UPROPERTY(BlueprintReadWrite, EditAnywhere, Category = Visualization)
+	float LineThickness = 1.0f;
+
+	UPROPERTY(BlueprintReadWrite, EditAnywhere, Category = Visualization)
+	int32 ShapeDetail = 16;
+};
+
 /** 
  * Specifies what space a kinematic target is defined in 
  * Danny TODO this is much better than the similar mechanism in PCC and RBWC (which has a 
@@ -435,6 +451,85 @@ struct CONTROLRIGPHYSICS_API FRigPhysicsDynamics
 	FVector MomentsOfInertiaOverride = FVector(1.0f, 1.0f, 1.0f);
 };
 
+// This represents the motor drive associated with a physics joint, that can drive the bodies
+// towards a target pose.
+USTRUCT(BlueprintType)
+struct CONTROLRIGPHYSICS_API FRigPhysicsDrive
+{
+	GENERATED_BODY()
+
+	friend FArchive& operator <<(FArchive& Ar, FRigPhysicsDrive& Data)
+	{
+		Ar << Data.bEnable;
+		Ar << Data.LinearStrength;
+		Ar << Data.LinearDampingRatio;
+		Ar << Data.LinearExtraDamping;
+		Ar << Data.MaxForce;
+		Ar << Data.AngularStrength;
+		Ar << Data.AngularDampingRatio;
+		Ar << Data.AngularExtraDamping;
+		Ar << Data.MaxTorque;
+		Ar << Data.SkeletalAnimationVelocityMultiplier;
+		return Ar;
+	}
+
+	UPROPERTY(BlueprintReadWrite, EditAnywhere, Category = PhysicsDrive)
+	bool bEnable = true;
+
+	/** The strength used to drive linear motion */
+	UPROPERTY(EditAnywhere, BlueprintReadWrite, Category = PhysicsControl, meta = (ClampMin = "0.0"))
+	float LinearStrength = 0.0f;
+
+	/**
+	 * The amount of damping associated with the linear strength. A value of 1 Results in critically
+	 * damped motion where the control drives as quickly as possible to the target without overshooting.
+	 * Values > 1 result in more damped motion, and values below 1 result in faster, but more "wobbly" motion.
+	 */
+	UPROPERTY(EditAnywhere, BlueprintReadWrite, Category = PhysicsControl, meta = (ClampMin = "0.0"))
+	float LinearDampingRatio = 1.0f;
+
+	/**
+	 * The amount of additional linear damping. This is added to the damping that comes from LinearDampingRatio
+	 * and can be useful when you want damping even when LinearStrength is zero.
+	 */
+	UPROPERTY(EditAnywhere, BlueprintReadWrite, Category = PhysicsControl, meta = (ClampMin = "0.0"))
+	float LinearExtraDamping = 0.0f;
+
+	/** The maximum force used to drive the linear motion. Zero indicates no limit. */
+	UPROPERTY(EditAnywhere, BlueprintReadWrite, Category = PhysicsControl, meta = (ClampMin = "0.0"))
+	float MaxForce = 0.0;
+
+	/** The strength used to drive angular motion */
+	UPROPERTY(EditAnywhere, BlueprintReadWrite, Category = PhysicsControl, meta = (ClampMin = "0.0"))
+	float AngularStrength = 0.0f;
+
+	/**
+	 * The amount of damping associated with the angular strength. A value of 1 Results in critically
+	 * damped motion where the control drives as quickly as possible to the target without overshooting.
+	 * Values > 1 result in more damped motion, and values below 1 result in faster, but more "wobbly" motion.
+	 */
+	UPROPERTY(EditAnywhere, BlueprintReadWrite, Category = PhysicsControl, meta = (ClampMin = "0.0"))
+	float AngularDampingRatio = 1.0f;
+
+	/**
+	 * The amount of additional angular damping. This is added to the damping that comes from AngularDampingRatio
+	 * and can be useful when you want damping even when AngularStrength is zero.
+	 */
+	UPROPERTY(EditAnywhere, BlueprintReadWrite, Category = PhysicsControl, meta = (ClampMin = "0.0"))
+	float AngularExtraDamping = 0.0f;
+
+	/** The maximum torque used to drive the angular motion. Zero indicates no limit. */
+	UPROPERTY(EditAnywhere, BlueprintReadWrite, Category = PhysicsControl, meta = (ClampMin = "0.0"))
+	float MaxTorque = 0.0f;
+
+	/**
+	 * The amount of skeletal animation velocity to use in the targets
+	 * Danny TODO implement - this requires tracking/calculating the velocity
+	 */
+	UPROPERTY(EditAnywhere, BlueprintReadWrite, Category = PhysicsControl)
+	float SkeletalAnimationVelocityMultiplier = 1.0f;
+};
+
 /**
  * This represents an "articulation" - a general purpose "character joint" consisting of
  *   * A linear limit (a linear constraint), which is normally used to pin two bones together
@@ -499,10 +594,11 @@ struct CONTROLRIGPHYSICS_API FRigPhysicsJoint
 	UPROPERTY(BlueprintReadWrite, EditAnywhere, Category = PhysicsJoint)
 	FTransform ExtraChildOffset;
 
+	// 0 locks the movement. -ve makes the movement free
 	UPROPERTY(BlueprintReadWrite, EditAnywhere, Category = PhysicsJoint)
-	FVector LinearLimit = FVector::ZeroVector;
+	float LinearLimit = 0.0f;
 
-	// Angular limit in degrees (twist, swing1, swing2)
+	// Angular limit in degrees (twist, swing1, swing2). -ve indicates the movement is free
 	UPROPERTY(BlueprintReadWrite, EditAnywhere, Category = PhysicsJoint)
 	FVector AngularLimit = FVector::ZeroVector;
 
@@ -511,9 +607,9 @@ struct CONTROLRIGPHYSICS_API FRigPhysicsJoint
 	bool bDisableCollision = true;
 
 	// The amount of linear projection to reduce joint separation when the linear constraint is
-	// locked or has a hard limit.
+	// locked or has a hard limit. A value of one will apply full projection, but can introduce artefacts.
 	UPROPERTY(BlueprintReadWrite, EditAnywhere, Category = PhysicsJoint, meta = (ClampMin = 0, ClampMax = 1))
-	float LinearProjectionAmount = 1.0f;
+	float LinearProjectionAmount = 0.5f;
 
 	// The amount of angular projection to reduce joint separation when the angular constraint is
 	// locked or has a hard limit. Note that projecting back to angular limits will often break the
@@ -525,6 +621,10 @@ struct CONTROLRIGPHYSICS_API FRigPhysicsJoint
 	// body, making the joint behave in one direction.
 	UPROPERTY(BlueprintReadWrite, EditAnywhere, Category = PhysicsJoint, meta = (ClampMin = 0))
 	float ParentInverseMassScale = 1.0f;
+
+	// Optional motor/drive associated with the physics joint
+	UPROPERTY(EditAnywhere, BlueprintReadWrite, Category = PhysicsControl, meta = (ShowOnlyInnerProperties))
+	FRigPhysicsDrive Drive;
 };
 
 /**
diff --git a/Engine/Plugins/Experimental/PhysicsControl/Source/ControlRigPhysics/Public/RigPhysicsExecution.h b/Engine/Plugins/Experimental/PhysicsControl/Source/ControlRigPhysics/Public/RigPhysicsExecution.h
--- a/Engine/Plugins/Experimental/PhysicsControl/Source/ControlRigPhysics/Public/RigPhysicsExecution.h
+++ b/Engine/Plugins/Experimental/PhysicsControl/Source/ControlRigPhysics/Public/RigPhysicsExecution.h
@@ -83,7 +83,7 @@ struct CONTROLRIGPHYSICS_API FRigUnit_StepPhysicsSolver : public FRigUnitMutable
 	}
 
 	RIGVM_METHOD()
-		virtual void Execute() override;
+	virtual void Execute() override;
 
 	/*
 	 * The solver to relate this new physics element to
@@ -97,9 +97,8 @@ struct CONTROLRIGPHYSICS_API FRigUnit_StepPhysicsSolver : public FRigUnitMutable
 	UPROPERTY(meta = (Input))
 	float DeltaTimeOverride;
 
-	/* Enable visualization (includes enabling low-level chaos debug draw) */
 	UPROPERTY(meta = (Input))
-	bool bEnableVisualization = true;
+	FRigPhysicsVisualizationSettings VisualizationSettings;
 };
 
 /**
diff --git a/Engine/Plugins/Experimental/PhysicsControl/Source/PhysicsControl/Private/PhysicsControlComponent.cpp b/Engine/Plugins/Experimental/PhysicsControl/Source/PhysicsControl/Private/PhysicsControlComponent.cpp
--- a/Engine/Plugins/Experimental/PhysicsControl/Source/PhysicsControl/Private/PhysicsControlComponent.cpp
+++ b/Engine/Plugins/Experimental/PhysicsControl/Source/PhysicsControl/Private/PhysicsControlComponent.cpp
@@ -3,7 +3,7 @@
 #include "PhysicsControlComponent.h"
 #include "PhysicsControlLog.h"
 #include "PhysicsControlRecord.h"
-#include "PhysicsControlComponentHelpers.h"
+#include "PhysicsControlHelpers.h"
 #include "PhysicsControlOperatorNameGeneration.h"
 
 #include "PhysicsEngine/ConstraintInstance.h"
diff --git a/Engine/Plugins/Experimental/PhysicsControl/Source/PhysicsControl/Private/PhysicsControlComponentHelpers.cpp b/Engine/Plugins/Experimental/PhysicsControl/Source/PhysicsControl/Private/PhysicsControlComponentHelpers.cpp
deleted file mode 100644
diff --git a/Engine/Plugins/Experimental/PhysicsControl/Source/PhysicsControl/Private/PhysicsControlComponentHelpers.h b/Engine/Plugins/Experimental/PhysicsControl/Source/PhysicsControl/Private/PhysicsControlComponentHelpers.h
deleted file mode 100644
diff --git a/Engine/Plugins/Experimental/PhysicsControl/Source/PhysicsControl/Private/PhysicsControlComponentImpl.cpp b/Engine/Plugins/Experimental/PhysicsControl/Source/PhysicsControl/Private/PhysicsControlComponentImpl.cpp
--- a/Engine/Plugins/Experimental/PhysicsControl/Source/PhysicsControl/Private/PhysicsControlComponentImpl.cpp
+++ b/Engine/Plugins/Experimental/PhysicsControl/Source/PhysicsControl/Private/PhysicsControlComponentImpl.cpp
@@ -2,7 +2,7 @@
 
 #include "PhysicsControlComponent.h"
 #include "PhysicsControlLog.h"
-#include "PhysicsControlComponentHelpers.h"
+#include "PhysicsControlHelpers.h"
 
 #include "PhysicsEngine/BodySetup.h"
 #include "PhysicsEngine/BodyInstance.h"
diff --git a/Engine/Plugins/Experimental/PhysicsControl/Source/PhysicsControl/Private/PhysicsControlHelpers.cpp b/Engine/Plugins/Experimental/PhysicsControl/Source/PhysicsControl/Private/PhysicsControlHelpers.cpp
new file mode 100644
--- /dev/null
+++ b/Engine/Plugins/Experimental/PhysicsControl/Source/PhysicsControl/Private/PhysicsControlHelpers.cpp
@@ -0,0 +1,146 @@
+// Copyright Epic Games, Inc. All Rights Reserved.
+
+#include "PhysicsControlHelpers.h"
+#include "PhysicsControlData.h"
+#include "PhysicsEngine/ConstraintInstance.h"
+
+#include "Components/PrimitiveComponent.h"
+#include "Components/SkeletalMeshComponent.h"
+
+namespace UE
+{
+namespace PhysicsControl
+{
+
+//======================================================================================================================
+void ConvertStrengthToSpringParams(
+	FVector& OutSpring, FVector& OutDamping, 
+	const FVector& InStrength, const FVector& InDampingRatio, const FVector& InExtraDamping)
+{
+	ConvertStrengthToSpringParams(OutSpring.X, OutDamping.X, InStrength.X, InDampingRatio.X, InExtraDamping.X);
+	ConvertStrengthToSpringParams(OutSpring.Y, OutDamping.Y, InStrength.Y, InDampingRatio.Y, InExtraDamping.Y);
+	ConvertStrengthToSpringParams(OutSpring.Z, OutDamping.Z, InStrength.Z, InDampingRatio.Z, InExtraDamping.Z);
+}
+
+inline double GetLinearDriveStiffness(const FLinearDriveConstraint& InDrive)
+{
+	return (
+		(InDrive.XDrive.bEnablePositionDrive ? InDrive.XDrive.Stiffness : 0.0) + 
+		(InDrive.YDrive.bEnablePositionDrive ? InDrive.YDrive.Stiffness : 0.0) +
+		(InDrive.ZDrive.bEnablePositionDrive ? InDrive.ZDrive.Stiffness : 0.0) 
+		) / 3.0;
+}
+inline double GetLinearDriveDamping(const FLinearDriveConstraint& InDrive)
+{
+	return (
+		(InDrive.XDrive.bEnableVelocityDrive ? InDrive.XDrive.Damping : 0.0) +
+		(InDrive.YDrive.bEnableVelocityDrive ? InDrive.YDrive.Damping : 0.0) + 
+		(InDrive.ZDrive.bEnableVelocityDrive ? InDrive.ZDrive.Damping : 0.0) 
+		) / 3.0;
+}
+inline double GetLinearDriveMaxForce(const FLinearDriveConstraint& InDrive)
+{
+	return FVector(InDrive.XDrive.MaxForce, InDrive.YDrive.MaxForce, InDrive.ZDrive.MaxForce).Length();
+}
+
+inline double GetAngularDriveStiffness(const FAngularDriveConstraint& InDrive)
+{
+	if (InDrive.AngularDriveMode == EAngularDriveMode::SLERP)
+	{
+		return InDrive.SlerpDrive.bEnablePositionDrive ? InDrive.SlerpDrive.Stiffness : 0.0;
+	}
+	else
+	{
+		return (
+			(InDrive.TwistDrive.bEnablePositionDrive ? InDrive.TwistDrive.Stiffness : 0.0) + 
+			(InDrive.SwingDrive.bEnablePositionDrive ? InDrive.SwingDrive.Stiffness : 0.0)
+			) / 2.0;
+	}
+}
+inline double GetAngularDriveDamping(const FAngularDriveConstraint& InDrive)
+{
+	if (InDrive.AngularDriveMode == EAngularDriveMode::SLERP)
+	{
+		return InDrive.SlerpDrive.bEnableVelocityDrive ? InDrive.SlerpDrive.Damping : 0.0;
+	}
+	else
+	{
+		return (
+			(InDrive.TwistDrive.bEnableVelocityDrive ? InDrive.TwistDrive.Damping : 0.0) +
+			(InDrive.SwingDrive.bEnableVelocityDrive ? InDrive.SwingDrive.Damping : 0.0)
+			) / 2.0;
+	}
+}
+inline double GetAngularDriveMaxTorque(const FAngularDriveConstraint& InDrive)
+{
+	if (InDrive.AngularDriveMode == EAngularDriveMode::SLERP)
+	{
+		return InDrive.SlerpDrive.MaxForce;
+	}
+	else
+	{
+		return FVector2D(InDrive.TwistDrive.MaxForce, InDrive.SwingDrive.MaxForce).Length();
+	}
+}
+
+//======================================================================================================================
+void ConvertConstraintProfileToControlData(
+	FPhysicsControlData& OutControlData, const FConstraintProfileProperties& InProfileProperties)
+{
+	ConvertSpringToStrengthParams(
+		OutControlData.LinearStrength, 
+		OutControlData.LinearDampingRatio, 
+		OutControlData.LinearExtraDamping,
+		GetLinearDriveStiffness(InProfileProperties.LinearDrive), 
+		GetLinearDriveDamping(InProfileProperties.LinearDrive));
+
+	OutControlData.MaxForce = float(GetLinearDriveMaxForce(InProfileProperties.LinearDrive));
+
+	ConvertSpringToStrengthParams(
+		OutControlData.AngularStrength,
+		OutControlData.AngularDampingRatio,
+		OutControlData.AngularExtraDamping,
+		GetAngularDriveStiffness(InProfileProperties.AngularDrive),
+		GetAngularDriveDamping(InProfileProperties.AngularDrive));
+
+	OutControlData.MaxTorque = float(GetAngularDriveMaxTorque(InProfileProperties.AngularDrive));
+
+	OutControlData.bEnabled = true;
+}
+
+//======================================================================================================================
+FBodyInstance* GetBodyInstance(UPrimitiveComponent* Component, const FName BoneName)
+{
+	if (USkeletalMeshComponent* SkeletalMeshComponent = Cast<USkeletalMeshComponent>(Component))
+	{
+		return SkeletalMeshComponent->GetBodyInstance(BoneName);
+	}
+	if (Component)
+	{
+		return Component->GetBodyInstance();
+	}
+	return nullptr;
+}
+
+//======================================================================================================================
+FName GetPhysicalParentBone(USkeletalMeshComponent* SkeletalMeshComponent, FName BoneName)
+{
+	while (true)
+	{
+		FName ParentBoneName = SkeletalMeshComponent->GetParentBone(BoneName);
+		if (ParentBoneName.IsNone() || ParentBoneName == BoneName)
+		{
+			return FName();
+		}
+		if (const FBodyInstance* ParentBodyInstance = GetBodyInstance(SkeletalMeshComponent, ParentBoneName))
+		{
+			return ParentBoneName;
+		}
+		BoneName = ParentBoneName;
+	}
+	return FName();
+}
+
+} // namespace PhysicsControl
+} // namespace UE
+
diff --git a/Engine/Plugins/Experimental/PhysicsControl/Source/PhysicsControl/Private/PhysicsControlRecord.cpp b/Engine/Plugins/Experimental/PhysicsControl/Source/PhysicsControl/Private/PhysicsControlRecord.cpp
--- a/Engine/Plugins/Experimental/PhysicsControl/Source/PhysicsControl/Private/PhysicsControlRecord.cpp
+++ b/Engine/Plugins/Experimental/PhysicsControl/Source/PhysicsControl/Private/PhysicsControlRecord.cpp
@@ -2,7 +2,7 @@
 
 #include "PhysicsControlRecord.h"
 #include "PhysicsControlLog.h"
-#include "PhysicsControlComponentHelpers.h"
+#include "PhysicsControlHelpers.h"
 
 #include "PhysicsEngine/ConstraintInstance.h"
 #include "PhysicsEngine/BodyInstance.h"
diff --git a/Engine/Plugins/Experimental/PhysicsControl/Source/PhysicsControl/Private/RigidBodyWithControl.cpp b/Engine/Plugins/Experimental/PhysicsControl/Source/PhysicsControl/Private/RigidBodyWithControl.cpp
--- a/Engine/Plugins/Experimental/PhysicsControl/Source/PhysicsControl/Private/RigidBodyWithControl.cpp
+++ b/Engine/Plugins/Experimental/PhysicsControl/Source/PhysicsControl/Private/RigidBodyWithControl.cpp
@@ -5,7 +5,7 @@
 #include "PhysicsControlOperatorNameGeneration.h"
 #include "PhysicsControlLog.h"
 #include "PhysicsControlAsset.h"
-#include "PhysicsControlComponentHelpers.h"
+#include "PhysicsControlHelpers.h"
 
 #include "Animation/AnimInstance.h"
 #include "Animation/AnimInstanceProxy.h"
diff --git a/Engine/Plugins/Experimental/PhysicsControl/Source/PhysicsControl/Public/PhysicsControlHelpers.h b/Engine/Plugins/Experimental/PhysicsControl/Source/PhysicsControl/Public/PhysicsControlHelpers.h
new file mode 100644
--- /dev/null
+++ b/Engine/Plugins/Experimental/PhysicsControl/Source/PhysicsControl/Public/PhysicsControlHelpers.h
@@ -0,0 +1,101 @@
+// Copyright Epic Games, Inc. All Rights Reserved.
+
+#pragma once
+
+#include "UObject/NameTypes.h"
+#include "Math/Vector.h"
+
+struct FBodyInstance;
+struct FPhysicsControlData;
+struct FConstraintProfileProperties;
+
+class UPrimitiveComponent;
+class USkeletalMeshComponent;
+
+namespace UE
+{
+namespace PhysicsControl
+{
+
+/** 
+ * Converts strength/damping ratio/extra damping into spring stiffness/damping 
+ */
+PHYSICSCONTROL_API void ConvertStrengthToSpringParams(
+	FVector& OutSpring, FVector& OutDamping, 
+	const FVector& InStrength, const FVector& InDampingRatio, const FVector& InExtraDamping);
+
+/**
+ * Converts the drive settings from the constraint profile into the control data strength/damping etc. 
+ * Approximations will be made if (a) the linear drive has different values for the x/y/z axes, or (b)
+ * the constraint profile is set to use twist/swing instead of slerp for the angular drive.
+ */
+PHYSICSCONTROL_API void ConvertConstraintProfileToControlData(
+	FPhysicsControlData& OutControlData, const FConstraintProfileProperties& InProfileProperties);
+
+/** 
+ * Attempts to find a BodyInstance from the component. If it is a primitive or static mesh component the 
+ * single body instance will be returned. If it is a skeletal mesh then if BoneName can be found, the 
+ * body instance corresponding to that bone will be returned. Otherwise it will return nullptr if 
+ * the bone can't be found.
+ */
+PHYSICSCONTROL_API FBodyInstance* GetBodyInstance(UPrimitiveComponent* Component, const FName BoneName);
+
+/**
+ * Attempts to find the parent physical bone given a skeletal mesh and starting bone. This walks up
+ * the hierarchy, ignoring non-physical bones, until either a physical bone is found, or it has reached
+ * the root without finding a physical bone.
+ */
+PHYSICSCONTROL_API FName GetPhysicalParentBone(USkeletalMeshComponent* SkeletalComponent, FName BoneName);
+
+/**
+ * Converts strength/damping ratio/extra damping into spring stiffness/damping.
+ */
+template<typename TOut>
+void ConvertStrengthToSpringParams(
+	TOut& OutSpring, TOut& OutDamping,
+	double InStrength, double InDampingRatio, double InExtraDamping)
+{
+	TOut AngularFrequency = TOut(InStrength * UE_DOUBLE_TWO_PI);
+	TOut Stiffness = AngularFrequency * AngularFrequency;
+
+	OutSpring = Stiffness;
+	OutDamping = TOut(InExtraDamping + 2 * InDampingRatio * AngularFrequency);
+}
+
+
+/**
+ * Converts spring/damping values into strength/damping ratio/extra damping. This tries to get
+ * as much damping into the damping ratio term as possible, without letting it go above 1.
+ */
+template<typename TOut>
+void ConvertSpringToStrengthParams(
+	TOut& OutStrength, TOut& OutDampingRatio, TOut& OutExtraDamping,
+	const double InSpring, const double InDamping)
+{
+	// Simple calculation to get the strength
+	TOut AngularFrequency = TOut(FMath::Sqrt(InSpring));
+	OutStrength = AngularFrequency / UE_TWO_PI;
+
+	// For damping, try to put as much into the damping ratio as possible, up to a max DR of 1. Then
+	// the rest goes into extra damping.
+	OutDampingRatio = 1;
+	TOut ImpliedDamping = TOut(2.0 * OutDampingRatio * AngularFrequency);
+
+	if (ImpliedDamping < InDamping)
+	{
+		OutExtraDamping = TOut(InDamping) - ImpliedDamping;
+	}
+	else if (AngularFrequency > 0)
+	{
+		OutExtraDamping = 0;
+		OutDampingRatio = TOut(InDamping) / (2 * AngularFrequency);
+	}
+	else
+	{
+		OutDampingRatio = 1;
+		OutExtraDamping = TOut(InDamping);
+	}
+}
+
+} // namespace PhysicsControl
+} // namespace UE
