Build fixes (static analysis warnings)

diff --git a/Engine/Plugins/Experimental/PhysicsControl/Source/ControlRigPhysics/Private/RigPhysicsSolver.cpp b/Engine/Plugins/Experimental/PhysicsControl/Source/ControlRigPhysics/Private/RigPhysicsSolver.cpp
--- a/Engine/Plugins/Experimental/PhysicsControl/Source/ControlRigPhysics/Private/RigPhysicsSolver.cpp
+++ b/Engine/Plugins/Experimental/PhysicsControl/Source/ControlRigPhysics/Private/RigPhysicsSolver.cpp
@@ -379,7 +379,7 @@ void FRigPhysicsSolver::InstantiatePhysicsJoints()
 			// are set to not collide for other reasons.
 			ConstraintSettings.bCollisionEnabled = !PhysicsJoint.bDisableCollision;
 			ConstraintSettings.bProjectionEnabled = 
-				PhysicsJoint.LinearProjectionAmount > 0 || PhysicsJoint.AngularProjectionAmount > 0;
+				PhysicsJoint.LinearProjectionAmount > 0.0f || PhysicsJoint.AngularProjectionAmount > 0.0f;
 			ConstraintSettings.AngularProjection = PhysicsJoint.AngularProjectionAmount;
 			ConstraintSettings.LinearProjection = PhysicsJoint.LinearProjectionAmount;
 			ConstraintSettings.ParentInvMassScale = PhysicsJoint.ParentInverseMassScale;
@@ -414,8 +414,8 @@ void FRigPhysicsSolver::InstantiatePhysicsJoints()
 
 					// Linear
 					Settings.LinearDriveForceMode = Chaos::EJointForceMode::Acceleration;
-					const bool bLinearPositionDriveEnabled = LinearSpring > 0;
-					const bool bLinearVelocityDriveEnabled = LinearDamping > 0;
+					const bool bLinearPositionDriveEnabled = LinearSpring > 0.0f;
+					const bool bLinearVelocityDriveEnabled = LinearDamping > 0.0f;
 					Settings.bLinearPositionDriveEnabled = { 
 						bLinearPositionDriveEnabled, bLinearPositionDriveEnabled, bLinearPositionDriveEnabled};
 					Settings.bLinearVelocityDriveEnabled = {
@@ -425,18 +425,18 @@ void FRigPhysicsSolver::InstantiatePhysicsJoints()
 					Settings.AngularDriveForceMode = Chaos::EJointForceMode::Acceleration;
 					// SLERP is generally better, but doesn't work if anything is locked
 					if (ConstraintSettings.AngularMotionTypes[0] == Chaos::EJointMotionType::Locked ||
-						ConstraintSettings.AngularMotionTypes[0] == Chaos::EJointMotionType::Locked ||
-						ConstraintSettings.AngularMotionTypes[0] == Chaos::EJointMotionType::Locked)
+						ConstraintSettings.AngularMotionTypes[1] == Chaos::EJointMotionType::Locked ||
+						ConstraintSettings.AngularMotionTypes[2] == Chaos::EJointMotionType::Locked)
 					{
-						Settings.bAngularSwingPositionDriveEnabled = AngularSpring > 0;
-						Settings.bAngularTwistPositionDriveEnabled = AngularSpring > 0;
-						Settings.bAngularSwingVelocityDriveEnabled = AngularDamping > 0;
-						Settings.bAngularTwistVelocityDriveEnabled = AngularDamping > 0;
+						Settings.bAngularSwingPositionDriveEnabled = AngularSpring > 0.0f;
+						Settings.bAngularTwistPositionDriveEnabled = AngularSpring > 0.0f;
+						Settings.bAngularSwingVelocityDriveEnabled = AngularDamping > 0.0f;
+						Settings.bAngularTwistVelocityDriveEnabled = AngularDamping > 0.0f;
 					}
 					else
 					{
-						Settings.bAngularSLerpPositionDriveEnabled = AngularSpring > 0;
-						Settings.bAngularSLerpVelocityDriveEnabled = AngularDamping > 0;
+						Settings.bAngularSLerpPositionDriveEnabled = AngularSpring > 0.0f;
+						Settings.bAngularSLerpVelocityDriveEnabled = AngularDamping > 0.0f;
 					}
 
 					Constraint->SetSettings(Settings);
@@ -664,10 +664,10 @@ void FRigPhysicsSolver::UpdatePrePhysics(
 					{
 						const Chaos::FPBDJointSettings& Settings = Constraint->GetSettings();
 
-						if (Settings.AngularDriveStiffness.SquaredLength() > 0 ||
-							Settings.AngularDriveDamping.SquaredLength() > 0 ||
-							Settings.LinearDriveStiffness.SquaredLength() > 0 ||
-							Settings.LinearDriveDamping.SquaredLength() > 0)
+						if (Settings.AngularDriveStiffness.SquaredLength() > 0.0f ||
+							Settings.AngularDriveDamping.SquaredLength() > 0.0f ||
+							Settings.LinearDriveStiffness.SquaredLength() > 0.0f ||
+							Settings.LinearDriveDamping.SquaredLength() > 0.0f)
 						{
 							// Multiplier on the velocity calculated from the current and previous target
 							float TargetVelocityMultiplier = 1.0f;
@@ -716,7 +716,7 @@ void FRigPhysicsSolver::UpdatePrePhysics(
 							Constraint->SetAngularDrivePositionTarget(DriveTargetTM.GetRotation());
 
 							if (Record.PreviousDriveTargetUpdateCounter + 1 == UpdateCounter && 
-								TargetVelocityMultiplier > 0)
+								TargetVelocityMultiplier > 0.0f)
 							{
 								if (DeltaTime > SMALL_NUMBER)
 								{
@@ -779,7 +779,7 @@ void FRigPhysicsSolver::UpdatePostPhysics()
 				{
 					FTransform TargetSimSpaceTM = Record.ActorHandle->GetWorldTransform();
 					double DistSq = TargetSimSpaceTM.GetTranslation().SquaredLength();
-					if (!TargetSimSpaceTM.IsValid() || (SolverSettings.DistanceThresholdForReset > 0 &&
+					if (!TargetSimSpaceTM.IsValid() || (SolverSettings.DistanceThresholdForReset > 0.0f &&
 						DistSq > FMath::Square(SolverSettings.DistanceThresholdForReset)))
 					{
 						bGotInvalidSimulationData = true;
