Add coloring and move rig physics nodes to a Rig Physics category in Control Rig. Also adds the ability to control the amount of blend back (alpha), and whether to step physics.

diff --git a/Engine/Plugins/Experimental/PhysicsControl/Source/ControlRigPhysics/Private/RigPhysicsExecution.cpp b/Engine/Plugins/Experimental/PhysicsControl/Source/ControlRigPhysics/Private/RigPhysicsExecution.cpp
--- a/Engine/Plugins/Experimental/PhysicsControl/Source/ControlRigPhysics/Private/RigPhysicsExecution.cpp
+++ b/Engine/Plugins/Experimental/PhysicsControl/Source/ControlRigPhysics/Private/RigPhysicsExecution.cpp
@@ -51,8 +51,7 @@ FRigUnit_StepPhysicsSolver_Execute()
 		{
 			if (FRigPhysicsSolver* PhysicsSolver = Cast<FRigPhysicsSolver>(ControlRig->GetPhysicsSolver(SolverID)))
 			{
-				double DeltaTime = DeltaTimeOverride > 0 ? DeltaTimeOverride : ExecuteContext.GetDeltaTime();
-				PhysicsSolver->StepSimulation(ExecuteContext, DeltaTime);
+				PhysicsSolver->StepSimulation(ExecuteContext, DeltaTimeOverride, Alpha);
 				if (VisualizationSettings.bEnableVisualization)
 				{
 					PhysicsSolver->Draw(
diff --git a/Engine/Plugins/Experimental/PhysicsControl/Source/ControlRigPhysics/Private/RigPhysicsSolver.cpp b/Engine/Plugins/Experimental/PhysicsControl/Source/ControlRigPhysics/Private/RigPhysicsSolver.cpp
--- a/Engine/Plugins/Experimental/PhysicsControl/Source/ControlRigPhysics/Private/RigPhysicsSolver.cpp
+++ b/Engine/Plugins/Experimental/PhysicsControl/Source/ControlRigPhysics/Private/RigPhysicsSolver.cpp
@@ -568,7 +568,7 @@ void FRigPhysicsSolver::SetBodyKinematicTarget(
 
 //======================================================================================================================
 void FRigPhysicsSolver::UpdatePrePhysics(
-	FRigVMExecuteContext& ExecuteContext, const double DeltaTime, const FSimulationSpaceData& SimulationSpaceData)
+	FRigVMExecuteContext& ExecuteContext, const float DeltaTime, const FSimulationSpaceData& SimulationSpaceData)
 {
 	URigHierarchy* Hierarchy = OwningControlRig->GetHierarchy();
 
@@ -753,8 +753,13 @@ void FRigPhysicsSolver::UpdatePrePhysics(
 //======================================================================================================================
 // Note that we read back into a target bone, which may have been specified explicitly, or will
 // otherwise default to the physics element parent.
-void FRigPhysicsSolver::UpdatePostPhysics()
+void FRigPhysicsSolver::UpdatePostPhysics(float Alpha)
 {
+	if (Alpha == 0.0f)
+	{
+		return;
+	}
+
 	bool bGotInvalidSimulationData = false;
 	URigHierarchy* Hierarchy = OwningControlRig->GetHierarchy();
 	for (TPair<FRigComponentKey, FRigBodyRecord>& BodyRecordPair : BodyRecords)
@@ -787,6 +792,20 @@ void FRigPhysicsSolver::UpdatePostPhysics()
 					else
 					{
 						FTransform TargetComponentSpaceTM = ConvertSimSpaceTransformToComponentSpace(TargetSimSpaceTM);
+						if (Alpha < 0.999f)
+						{
+							// Danny TODO Note that this uses Alpha to blend in component space.
+							// This can cause joint separation. We probably want an option to blend
+							// in local (joint) space, perhaps splitting the alpha into orientation
+							// and position.
+							FTransform CurrentTM = Hierarchy->GetGlobalTransform(TargetElementKey);
+							FQuat TargetQ = FQuat::Slerp(
+								CurrentTM.GetRotation(), TargetComponentSpaceTM.GetRotation(), Alpha);
+							FVector TargetT = FMath::Lerp(
+								CurrentTM.GetTranslation(), TargetComponentSpaceTM.GetTranslation(), Alpha);
+							TargetComponentSpaceTM.SetRotation(TargetQ);
+							TargetComponentSpaceTM.SetTranslation(TargetT);
+						}
 						// Danny TODO should we be "updating" the children here or not? 
 						Hierarchy->SetGlobalTransform(TargetElementKey, TargetComponentSpaceTM);
 					}
@@ -829,10 +848,21 @@ void FRigPhysicsSolver::UpdatePostPhysics()
 }
 
 //======================================================================================================================
-void FRigPhysicsSolver::StepSimulation(FRigVMExecuteContext& ExecuteContext, double DeltaTime)
+void FRigPhysicsSolver::StepSimulation(
+	FRigVMExecuteContext& ExecuteContext, const float DeltaTimeOverride, const float Alpha)
 {
+	float PhysicsDeltaTime = ExecuteContext.GetDeltaTime();
+	if (DeltaTimeOverride > 0.0f)
+	{
+		PhysicsDeltaTime = DeltaTimeOverride;
+	}
+	else if (DeltaTimeOverride < 0.0f)
+	{
+		PhysicsDeltaTime = 0.0f;
+	}
+
 	// We need to know about the simulation space etc before we can instantiate anything into the right place
-	FSimulationSpaceData SimulationSpaceData = UpdateSimulationSpaceStateAndCalculateData(DeltaTime);
+	FSimulationSpaceData SimulationSpaceData = UpdateSimulationSpaceStateAndCalculateData(PhysicsDeltaTime);
 
 	// We instantiate when we do the first simulation - this makes sure any changes applied by
 	// the user have been made. It also means there is no overhead if physics is never stepped.
@@ -872,13 +902,16 @@ void FRigPhysicsSolver::StepSimulation(FRigVMExecuteContext& ExecuteContext, dou
 		SimulationSpaceData.LinearAcceleration,
 		SimulationSpaceData.AngularAcceleration);
 
-	UpdatePrePhysics(ExecuteContext, DeltaTime, SimulationSpaceData);
+	UpdatePrePhysics(ExecuteContext, PhysicsDeltaTime, SimulationSpaceData);
 
-	Simulation->Simulate(
-		DeltaTime, SolverSettings.MaxDeltaTime, SolverSettings.MaxTimeSteps, 
-		SimulationSpaceData.Gravity, &ChaosJointSolverSettings);
+	if (PhysicsDeltaTime > 0.0f)
+	{
+		Simulation->Simulate(
+			PhysicsDeltaTime, SolverSettings.MaxDeltaTime, SolverSettings.MaxTimeSteps,
+			SimulationSpaceData.Gravity, &ChaosJointSolverSettings);
+	}
 
-	UpdatePostPhysics();
+	UpdatePostPhysics(Alpha);
 
 	++UpdateCounter;
 }
diff --git a/Engine/Plugins/Experimental/PhysicsControl/Source/ControlRigPhysics/Private/RigPhysicsSolver.h b/Engine/Plugins/Experimental/PhysicsControl/Source/ControlRigPhysics/Private/RigPhysicsSolver.h
--- a/Engine/Plugins/Experimental/PhysicsControl/Source/ControlRigPhysics/Private/RigPhysicsSolver.h
+++ b/Engine/Plugins/Experimental/PhysicsControl/Source/ControlRigPhysics/Private/RigPhysicsSolver.h
@@ -88,8 +88,11 @@ public:
 	// This will initialise/create the simulation and then create everything we need in it. 
 	void Instantiate();
 
-	// Integrates the simulation forwards by DeltaTime
-	void StepSimulation(FRigVMExecuteContext& ExecuteContext, double DeltaTime);
+	// Integrates the simulation forwards.
+	// If DeltaTimeOverride is +ve, then that value is used.
+	// If it is zero, then delta time is taken from the execute context
+	// If it is negative, then the simulation isn't stepped.
+	void StepSimulation(FRigVMExecuteContext& ExecuteContext, const float DeltaTimeOverride, const float Alpha);
 
 	// Draws shapes etc AND (potentially) enables the low-level Chaos debug draw
 	void Draw(
@@ -147,8 +150,8 @@ private:
 	void DestroyPhysicsSimulation();
 
 	void UpdatePrePhysics(
-		FRigVMExecuteContext& ExecuteContext, const double DeltaTime, const FSimulationSpaceData& SimulationSpaceData);
-	void UpdatePostPhysics();
+		FRigVMExecuteContext& ExecuteContext, const float DeltaTime, const FSimulationSpaceData& SimulationSpaceData);
+	void UpdatePostPhysics(const float Alpha);
 
 	void InitSimulationSpace(const FTransform& ComponentToWorld, const FTransform& BoneToComponent);
 
diff --git a/Engine/Plugins/Experimental/PhysicsControl/Source/ControlRigPhysics/Public/RigPhysicsControlExecution.h b/Engine/Plugins/Experimental/PhysicsControl/Source/ControlRigPhysics/Public/RigPhysicsControlExecution.h
--- a/Engine/Plugins/Experimental/PhysicsControl/Source/ControlRigPhysics/Public/RigPhysicsControlExecution.h
+++ b/Engine/Plugins/Experimental/PhysicsControl/Source/ControlRigPhysics/Public/RigPhysicsControlExecution.h
@@ -5,6 +5,7 @@
 #include "RigPhysicsComponent.h"
 #include "RigPhysicsSolver.h"
 #include "PhysicsControlData.h"
+#include "RigPhysicsExecution.h"
 
 #include "Units/Execution/RigUnit_DynamicHierarchy.h"
 
@@ -14,7 +15,7 @@
  * Adds a new physics control as a component on the owner element.
  * Note: This node only runs as part of the construction event.
  */
-USTRUCT(meta = (DisplayName = "Add Physics Control Component", Keywords = "Construction,Create,New,Joint,Control,Physics", Varying))
+USTRUCT(meta = (Category = "RigPhysics", NodeColor = "1.0 0.6 0.3", DisplayName = "Add Physics Control Component", Keywords = "Construction,Create,New,Joint,Control,Physics", Varying))
 struct CONTROLRIGPHYSICS_API FRigUnit_AddPhysicsControl : public FRigUnitMutable
 {
 	GENERATED_BODY()
diff --git a/Engine/Plugins/Experimental/PhysicsControl/Source/ControlRigPhysics/Public/RigPhysicsData.h b/Engine/Plugins/Experimental/PhysicsControl/Source/ControlRigPhysics/Public/RigPhysicsData.h
--- a/Engine/Plugins/Experimental/PhysicsControl/Source/ControlRigPhysics/Public/RigPhysicsData.h
+++ b/Engine/Plugins/Experimental/PhysicsControl/Source/ControlRigPhysics/Public/RigPhysicsData.h
@@ -26,10 +26,10 @@ struct FRigPhysicsVisualizationSettings
 	UPROPERTY(BlueprintReadWrite, EditAnywhere, Category = Visualization)
 	bool bEnableVisualization = true;
 
-	UPROPERTY(BlueprintReadWrite, EditAnywhere, Category = Visualization)
+	UPROPERTY(BlueprintReadWrite, EditAnywhere, Category = Visualization, meta = (ClampMin = "0.0"))
 	float LineThickness = 1.0f;
 
-	UPROPERTY(BlueprintReadWrite, EditAnywhere, Category = Visualization)
+	UPROPERTY(BlueprintReadWrite, EditAnywhere, Category = Visualization, meta = (ClampMin = "0.0"))
 	int32 ShapeDetail = 16;
 };
 
@@ -90,7 +90,7 @@ struct CONTROLRIGPHYSICS_API FRigPhysicsCollisionShape
 	/** Offset used when generating contact points. This allows you to smooth out
 	the Minkowski sum by radius R. Useful for making objects slide smoothly
 	on top of irregularities  */
-	UPROPERTY(Category = Shape, EditAnywhere)
+	UPROPERTY(Category = Shape, EditAnywhere, meta = (ClampMin = "0.0"))
 	float RestOffset = 0.0f;
 
 	UPROPERTY(Category = Shape, EditAnywhere)
@@ -144,7 +144,7 @@ struct CONTROLRIGPHYSICS_API FRigPhysicsCollisionSphere : public FRigPhysicsColl
 	UPROPERTY(BlueprintReadWrite, EditAnywhere, Category = Shapes)
 	FTransform TM;
 
-	UPROPERTY(BlueprintReadWrite, EditAnywhere, Category = Shapes)
+	UPROPERTY(BlueprintReadWrite, EditAnywhere, Category = Shapes, meta = (ClampMin = "0.0"))
 	float Radius = RigPhysicsDefaultShapeSize;
 };
 
@@ -167,11 +167,11 @@ struct CONTROLRIGPHYSICS_API FRigPhysicsCollisionCapsule : public FRigPhysicsCol
 	UPROPERTY(BlueprintReadWrite, EditAnywhere, Category = Shapes)
 	FTransform TM;
 
-	UPROPERTY(BlueprintReadWrite, EditAnywhere, Category = Shapes)
+	UPROPERTY(BlueprintReadWrite, EditAnywhere, Category = Shapes, meta = (ClampMin = "0.0"))
 	float Radius = RigPhysicsDefaultShapeSize;
 
 	// This is the length of the core part of the capsule. The total length will be Length + 2 * Radius
-	UPROPERTY(BlueprintReadWrite, EditAnywhere, Category = Shapes)
+	UPROPERTY(BlueprintReadWrite, EditAnywhere, Category = Shapes, meta = (ClampMin = "0.0"))
 	float Length = RigPhysicsDefaultShapeSize;
 };
 
@@ -213,7 +213,7 @@ struct FRigPhysicsSimulationSpaceSettings
 	GENERATED_BODY()
 
 	// Overall multiplier on the effects of simulation space movement on the simulation
-	UPROPERTY(EditAnywhere, BlueprintReadWrite, Category = Settings)
+	UPROPERTY(EditAnywhere, BlueprintReadWrite, Category = Settings, meta = (ClampMin = "0.0", ClampMax = "1.0"))
 	float SpaceMovementAmount = 1.0f;
 
 	// Multiplier on the components of vertical velocity of the simulation space that is passed to the simulation.
diff --git a/Engine/Plugins/Experimental/PhysicsControl/Source/ControlRigPhysics/Public/RigPhysicsExecution.h b/Engine/Plugins/Experimental/PhysicsControl/Source/ControlRigPhysics/Public/RigPhysicsExecution.h
--- a/Engine/Plugins/Experimental/PhysicsControl/Source/ControlRigPhysics/Public/RigPhysicsExecution.h
+++ b/Engine/Plugins/Experimental/PhysicsControl/Source/ControlRigPhysics/Public/RigPhysicsExecution.h
@@ -14,7 +14,7 @@
  * Adds a new physics solver
  * Note: This node only runs as part of the construction event.
  */
-USTRUCT(meta=(DisplayName="Spawn Physics Solver", Keywords="Construction,Create,New,Simulation,Physics", Varying))
+USTRUCT(meta=(Category = "RigPhysics", NodeColor = "1.0 0.6 0.3", DisplayName="Spawn Physics Solver", Keywords="Construction,Create,New,Simulation,Physics", Varying))
 struct CONTROLRIGPHYSICS_API FRigUnit_AddPhysicsSolver : public FRigUnitMutable
 {
 	GENERATED_BODY()
@@ -53,7 +53,7 @@ struct CONTROLRIGPHYSICS_API FRigUnit_AddPhysicsSolver : public FRigUnitMutable
  * requested. Explicit instantiation allows the timing to be controlled, as allocations etc may cause some
  * delays.
  */
-USTRUCT(meta = (DisplayName = "Instantiate physics", Keywords = "Physics"))
+USTRUCT(meta = (Category = "RigPhysics", NodeColor = "1.0 0.6 0.3", DisplayName = "Instantiate physics", Keywords = "Physics"))
 struct CONTROLRIGPHYSICS_API FRigUnit_InstantiatePhysics : public FRigUnitMutable
 {
 	GENERATED_BODY()
@@ -72,30 +72,28 @@ struct CONTROLRIGPHYSICS_API FRigUnit_InstantiatePhysics : public FRigUnitMutabl
 /**
  * Steps the specified physics solver
  */
-USTRUCT(meta = (DisplayName = "Step Physics Solver", Keywords = "Simulate,Physics"))
+USTRUCT(meta = (Category = "RigPhysics", NodeColor = "1.0 0.6 0.3", DisplayName = "Step Physics Solver", Keywords = "Simulate,Physics"))
 struct CONTROLRIGPHYSICS_API FRigUnit_StepPhysicsSolver : public FRigUnitMutable
 {
 	GENERATED_BODY()
 
-	FRigUnit_StepPhysicsSolver()
-	{
-		DeltaTimeOverride = -1.0f;
-	}
-
 	RIGVM_METHOD()
 	virtual void Execute() override;
 
-	/*
-	 * The solver to relate this new physics element to
-	 */
+	// The solver to relate this new physics element to
 	UPROPERTY(meta = (Input))
 	FRigPhysicsSolverID SolverID;
-
-	/*
-	 * Normally the execute context time will be used, but if this is positive then it will override it.
-	 */
+	
+	// If this is zero, then the execute context time will be used. If this is positive then it will
+	// override the delta time. A negative value will prevent the solver from stepping, but there will
+	// still be update costs associated with the node.
 	UPROPERTY(meta = (Input))
-	float DeltaTimeOverride;
+	float DeltaTimeOverride = 0.0f;
+
+	// How much of the simulation is combined with the input bone. This currently happens in
+	// component space. Note that the simulation will continue to run, even if alpha = 0
+	UPROPERTY(meta = (Input, ClampMin = "0.0", ClampMax = "1.0"))
+	float Alpha = 1.0f;
 
 	UPROPERTY(meta = (Input))
 	FRigPhysicsVisualizationSettings VisualizationSettings;
@@ -105,7 +103,7 @@ struct CONTROLRIGPHYSICS_API FRigUnit_StepPhysicsSolver : public FRigUnitMutable
  * Adds a new physics body as a component on the owner element.
  * Note: This node only runs as part of the construction event.
  */
-USTRUCT(meta = (DisplayName = "Add Physics Body Component", Keywords = "Construction,Create,New,Joint,Body,Skeleton,Physics", Varying))
+USTRUCT(meta = (Category = "RigPhysics", NodeColor = "1.0 0.6 0.3", DisplayName = "Add Physics Body Component", Keywords = "Construction,Create,New,Joint,Body,Skeleton,Physics", Varying))
 struct CONTROLRIGPHYSICS_API FRigUnit_AddPhysicsBody : public FRigUnitMutable
 {
 	GENERATED_BODY()
@@ -164,7 +162,7 @@ struct CONTROLRIGPHYSICS_API FRigUnit_AddPhysicsBody : public FRigUnitMutable
  * Sets what bone is used as a source transform for the physics body. This is used as a kinematic target, and when
  * initializing the simulation.
  */
-USTRUCT(meta = (DisplayName = "Set Source Bone", Keywords = "Physics", Varying))
+USTRUCT(meta = (Category = "RigPhysics", NodeColor = "1.0 0.6 0.3", DisplayName = "Set Source Bone", Keywords = "Physics", Varying))
 struct CONTROLRIGPHYSICS_API FRigUnit_HierarchySetPhysicsBodySourceBone : public FRigUnitMutable
 {
 	GENERATED_BODY()
@@ -182,7 +180,7 @@ struct CONTROLRIGPHYSICS_API FRigUnit_HierarchySetPhysicsBodySourceBone : public
 /**
  * Sets what bone is targeted by the simulation - i.e. where the simulation output is written to.
  */
-USTRUCT(meta = (DisplayName = "Set Target Bone", Keywords = "Physics", Varying))
+USTRUCT(meta = (Category = "RigPhysics", NodeColor = "1.0 0.6 0.3", DisplayName = "Set Target Bone", Keywords = "Physics", Varying))
 struct CONTROLRIGPHYSICS_API FRigUnit_HierarchySetPhysicsBodyTargetBone : public FRigUnitMutable
 {
 	GENERATED_BODY()
@@ -203,7 +201,7 @@ struct CONTROLRIGPHYSICS_API FRigUnit_HierarchySetPhysicsBodyTargetBone : public
  * Danny TODO - Note that the sparse data does not get displayed correctly, so this is largely 
  * unusable - the flags that enable/disable all end up getting reset if the user attempts to change them.
  */
-USTRUCT(meta = (DisplayName = "Set Physics Body Data", Keywords = "Physics", Varying))
+USTRUCT(meta = (Category = "RigPhysics", NodeColor = "1.0 0.6 0.3", DisplayName = "Set Physics Body Data", Keywords = "Physics", Varying))
 struct CONTROLRIGPHYSICS_API FRigUnit_HierarchySetPhysicsBodySparseData : public FRigUnitMutable
 {
 	GENERATED_BODY()
@@ -221,7 +219,7 @@ struct CONTROLRIGPHYSICS_API FRigUnit_HierarchySetPhysicsBodySparseData : public
 /**
  * Sets the kinematic target for a body - note that this won't actually make the body kinematic
  */
-USTRUCT(meta = (DisplayName = "Set Kinematic Target", Keywords = "Physics", Varying))
+USTRUCT(meta = (Category = "RigPhysics", NodeColor = "1.0 0.6 0.3", DisplayName = "Set Kinematic Target", Keywords = "Physics", Varying))
 struct CONTROLRIGPHYSICS_API FRigUnit_HierarchySetPhysicsBodyKinematicTarget : public FRigUnitMutable
 {
 	GENERATED_BODY()
@@ -243,7 +241,7 @@ struct CONTROLRIGPHYSICS_API FRigUnit_HierarchySetPhysicsBodyKinematicTarget : p
  * Sets the movement type for this body. 
  * Danny TODO explain the different between kinematic and static
  */
-USTRUCT(meta = (DisplayName = "Set Physics Body Movement Type", Keywords = "Physics", Varying))
+USTRUCT(meta = (Category = "RigPhysics", NodeColor = "1.0 0.6 0.3", DisplayName = "Set Physics Body Movement Type", Keywords = "Physics", Varying))
 struct CONTROLRIGPHYSICS_API FRigUnit_HierarchySetPhysicsBodyMovementType : public FRigUnitMutable
 {
 	GENERATED_BODY()
@@ -261,7 +259,7 @@ struct CONTROLRIGPHYSICS_API FRigUnit_HierarchySetPhysicsBodyMovementType : publ
 /**
  * Sets what collision is used for this body
  */
-USTRUCT(meta = (DisplayName = "Set Physics Body Collision Type", Keywords = "Physics", Varying))
+USTRUCT(meta = (Category = "RigPhysics", NodeColor = "1.0 0.6 0.3", DisplayName = "Set Physics Body Collision Type", Keywords = "Physics", Varying))
 struct CONTROLRIGPHYSICS_API FRigUnit_HierarchySetPhysicsBodyCollisionType : public FRigUnitMutable
 {
 	GENERATED_BODY()
@@ -279,7 +277,7 @@ struct CONTROLRIGPHYSICS_API FRigUnit_HierarchySetPhysicsBodyCollisionType : pub
 /**
  * Sets the multiplier on gravity that should be applied to the body.
  */
-USTRUCT(meta = (DisplayName = "Set Physics Body Gravity Multiplier", Keywords = "Physics", Varying))
+USTRUCT(meta = (Category = "RigPhysics", NodeColor = "1.0 0.6 0.3", DisplayName = "Set Physics Body Gravity Multiplier", Keywords = "Physics", Varying))
 struct CONTROLRIGPHYSICS_API FRigUnit_HierarchySetPhysicsBodyGravityMultiplier : public FRigUnitMutable
 {
 	GENERATED_BODY()
@@ -298,7 +296,7 @@ struct CONTROLRIGPHYSICS_API FRigUnit_HierarchySetPhysicsBodyGravityMultiplier :
  * Controls the amount that the simulation is blended back into the target bones. 
  * Danny TODO implement
  */
-USTRUCT(meta = (DisplayName = "Set Physics Body Physics Blend Weight", Keywords = "Physics", Varying))
+USTRUCT(meta = (Category = "RigPhysics", NodeColor = "1.0 0.6 0.3", DisplayName = "Set Physics Body Physics Blend Weight", Keywords = "Physics", Varying))
 struct CONTROLRIGPHYSICS_API FRigUnit_HierarchySetPhysicsBodyPhysicsBlendWeight : public FRigUnitMutable
 {
 	GENERATED_BODY()
@@ -317,7 +315,7 @@ struct CONTROLRIGPHYSICS_API FRigUnit_HierarchySetPhysicsBodyPhysicsBlendWeight
  * Danny TODO I think this is redundant given the ability to specify the space when kinematic 
  * targets are set. Investigate propagating that into the rest of Physics Control.
  */
-USTRUCT(meta = (DisplayName = "Set Physics Body Use Skeletal Animation", Keywords = "Physics", Varying))
+USTRUCT(meta = (Category = "RigPhysics", NodeColor = "1.0 0.6 0.3", DisplayName = "Set Physics Body Use Skeletal Animation", Keywords = "Physics", Varying))
 struct CONTROLRIGPHYSICS_API FRigUnit_HierarchySetPhysicsBodyUseSkeletalAnimation : public FRigUnitMutable
 {
 	GENERATED_BODY()
@@ -337,7 +335,7 @@ struct CONTROLRIGPHYSICS_API FRigUnit_HierarchySetPhysicsBodyUseSkeletalAnimatio
  * necessary when either kinematic targets are being used, or when the target bone differs from the source bone.
  * Danny TODO - check the implementation, and also can we auto-calculate this at runtime?
  */
-USTRUCT(meta = (DisplayName = "Set Physics Body Update Kinematic From Simulation", Keywords = "Physics", Varying))
+USTRUCT(meta = (Category = "RigPhysics", NodeColor = "1.0 0.6 0.3", DisplayName = "Set Physics Body Update Kinematic From Simulation", Keywords = "Physics", Varying))
 struct CONTROLRIGPHYSICS_API FRigUnit_HierarchySetPhysicsBodyUpdateKinematicFromSimulation : public FRigUnitMutable
 {
 	GENERATED_BODY()
