PR #12685: Fix PhysicsControlComponent::CreateControl parameter documentation

#rnx

diff --git a/Engine/Plugins/Experimental/PhysicsControl/Source/PhysicsControl/Public/PhysicsControlComponent.h b/Engine/Plugins/Experimental/PhysicsControl/Source/PhysicsControl/Public/PhysicsControlComponent.h
--- a/Engine/Plugins/Experimental/PhysicsControl/Source/PhysicsControl/Public/PhysicsControlComponent.h
+++ b/Engine/Plugins/Experimental/PhysicsControl/Source/PhysicsControl/Public/PhysicsControlComponent.h
@@ -112,7 +112,7 @@ public:
 	/**
 	 * Creates a new control for mesh components
 	 * 
-	 * @param The name of the control that will be created. Creation will fail if this name is already in use.
+	 * @param Name The name of the control that will be created. Creation will fail if this name is already in use.
 	 * @param ControlData   Describes the initial strength etc of the new control
 	 * @param ControlTarget Describes the initial target for the new control
 	 * @param Set Which set to include the control in (optional). Note that it automatically
