Support for manual as well as procedural creation of physics solvers and bodies

diff --git a/Engine/Plugins/Animation/ControlRig/Source/ControlRig/Private/ControlRig.cpp b/Engine/Plugins/Animation/ControlRig/Source/ControlRig/Private/ControlRig.cpp
--- a/Engine/Plugins/Animation/ControlRig/Source/ControlRig/Private/ControlRig.cpp
+++ b/Engine/Plugins/Animation/ControlRig/Source/ControlRig/Private/ControlRig.cpp
@@ -11,6 +11,7 @@
 #include "RigVMBlueprintGeneratedClass.h"
 #include "ControlRigObjectVersion.h"
 #include "Rigs/RigHierarchyController.h"
+#include "Rigs/RigPhysics.h"
 #include "Units/Execution/RigUnit_BeginExecution.h"
 #include "Units/Execution/RigUnit_PrepareForExecution.h"
 #include "Units/Execution/RigUnit_InverseExecution.h"
@@ -578,7 +579,7 @@ void UControlRig::InitializeFromCDO()
 			// Danny TODO this could copy the physics solvers from the CDO, but they are transient
 			// objects that only exist through construction, so I think we just need to start afresh
 			// each time. 
-			PhysicsSolvers.Empty();
+			PhysicsSimulations.Empty();
 
 			// copy hierarchy
 			{
@@ -1855,7 +1856,7 @@ void UControlRig::Serialize(FArchive& Ar)
 	}
 	else
 	{
-		PhysicsSolvers.Reset();
+		PhysicsSimulations.Reset();
 	}
 }
 
@@ -3644,51 +3645,27 @@ void UControlRig::OnHierarchyTransformUndoRedo(URigHierarchy* InHierarchy, const
 	}
 }
 
-FRigPhysicsSolverBase* UControlRig::GetPhysicsSolver(const FRigPhysicsSolverID& InID)
+FRigPhysicsSimulationBase* UControlRig::GetPhysicsSimulation(const FRigComponentKey& InComponentKey)
 {
-	TSharedPtr<FRigPhysicsSolverBase>* PhysicsSolver = PhysicsSolvers.Find(InID);
-	return PhysicsSolver ? PhysicsSolver->Get() : nullptr;
+	TSharedPtr<FRigPhysicsSimulationBase>* PhysicsSimulation = PhysicsSimulations.Find(InComponentKey);
+	return PhysicsSimulation ? PhysicsSimulation->Get() : nullptr;
 }
 
-const FRigPhysicsSolverBase* UControlRig::GetPhysicsSolver(const FRigPhysicsSolverID& InID) const
+const FRigPhysicsSimulationBase* UControlRig::GetPhysicsSimulation(const FRigComponentKey& InComponentKey) const
 {
-	const TSharedPtr<FRigPhysicsSolverBase>* PhysicsSolver = PhysicsSolvers.Find(InID);
-	return PhysicsSolver ? PhysicsSolver->Get() : nullptr;
+	const TSharedPtr<FRigPhysicsSimulationBase>* PhysicsSimulation = PhysicsSimulations.Find(InComponentKey);
+	return PhysicsSimulation ? PhysicsSimulation->Get() : nullptr;
 }
 
-const FRigPhysicsSolverBase* UControlRig::FindPhysicsSolverByName(FName InName) const
+bool UControlRig::RegisterPhysicsSimulation(TSharedPtr<FRigPhysicsSimulationBase> InPhysicsSimulation, const FRigComponentKey& InComponentKey)
 {
-	for (const TPair<FRigPhysicsSolverID, TSharedPtr<FRigPhysicsSolverBase>>& SolverPair : PhysicsSolvers)
+	if(!InPhysicsSimulation.IsValid())
 	{
-		if (SolverPair.Value.IsValid() && SolverPair.Value->Name == InName)
-		{
-			return SolverPair.Value.Get();
-		}
-	}
-	return nullptr;
-}
-
-FRigPhysicsSolverID UControlRig::RegisterPhysicsSolver(TSharedPtr<FRigPhysicsSolverBase> InPhysicsWorld, FName InName)
-{
-	if(InName.IsNone() || !InPhysicsWorld.IsValid())
-	{
-		return FRigPhysicsSolverID();
+		return false;
 	}
-
-	const FName Name = UtilityHelpers::CreateUniqueName(InName, [this](const FName& InName) -> bool
-	{
-		return FindPhysicsSolverByName(InName) == nullptr;
-	});
-
-	const FString CompletePath = FString::Printf(TEXT("%s|%s"), *GetPathName(), *Name.ToString());
-	FRigPhysicsSolverID SolverID(FGuid::NewDeterministicGuid(CompletePath));
-
-	PhysicsSolvers.FindOrAdd(SolverID, InPhysicsWorld);
-
-	InPhysicsWorld->Name = Name;
-	InPhysicsWorld->SolverID = SolverID;
-
-	return SolverID;
+	PhysicsSimulations.FindOrAdd(InComponentKey, InPhysicsSimulation);
+	InPhysicsSimulation->PhysicsSolverComponentKey = InComponentKey;
+	return true;
 }
 
 UControlRig::FPoseScope::FPoseScope(UControlRig* InControlRig, ERigElementType InFilter, const TArray<FRigElementKey>& InElements, const ERigTransformType::Type InTransformType)
diff --git a/Engine/Plugins/Animation/ControlRig/Source/ControlRig/Public/ControlRig.h b/Engine/Plugins/Animation/ControlRig/Source/ControlRig/Public/ControlRig.h
--- a/Engine/Plugins/Animation/ControlRig/Source/ControlRig/Public/ControlRig.h
+++ b/Engine/Plugins/Animation/ControlRig/Source/ControlRig/Public/ControlRig.h
@@ -21,7 +21,6 @@
 #include "Animation/MeshDeformerInstance.h"
 #include "Animation/MeshDeformerProducer.h"
 #include "Rigs/RigModuleDefines.h"
-#include "Rigs/RigPhysics.h"
 
 #if WITH_EDITOR
 #include "RigVMModel/RigVMPin.h"
@@ -46,7 +45,7 @@ class UControlRigReplay;
 struct FReferenceSkeleton;
 struct FRigUnit;
 struct FRigControl;
-struct FRigPhysicsSolverBase;
+struct FRigPhysicsSimulationBase;
 
 CONTROLRIG_API DECLARE_LOG_CATEGORY_EXTERN(LogControlRig, Log, All);
 
@@ -723,37 +722,32 @@ protected:
 private:
 	float DebugBoneRadiusMultiplier;
 
-	// Physics Solvers/Worlds
-	TMap<FRigPhysicsSolverID, TSharedPtr<FRigPhysicsSolverBase>> PhysicsSolvers;
+	// Physics Simulations (i.e. the simulations instantiated based on the Physics Solver components)
+	TMap<FRigComponentKey, TSharedPtr<FRigPhysicsSimulationBase>> PhysicsSimulations;
 
 public:
 
 	/**
-	 * Returns the physics world given a solver ID
-	 * @param InID The id identifying the physics solver
-	 * @return The physics solver (may be null)
+	 * Returns the physics simulation given the key of the component used to create it
+	 * @param InComponentKey The id identifying the physics simulation
+	 * @return The physics simulation (may be null)
 	 */
-	FRigPhysicsSolverBase* GetPhysicsSolver(const FRigPhysicsSolverID& InID);
+	FRigPhysicsSimulationBase* GetPhysicsSimulation(const FRigComponentKey& InComponentKey);
 
 	/**
-	 * Returns the physics world given a solver ID
-	 * @param InID The id identifying the physics solver
-	 * @return The physics solver (may be null)
+	 * Returns the physics simulation given the key of the component used to create it
+	 * @param InComponentKey The id identifying the physics simulation
+	 * @return The physics simulation (may be null)
 	 */
-	const FRigPhysicsSolverBase* GetPhysicsSolver(const FRigPhysicsSolverID& InID) const;
+	const FRigPhysicsSimulationBase* GetPhysicsSimulation(const FRigComponentKey& InComponentKey) const;
 
 	/**
-	 * Finds a physics solver given its name
-	 * @param InName The name identifying the physics solver in the scope of this control rig
-	 * @return The physics solver
+	 * Registers and stores (taking ownership) of a physics simulation. The simulation will be 
+	 * responsible for the low-level storage and simulation of objects. Returns false if the world 
+	 * can't be registered - e.g. if there is already a simulation registered with the key.
 	 */
-	const FRigPhysicsSolverBase* FindPhysicsSolverByName(FName InName) const; 
-
-	/**
-	 * Registers and stores (taking ownership) of a physics world, returning an ID for it. The world will be 
-	 * responsible for the low-level storage and simulation of objects.
-	 */
-	FRigPhysicsSolverID RegisterPhysicsSolver(TSharedPtr<FRigPhysicsSolverBase> PhysicsWorld, FName Name);
+	bool RegisterPhysicsSimulation(
+		TSharedPtr<FRigPhysicsSimulationBase> PhysicsSimulation, const FRigComponentKey& InComponentKey);
 
 #if WITH_EDITOR	
 
diff --git a/Engine/Plugins/Animation/ControlRig/Source/ControlRig/Public/Rigs/RigPhysics.h b/Engine/Plugins/Animation/ControlRig/Source/ControlRig/Public/Rigs/RigPhysics.h
--- a/Engine/Plugins/Animation/ControlRig/Source/ControlRig/Public/Rigs/RigPhysics.h
+++ b/Engine/Plugins/Animation/ControlRig/Source/ControlRig/Public/Rigs/RigPhysics.h
@@ -6,8 +6,37 @@
 #include "RigHierarchyComponents.h"
 #include "RigPhysics.generated.h"
 
-// This simply exposes FGuid to blueprint, and thus custom UI at some point.
-USTRUCT(BlueprintType)
+// Base class for physics simulations registered with the ControlRig. The concrete class will be
+// instantiated by a plugin, so there is casting support to ensure that the correct concrete class
+// is being used.
+USTRUCT()
+struct CONTROLRIG_API FRigPhysicsSimulationBase
+{
+	GENERATED_BODY()
+
+	FRigPhysicsSimulationBase(const UScriptStruct* InType = nullptr) : Type(InType) {}
+
+	template<typename T> friend
+		T* Cast(FRigPhysicsSimulationBase* BasePtr)
+	{
+		if (BasePtr && BasePtr->Type == T::StaticStruct())
+		{
+			return static_cast<T*>(BasePtr);
+		}
+		return nullptr;
+	}
+
+	virtual ~FRigPhysicsSimulationBase() {}
+
+	// Used to support casting from the base class to the concrete class
+	const UScriptStruct* Type;
+
+	// The component key will be set when the concrete class is registered with the control rig
+	FRigComponentKey PhysicsSolverComponentKey;
+};
+
+// Note - FRigPhysicsSolverID is no long used - solvers are referenced using their component key
+USTRUCT(BlueprintType, meta = (Deprecated = "5.6"))
 struct CONTROLRIG_API FRigPhysicsSolverID
 {
 public:
@@ -52,19 +81,7 @@ public:
 	}
 };
 
-// Danny TODO This is just inheriting from the base component to acquire the casting functions - is there a better way?
-USTRUCT(BlueprintType)
-struct CONTROLRIG_API FRigPhysicsSolverBase : public FRigBaseComponent
-{
-	GENERATED_BODY()
-	DECLARE_RIG_COMPONENT_METHODS(FRigPhysicsSolverBase)
-
-	virtual ~FRigPhysicsSolverBase() {}
-
-	FName               Name;
-	FRigPhysicsSolverID SolverID;
-};
-
+// Note - FRigPhysicsSolverDescription is no long used - solvers are referenced using their component key
 USTRUCT()
 struct CONTROLRIG_API FRigPhysicsSolverDescription
 {
diff --git a/Engine/Plugins/Animation/ControlRig/Source/ControlRigDeveloper/Private/ControlRigSchema.cpp b/Engine/Plugins/Animation/ControlRig/Source/ControlRigDeveloper/Private/ControlRigSchema.cpp
--- a/Engine/Plugins/Animation/ControlRig/Source/ControlRigDeveloper/Private/ControlRigSchema.cpp
+++ b/Engine/Plugins/Animation/ControlRig/Source/ControlRigDeveloper/Private/ControlRigSchema.cpp
@@ -4,6 +4,7 @@
 #include "ControlRigBlueprint.h"
 #include "RigVMModel/RigVMController.h"
 #include "Rigs/RigHierarchyPose.h"
+#include "Rigs/RigPhysics.h"
 #include "Curves/CurveFloat.h"
 #include "Units/RigUnitContext.h"
 #include "Units/Modules/RigUnit_ConnectorExecution.h"
diff --git a/Engine/Plugins/Experimental/PhysicsControl/Source/ControlRigPhysics/Private/ControlRigPhysicsModule.cpp b/Engine/Plugins/Experimental/PhysicsControl/Source/ControlRigPhysics/Private/ControlRigPhysicsModule.cpp
--- a/Engine/Plugins/Experimental/PhysicsControl/Source/ControlRigPhysics/Private/ControlRigPhysicsModule.cpp
+++ b/Engine/Plugins/Experimental/PhysicsControl/Source/ControlRigPhysics/Private/ControlRigPhysicsModule.cpp
@@ -4,10 +4,14 @@
 
 #define LOCTEXT_NAMESPACE "FControlRigPhysicsModule"
 
+DEFINE_LOG_CATEGORY(LogControlRigPhysics);
+
+//======================================================================================================================
 void FControlRigPhysicsModule::StartupModule()
 {
 }
 
+//======================================================================================================================
 void FControlRigPhysicsModule::ShutdownModule()
 {
 }
diff --git a/Engine/Plugins/Experimental/PhysicsControl/Source/ControlRigPhysics/Private/RigPhysicsComponent.cpp b/Engine/Plugins/Experimental/PhysicsControl/Source/ControlRigPhysics/Private/RigPhysicsComponent.cpp
--- a/Engine/Plugins/Experimental/PhysicsControl/Source/ControlRigPhysics/Private/RigPhysicsComponent.cpp
+++ b/Engine/Plugins/Experimental/PhysicsControl/Source/ControlRigPhysics/Private/RigPhysicsComponent.cpp
@@ -2,20 +2,17 @@
 
 #include "RigPhysicsComponent.h"
 
-////////////////////////////////////////////////////////////////////////////////
-// FRigPhysicsComponent
-////////////////////////////////////////////////////////////////////////////////
-
+//======================================================================================================================
 void FRigPhysicsComponent::Save(FArchive& Ar)
 {
 	FRigBaseComponent::Save(Ar);
-
 	Ar << BodySolverSettings;
 	Ar << Dynamics;
 	Ar << Collision;
 	Ar << Joint;
 }
 
+//======================================================================================================================
 void FRigPhysicsComponent::Load(FArchive& Ar)
 {
 	FRigBaseComponent::Load(Ar);
diff --git a/Engine/Plugins/Experimental/PhysicsControl/Source/ControlRigPhysics/Private/RigPhysicsControlComponent.cpp b/Engine/Plugins/Experimental/PhysicsControl/Source/ControlRigPhysics/Private/RigPhysicsControlComponent.cpp
--- a/Engine/Plugins/Experimental/PhysicsControl/Source/ControlRigPhysics/Private/RigPhysicsControlComponent.cpp
+++ b/Engine/Plugins/Experimental/PhysicsControl/Source/ControlRigPhysics/Private/RigPhysicsControlComponent.cpp
@@ -2,24 +2,20 @@
 
 #include "RigPhysicsControlComponent.h"
 
-////////////////////////////////////////////////////////////////////////////////
-// FRigPhysicsControlComponent
-////////////////////////////////////////////////////////////////////////////////
-
+//======================================================================================================================
 void FRigPhysicsControlComponent::Save(FArchive& Ar)
 {
 	FRigBaseComponent::Save(Ar);
-
-	Ar << ParentBody;
+	Ar << ParentBodyComponentKey;
 	Ar << ControlData;
 	Ar << ControlTarget;
 }
 
+//======================================================================================================================
 void FRigPhysicsControlComponent::Load(FArchive& Ar)
 {
 	FRigBaseComponent::Load(Ar);
-
-	Ar << ParentBody;
+	Ar << ParentBodyComponentKey;
 	Ar << ControlData;
 	Ar << ControlTarget;
 }
diff --git a/Engine/Plugins/Experimental/PhysicsControl/Source/ControlRigPhysics/Private/RigPhysicsControlExecution.cpp b/Engine/Plugins/Experimental/PhysicsControl/Source/ControlRigPhysics/Private/RigPhysicsControlExecution.cpp
--- a/Engine/Plugins/Experimental/PhysicsControl/Source/ControlRigPhysics/Private/RigPhysicsControlExecution.cpp
+++ b/Engine/Plugins/Experimental/PhysicsControl/Source/ControlRigPhysics/Private/RigPhysicsControlExecution.cpp
@@ -1,7 +1,7 @@
 // Copyright Epic Games, Inc. All Rights Reserved.
 
 #include "RigPhysicsControlExecution.h"
-#include "RigPhysicsSolver.h"
+#include "RigPhysicsSimulation.h"
 #include "RigPhysicsControlComponent.h"
 
 #include "Rigs/RigHierarchyController.h"
@@ -15,42 +15,25 @@
 //======================================================================================================================
 FRigUnit_AddPhysicsControl_Execute()
 {
-	Result.Reset();
+	ControlComponentKey.Reset();
 
 	if (ExecuteContext.GetEventName() == FRigUnit_PrepareForExecution::EventName)
 	{
 		if (URigHierarchy* Hierarchy = ExecuteContext.Hierarchy)
 		{
-			if (UControlRig* ControlRig = Cast<UControlRig>(Hierarchy->GetOuter()))
+			if (URigHierarchyController* Controller = Hierarchy->GetController())
 			{
-				if (ControlRig->GetPhysicsSolver(SolverID) == nullptr)
+				FRigHierarchyControllerInstructionBracket InstructionBracket(
+					Controller, ExecuteContext.GetInstructionIndex());
+				ControlComponentKey = Controller->AddComponent(FRigPhysicsControlComponent::StaticStruct(), Name, Owner);
+				if (ControlComponentKey.IsValid())
 				{
-					UE_CONTROLRIG_RIGUNIT_REPORT_ERROR(
-						TEXT("Physics Solver with guid '%s' cannot be found."), *SolverID.ToString());
-					return;
-				}
-
-				if (URigHierarchyController* Controller = Hierarchy->GetController())
-				{
-					FRigHierarchyControllerInstructionBracket InstructionBracket(
-						Controller, ExecuteContext.GetInstructionIndex());
-					Result = Controller->AddComponent(FRigPhysicsControlComponent::StaticStruct(), Name, Owner);
-					if (Result.IsValid())
+					if (FRigPhysicsControlComponent* Component = Cast<FRigPhysicsControlComponent>(
+						Hierarchy->FindComponent(ControlComponentKey)))
 					{
-						if (FRigPhysicsControlComponent* Component = Cast<FRigPhysicsControlComponent>(
-							Hierarchy->FindComponent(Result)))
-						{
-							Component->ParentBody = ParentBody;
-							Component->ControlData = ControlData;
-							Component->ControlTarget = ControlTarget;
-
-							if (FRigPhysicsSolver* PhysicsSolver = Cast<FRigPhysicsSolver>(
-								ControlRig->GetPhysicsSolver(SolverID)))
-							{
-								// Danny TODO
-								//PhysicsSolver->AddControl(Result, Set);
-							}
-						}
+						Component->ParentBodyComponentKey = ParentBodyComponentKey;
+						Component->ControlData = ControlData;
+						Component->ControlTarget = ControlTarget;
 					}
 				}
 			}
@@ -60,8 +43,4 @@ FRigUnit_AddPhysicsControl_Execute()
 	{
 		UE_CONTROLRIG_RIGUNIT_REPORT_ERROR(TEXT("AddPhysicsControl can only be used during Setup"));
 	}
-
-
 }
-
-
diff --git a/Engine/Plugins/Experimental/PhysicsControl/Source/ControlRigPhysics/Private/RigPhysicsExecution.cpp b/Engine/Plugins/Experimental/PhysicsControl/Source/ControlRigPhysics/Private/RigPhysicsExecution.cpp
--- a/Engine/Plugins/Experimental/PhysicsControl/Source/ControlRigPhysics/Private/RigPhysicsExecution.cpp
+++ b/Engine/Plugins/Experimental/PhysicsControl/Source/ControlRigPhysics/Private/RigPhysicsExecution.cpp
@@ -1,7 +1,8 @@
 // Copyright Epic Games, Inc. All Rights Reserved.
 
 #include "RigPhysicsExecution.h"
-#include "RigPhysicsSolver.h"
+#include "RigPhysicsSimulation.h"
+#include "RigPhysicsSolverComponent.h"
 #include "RigPhysicsControlComponent.h"
 
 #include "Rigs/RigHierarchyController.h"
@@ -15,28 +16,33 @@
 //======================================================================================================================
 FRigUnit_AddPhysicsSolver_Execute()
 {
-	FString ErrorMessage;
-	if(!FRigUnit_DynamicHierarchyBase::IsValidToRunInContext(ExecuteContext, true, &ErrorMessage))
+	PhysicsSolverComponentKey.Reset();
+
+	if (ExecuteContext.GetEventName() == FRigUnit_PrepareForExecution::EventName)
 	{
-		if(!ErrorMessage.IsEmpty())
+		if (URigHierarchy* Hierarchy = ExecuteContext.Hierarchy)
 		{
-			UE_CONTROLRIG_RIGUNIT_REPORT_ERROR(TEXT("%s"), *ErrorMessage);
+			if (URigHierarchyController* Controller = Hierarchy->GetController())
+			{
+				FRigHierarchyControllerInstructionBracket InstructionBracket(
+					Controller, ExecuteContext.GetInstructionIndex());
+				PhysicsSolverComponentKey = Controller->AddComponent(
+					FRigPhysicsSolverComponent::StaticStruct(), TEXT("PhysicsSolver"), Owner);
+				if (PhysicsSolverComponentKey.IsValid())
+				{
+					if (FRigPhysicsSolverComponent* Component = Cast<FRigPhysicsSolverComponent>(
+						Hierarchy->FindComponent(PhysicsSolverComponentKey)))
+					{
+						Component->SolverSettings = SolverSettings;
+						Component->SimulationSpaceSettings = SimulationSpaceSettings;
+					}
+				}
+			}
 		}
-		return;
 	}
-
-	SolverID = FRigPhysicsSolverID();
-
-	if (!ExecuteContext.Hierarchy)
-	{
-		return;
-	}
-
-	if(UControlRig* ControlRig = Cast<UControlRig>(ExecuteContext.Hierarchy->GetOuter()))
+	else
 	{
-		TSharedPtr<FRigPhysicsSolver> PhysicsSolver = MakeShared<FRigPhysicsSolver>(
-			ControlRig, SolverSettings, SimulationSpaceSettings);
-		SolverID = ControlRig->RegisterPhysicsSolver(PhysicsSolver, Name);
+		UE_CONTROLRIG_RIGUNIT_REPORT_ERROR(TEXT("AddPhysicsSolver can only be used during Setup"));
 	}
 }
 
@@ -45,17 +51,37 @@ FRigUnit_StepPhysicsSolver_Execute()
 {
 	DECLARE_SCOPE_HIERARCHICAL_COUNTER_RIGUNIT();
 
-	if (ExecuteContext.Hierarchy)
+	if (URigHierarchy* Hierarchy = ExecuteContext.Hierarchy)
 	{
-		if (UControlRig* ControlRig = Cast<UControlRig>(ExecuteContext.Hierarchy->GetOuter()))
+		if (UControlRig* ControlRig = Cast<UControlRig>(Hierarchy->GetOuter()))
 		{
-			if (FRigPhysicsSolver* PhysicsSolver = Cast<FRigPhysicsSolver>(ControlRig->GetPhysicsSolver(SolverID)))
+			if (FRigPhysicsSolverComponent* SolverComponent = Cast<FRigPhysicsSolverComponent>(
+				Hierarchy->FindComponent(PhysicsSolverComponentKey)))
 			{
-				PhysicsSolver->StepSimulation(ExecuteContext, DeltaTimeOverride, Alpha);
-				if (VisualizationSettings.bEnableVisualization)
+				FRigPhysicsSimulation* PhysicsSimulation = Cast<FRigPhysicsSimulation>(
+					ControlRig->GetPhysicsSimulation(PhysicsSolverComponentKey));
+
+				// Make the solver and instantiate if necessary
+				if (!PhysicsSimulation)
 				{
-					PhysicsSolver->Draw(
-						ExecuteContext.GetDrawInterface(), VisualizationSettings, ExecuteContext.GetWorld());
+					TSharedPtr<FRigPhysicsSimulation> NewPhysicsSimulation = MakeShared<FRigPhysicsSimulation>(ControlRig);
+
+					if (ControlRig->RegisterPhysicsSimulation(NewPhysicsSimulation, PhysicsSolverComponentKey))
+					{
+						NewPhysicsSimulation->Instantiate(SolverComponent);
+					}
+					PhysicsSimulation = NewPhysicsSimulation.Get();
+				}
+
+				if (PhysicsSimulation)
+				{
+					PhysicsSimulation->StepSimulation(ExecuteContext, SolverComponent, DeltaTimeOverride, Alpha);
+					if (VisualizationSettings.bEnableVisualization)
+					{
+						PhysicsSimulation->Draw(
+							ExecuteContext.GetDrawInterface(), SolverComponent->SolverSettings, 
+							VisualizationSettings, ExecuteContext.GetWorld());
+					}
 				}
 			}
 		}
@@ -67,13 +93,19 @@ FRigUnit_InstantiatePhysics_Execute()
 {
 	DECLARE_SCOPE_HIERARCHICAL_COUNTER_RIGUNIT();
 
-	if (ExecuteContext.Hierarchy)
+	if (URigHierarchy* Hierarchy = ExecuteContext.Hierarchy)
 	{
-		if (UControlRig* ControlRig = Cast<UControlRig>(ExecuteContext.Hierarchy->GetOuter()))
+		if (UControlRig* ControlRig = Cast<UControlRig>(Hierarchy->GetOuter()))
 		{
-			if (FRigPhysicsSolver* PhysicsSolver = Cast<FRigPhysicsSolver>(ControlRig->GetPhysicsSolver(SolverID)))
+			if (FRigPhysicsSolverComponent* SolverComponent = Cast<FRigPhysicsSolverComponent>(
+				Hierarchy->FindComponent(PhysicsSolverComponentKey)))
 			{
-				PhysicsSolver->Instantiate();
+				TSharedPtr<FRigPhysicsSimulation> PhysicsSimulation = MakeShared<FRigPhysicsSimulation>(ControlRig);
+
+				if (ControlRig->RegisterPhysicsSimulation(PhysicsSimulation, PhysicsSolverComponentKey))
+				{
+					PhysicsSimulation->Instantiate(SolverComponent);
+				}
 			}
 		}
 	}
@@ -82,49 +114,29 @@ FRigUnit_InstantiatePhysics_Execute()
 //======================================================================================================================
 FRigUnit_AddPhysicsBody_Execute()
 {
-	Result.Reset();
-
-	if (!SolverID.IsValid())
-	{
-		UE_CONTROLRIG_RIGUNIT_REPORT_ERROR(TEXT("Physics Solver Guid is not valid"));
-		return;
-	}
+	PhysicsBodyComponentKey.Reset();
 
 	if (ExecuteContext.GetEventName() == FRigUnit_PrepareForExecution::EventName)
 	{
 		if (URigHierarchy* Hierarchy = ExecuteContext.Hierarchy)
 		{
-			if (UControlRig* ControlRig = Cast<UControlRig>(Hierarchy->GetOuter()))
+			if (URigHierarchyController* Controller = Hierarchy->GetController())
 			{
-				if (ControlRig->GetPhysicsSolver(SolverID) == nullptr)
-				{
-					UE_CONTROLRIG_RIGUNIT_REPORT_ERROR(
-						TEXT("Physics Solver with guid '%s' cannot be found."), *SolverID.ToString());
-					return;
-				}
-
-				if (URigHierarchyController* Controller = Hierarchy->GetController())
+				FRigHierarchyControllerInstructionBracket InstructionBracket(
+					Controller, ExecuteContext.GetInstructionIndex());
+				PhysicsBodyComponentKey = Controller->AddComponent(
+					FRigPhysicsComponent::StaticStruct(), TEXT("PhysicsBody"), Owner);
+				if (PhysicsBodyComponentKey.IsValid())
 				{
-					FRigHierarchyControllerInstructionBracket InstructionBracket(
-						Controller, ExecuteContext.GetInstructionIndex());
-					Result = Controller->AddComponent(FRigPhysicsComponent::StaticStruct(), TEXT("PhysicsBody"), Owner);
-					if (Result.IsValid())
+					if (FRigPhysicsComponent* Component = Cast<FRigPhysicsComponent>(
+						Hierarchy->FindComponent(PhysicsBodyComponentKey)))
 					{
-						if (FRigPhysicsComponent* Component = Cast<FRigPhysicsComponent>(
-							Hierarchy->FindComponent(Result)))
-						{
-							Component->BodySolverSettings.SolverID = SolverID;
-							Component->Dynamics = Dynamics;
-							Component->Collision = Collision;
-							Component->Joint = Joint;
-							Component->BodyData = BodyData;
-
-							if (FRigPhysicsSolver* PhysicsSolver = Cast<FRigPhysicsSolver>(
-								ControlRig->GetPhysicsSolver(SolverID)))
-							{
-								PhysicsSolver->AddBody(Result);
-							}
-						}
+						Component->BodySolverSettings.PhysicsSolverComponentKey = PhysicsSolverComponentKey;
+						Component->Dynamics = Dynamics;
+						Component->Collision = Collision;
+						Component->Joint = Joint;
+						Component->BodyData = BodyData;
+						Component->CurrentBodyData = BodyData;
 					}
 				}
 			}
@@ -144,7 +156,7 @@ FRigUnit_HierarchySetPhysicsBodySourceBone_Execute()
 	if (ExecuteContext.Hierarchy)
 	{
 		if (FRigPhysicsComponent* PhysicsComponent = Cast<FRigPhysicsComponent>(
-			ExecuteContext.Hierarchy->FindComponent(PhysicsBody)))
+			ExecuteContext.Hierarchy->FindComponent(PhysicsBodyComponentKey)))
 		{
 			PhysicsComponent->BodySolverSettings.SourceBone = SourceBone;
 		}
@@ -159,7 +171,7 @@ FRigUnit_HierarchySetPhysicsBodyTargetBone_Execute()
 	if (ExecuteContext.Hierarchy)
 	{
 		if (FRigPhysicsComponent* PhysicsComponent = Cast<FRigPhysicsComponent>(
-			ExecuteContext.Hierarchy->FindComponent(PhysicsBody)))
+			ExecuteContext.Hierarchy->FindComponent(PhysicsBodyComponentKey)))
 		{
 			PhysicsComponent->BodySolverSettings.TargetBone = TargetBone;
 		}
@@ -174,16 +186,10 @@ FRigUnit_HierarchySetPhysicsBodyKinematicTarget_Execute()
 	if (ExecuteContext.Hierarchy)
 	{
 		if (FRigPhysicsComponent* PhysicsComponent = Cast<FRigPhysicsComponent>(
-			ExecuteContext.Hierarchy->FindComponent(PhysicsBody)))
+			ExecuteContext.Hierarchy->FindComponent(PhysicsBodyComponentKey)))
 		{
-			if (UControlRig* ControlRig = Cast<UControlRig>(ExecuteContext.Hierarchy->GetOuter()))
-			{
-				if (FRigPhysicsSolver* PhysicsSolver = Cast<FRigPhysicsSolver>(
-					ControlRig->GetPhysicsSolver((PhysicsComponent->BodySolverSettings.SolverID))))
-				{
-					PhysicsSolver->SetBodyKinematicTarget(PhysicsBody, KinematicTargetSpace, KinematicTarget);
-				}
-			}
+			PhysicsComponent->CurrentKinematicTarget = KinematicTarget;
+			PhysicsComponent->CurrentKinematicTargetSpace = KinematicTargetSpace;
 		}
 	}
 }
@@ -196,15 +202,15 @@ FRigUnit_HierarchySetPhysicsBodySparseData_Execute()
 	if (ExecuteContext.Hierarchy)
 	{
 		if (FRigPhysicsComponent* PhysicsComponent = Cast<FRigPhysicsComponent>(
-			ExecuteContext.Hierarchy->FindComponent(PhysicsBody)))
+			ExecuteContext.Hierarchy->FindComponent(PhysicsBodyComponentKey)))
 		{
-			if (UControlRig* ControlRig = Cast<UControlRig>(ExecuteContext.Hierarchy->GetOuter()))
+			if (bInitial)
 			{
-				if (FRigPhysicsSolver* PhysicsSolver = Cast<FRigPhysicsSolver>(
-					ControlRig->GetPhysicsSolver((PhysicsComponent->BodySolverSettings.SolverID))))
-				{
-					PhysicsSolver->SetBodySparseData(PhysicsBody, Data);
-				}
+				PhysicsComponent->BodyData.UpdateFromSparseData(Data);
+			}
+			else
+			{
+				PhysicsComponent->CurrentBodyData.UpdateFromSparseData(Data);
 			}
 		}
 	}
@@ -217,15 +223,15 @@ FRigUnit_HierarchySetPhysicsBodyMovementType_Execute()
 	if (ExecuteContext.Hierarchy)
 	{
 		if (FRigPhysicsComponent* PhysicsComponent = Cast<FRigPhysicsComponent>(
-			ExecuteContext.Hierarchy->FindComponent(PhysicsBody)))
+			ExecuteContext.Hierarchy->FindComponent(PhysicsBodyComponentKey)))
 		{
-			if (UControlRig* ControlRig = Cast<UControlRig>(ExecuteContext.Hierarchy->GetOuter()))
+			if (bInitial)
 			{
-				if (FRigPhysicsSolver* PhysicsSolver = Cast<FRigPhysicsSolver>(
-					ControlRig->GetPhysicsSolver((PhysicsComponent->BodySolverSettings.SolverID))))
-				{
-					PhysicsSolver->SetBodyMovementType(PhysicsBody, MovementType);
-				}
+				PhysicsComponent->BodyData.MovementType = MovementType;
+			}
+			else
+			{
+				PhysicsComponent->CurrentBodyData.MovementType = MovementType;
 			}
 		}
 	}
@@ -238,15 +244,15 @@ FRigUnit_HierarchySetPhysicsBodyCollisionType_Execute()
 	if (ExecuteContext.Hierarchy)
 	{
 		if (FRigPhysicsComponent* PhysicsComponent = Cast<FRigPhysicsComponent>(
-			ExecuteContext.Hierarchy->FindComponent(PhysicsBody)))
+			ExecuteContext.Hierarchy->FindComponent(PhysicsBodyComponentKey)))
 		{
-			if (UControlRig* ControlRig = Cast<UControlRig>(ExecuteContext.Hierarchy->GetOuter()))
+			if (bInitial)
 			{
-				if (FRigPhysicsSolver* PhysicsSolver = Cast<FRigPhysicsSolver>(
-					ControlRig->GetPhysicsSolver((PhysicsComponent->BodySolverSettings.SolverID))))
-				{
-					PhysicsSolver->SetBodyCollisionType(PhysicsBody, CollisionType);
-				}
+				PhysicsComponent->BodyData.CollisionType = CollisionType;
+			}
+			else
+			{
+				PhysicsComponent->CurrentBodyData.CollisionType = CollisionType;
 			}
 		}
 	}
@@ -259,15 +265,15 @@ FRigUnit_HierarchySetPhysicsBodyGravityMultiplier_Execute()
 	if (ExecuteContext.Hierarchy)
 	{
 		if (FRigPhysicsComponent* PhysicsComponent = Cast<FRigPhysicsComponent>(
-			ExecuteContext.Hierarchy->FindComponent(PhysicsBody)))
+			ExecuteContext.Hierarchy->FindComponent(PhysicsBodyComponentKey)))
 		{
-			if (UControlRig* ControlRig = Cast<UControlRig>(ExecuteContext.Hierarchy->GetOuter()))
+			if (bInitial)
 			{
-				if (FRigPhysicsSolver* PhysicsSolver = Cast<FRigPhysicsSolver>(
-					ControlRig->GetPhysicsSolver((PhysicsComponent->BodySolverSettings.SolverID))))
-				{
-					PhysicsSolver->SetBodyGravityMultiplier(PhysicsBody, GravityMultiplier);
-				}
+				PhysicsComponent->BodyData.GravityMultiplier = GravityMultiplier;
+			}
+			else
+			{
+				PhysicsComponent->CurrentBodyData.GravityMultiplier = GravityMultiplier;
 			}
 		}
 	}
@@ -280,15 +286,15 @@ FRigUnit_HierarchySetPhysicsBodyPhysicsBlendWeight_Execute()
 	if (ExecuteContext.Hierarchy)
 	{
 		if (FRigPhysicsComponent* PhysicsComponent = Cast<FRigPhysicsComponent>(
-			ExecuteContext.Hierarchy->FindComponent(PhysicsBody)))
+			ExecuteContext.Hierarchy->FindComponent(PhysicsBodyComponentKey)))
 		{
-			if (UControlRig* ControlRig = Cast<UControlRig>(ExecuteContext.Hierarchy->GetOuter()))
+			if (bInitial)
 			{
-				if (FRigPhysicsSolver* PhysicsSolver = Cast<FRigPhysicsSolver>(
-					ControlRig->GetPhysicsSolver((PhysicsComponent->BodySolverSettings.SolverID))))
-				{
-					PhysicsSolver->SetBodyPhysicsBlendWeight(PhysicsBody, PhysicsBlendWeight);
-				}
+				PhysicsComponent->BodyData.PhysicsBlendWeight = PhysicsBlendWeight;
+			}
+			else
+			{
+				PhysicsComponent->CurrentBodyData.PhysicsBlendWeight = PhysicsBlendWeight;
 			}
 		}
 	}
@@ -301,15 +307,15 @@ FRigUnit_HierarchySetPhysicsBodyUseSkeletalAnimation_Execute()
 	if (ExecuteContext.Hierarchy)
 	{
 		if (FRigPhysicsComponent* PhysicsComponent = Cast<FRigPhysicsComponent>(
-			ExecuteContext.Hierarchy->FindComponent(PhysicsBody)))
+			ExecuteContext.Hierarchy->FindComponent(PhysicsBodyComponentKey)))
 		{
-			if (UControlRig* ControlRig = Cast<UControlRig>(ExecuteContext.Hierarchy->GetOuter()))
+			if (bInitial)
 			{
-				if (FRigPhysicsSolver* PhysicsSolver = Cast<FRigPhysicsSolver>(
-					ControlRig->GetPhysicsSolver((PhysicsComponent->BodySolverSettings.SolverID))))
-				{
-					PhysicsSolver->SetBodyUseSkeletalAnimation(PhysicsBody, bUseSkeletalAnimation);
-				}
+				PhysicsComponent->BodyData.bUseSkeletalAnimation = bUseSkeletalAnimation;
+			}
+			else
+			{
+				PhysicsComponent->CurrentBodyData.bUseSkeletalAnimation = bUseSkeletalAnimation;
 			}
 		}
 	}
@@ -322,15 +328,15 @@ FRigUnit_HierarchySetPhysicsBodyUpdateKinematicFromSimulation_Execute()
 	if (ExecuteContext.Hierarchy)
 	{
 		if (FRigPhysicsComponent* PhysicsComponent = Cast<FRigPhysicsComponent>(
-			ExecuteContext.Hierarchy->FindComponent(PhysicsBody)))
+			ExecuteContext.Hierarchy->FindComponent(PhysicsBodyComponentKey)))
 		{
-			if (UControlRig* ControlRig = Cast<UControlRig>(ExecuteContext.Hierarchy->GetOuter()))
+			if (bInitial)
 			{
-				if (FRigPhysicsSolver* PhysicsSolver = Cast<FRigPhysicsSolver>(
-					ControlRig->GetPhysicsSolver((PhysicsComponent->BodySolverSettings.SolverID))))
-				{
-					PhysicsSolver->SetBodyUpdateKinematicFromSimulation(PhysicsBody, bUpdateKinematicFromSimulation);
-				}
+				PhysicsComponent->BodyData.bUpdateKinematicFromSimulation = bUpdateKinematicFromSimulation;
+			}
+			else
+			{
+				PhysicsComponent->CurrentBodyData.bUpdateKinematicFromSimulation = bUpdateKinematicFromSimulation;
 			}
 		}
 	}
diff --git a/Engine/Plugins/Experimental/PhysicsControl/Source/ControlRigPhysics/Private/RigPhysicsSimulation.cpp b/Engine/Plugins/Experimental/PhysicsControl/Source/ControlRigPhysics/Private/RigPhysicsSimulation.cpp
new file mode 100644
--- /dev/null
+++ b/Engine/Plugins/Experimental/PhysicsControl/Source/ControlRigPhysics/Private/RigPhysicsSimulation.cpp
@@ -0,0 +1,876 @@
+// Copyright Epic Games, Inc. All Rights Reserved.
+
+#include "RigPhysicsSimulation.h"
+#include "RigPhysicsComponent.h"
+#include "RigPhysicsSolverComponent.h"
+
+#include "AnimNode_RigidBodyWithControl.h"
+#include "PhysicsControlPoseData.h"
+#include "PhysicsControlHelpers.h"
+
+#include "ControlRig.h"
+#include "Rigs/RigHierarchy.h"
+#include "Rigs/RigHierarchyController.h"
+#include "Units/RigUnitContext.h"
+#include "RigVMCore/RigVMExecuteContext.h"
+
+#include "Engine/World.h"
+
+#include "Physics/ImmediatePhysics/ImmediatePhysicsActorHandle.h"
+#include "Physics/ImmediatePhysics/ImmediatePhysicsJointHandle.h"
+#include "Physics/ImmediatePhysics/ImmediatePhysicsSimulation.h"
+#include "Physics/ImmediatePhysics/ImmediatePhysicsChaos/ImmediatePhysicsSimulation_Chaos.h"
+
+#include "Physics/Experimental/PhysScene_Chaos.h"
+
+#include "Chaos/ParticleHandle.h"
+#include "Chaos/PBDJointConstraints.h"
+#include "Chaos/ImplicitObject.h"
+#include "Chaos/ShapeInstance.h"
+#include "Chaos/Capsule.h"
+#include "Chaos/ChaosScene.h"
+#include "Chaos/PBDJointConstraintUtilities.h"
+#include "Chaos/Evolution/SimulationSpace.h"
+
+#include "PhysicsEngine/BodyInstance.h"
+#include "PhysicsEngine/ConstraintInstance.h"
+#include "PhysicsEngine/BodySetup.h"
+
+DEFINE_LOG_CATEGORY(LogRigPhysics);
+
+//======================================================================================================================
+FRigPhysicsSimulation::FRigPhysicsSimulation(TObjectPtr<UControlRig> InOwningControlRig)
+	: FRigPhysicsSimulationBase(FRigPhysicsSimulation::StaticStruct()), OwningControlRig(InOwningControlRig)
+{
+	// Danny TODO temporary hackery to allow creation in worker thread
+	Chaos::CVars::bEnableAsyncInitBody = true;
+}
+
+//======================================================================================================================
+void FRigPhysicsSimulation::InitialiseSimulation(const FRigPhysicsSolverComponent* SolverComponent)
+{
+	DestroyPhysicsSimulation();
+
+	Simulation = MakeShared<ImmediatePhysics::FSimulation>();
+	// This is needed so that when using a fixed timestep, velocities are rewound as well as
+	// positions. This is not only more accurate, but it's needed in order to get soft constraint
+	// behavior (in particular, for controls) that behave fairly independently of the control-rig
+	// tick rate.
+	Simulation->SetRewindVelocities(true);
+
+	ensure(BodyRecords.IsEmpty());
+
+	const URigHierarchy* Hierarchy = OwningControlRig->GetHierarchy();
+	if (!Hierarchy || !SolverComponent)
+	{
+		return;
+	}
+
+	FRigComponentKey SolverComponentKey = SolverComponent->GetKey();
+
+	TArray<FRigComponentKey> AllComponentKeys = Hierarchy->GetAllComponentKeys();
+	for (FRigComponentKey RigComponentKey : AllComponentKeys)
+	{
+		if (const FRigPhysicsComponent* PhysicsComponent = Cast<FRigPhysicsComponent>(
+			Hierarchy->FindComponent(RigComponentKey)))
+		{
+			if (PhysicsComponent->BodySolverSettings.PhysicsSolverComponentKey == SolverComponentKey)
+			{
+				// Just make the record for now - it will be instantiated later
+				BodyRecords.Add(RigComponentKey);
+			}
+		}
+	}
+}
+
+//======================================================================================================================
+void FRigPhysicsSimulation::DestroyPhysicsSimulation()
+{
+	for (TPair<FRigComponentKey, FRigBodyRecord>& BodyRecordPair : BodyRecords)
+	{
+		FRigBodyRecord& Record = BodyRecordPair.Value;
+		if (Record.JointHandle)
+		{
+			Simulation->DestroyJoint(Record.JointHandle);
+		}
+		if (Record.ActorHandle)
+		{
+			Simulation->DestroyActor(Record.ActorHandle);
+		}
+	}
+	BodyRecords.Empty();
+
+	for (ImmediatePhysics::FActorHandle* ActorHandle : AdditionalActorHandles)
+	{
+		Simulation->DestroyActor(ActorHandle);
+	}
+	AdditionalActorHandles.Empty();
+
+	Simulation.Reset();
+}
+
+//======================================================================================================================
+static void SetCommonProperties(const FRigPhysicsCollisionShape& Shape, FKShapeElem& ShapeElem)
+{
+	ShapeElem.RestOffset = Shape.RestOffset;
+	ShapeElem.SetName(Shape.Name);
+	ShapeElem.SetContributeToMass(Shape.bContributeToMass);
+#ifdef PER_SHAPE_COLLISION
+	// Note that FKShapeElem supports enabling/disabling collision per shape, but this is discarded by the immediate
+	// solver.
+	ShapeElem.SetCollisionEnabled(Shape.CollisionEnabled);
+#endif
+}
+
+//======================================================================================================================
+static Chaos::EJointMotionType GetMotionTypeFromLimitValue(const float Value)
+{
+	// Danny TODO This +ve/-ve thing isn't nice for users - use an enum
+	if (Value > 0.0f)
+	{
+		return Chaos::EJointMotionType::Limited;
+	}
+	else if (Value < 0.0f)
+	{
+		return Chaos::EJointMotionType::Free;
+	}
+	else
+	{
+		return Chaos::EJointMotionType::Locked;
+	}
+}
+
+//======================================================================================================================
+ImmediatePhysics::FActorHandle* FRigPhysicsSimulation::CreateBody(
+	const FName                 BodyName,
+	const FRigPhysicsCollision& Collision,
+	const FRigPhysicsDynamics*  Dynamics, 
+	const FTransform&           BodyRelSimSpaceTM) const
+{
+	// Danny TODO Use the correct materials (potentially tricky since we may be called on
+	// any thread - see the RBAN code for reference).
+
+	// This appears to be the only way
+	UBodySetup* BodySetup = NewObject<UBodySetup>();
+	BodySetup->CollisionTraceFlag = CTF_UseSimpleAsComplex;
+
+	// Note that the base class (FKShapeElem) for the shapes has got additional
+	// functionality such as enable/disable collision etc.
+
+	bool bAddedCollision = false;
+	for (const FRigPhysicsCollisionBox& Shape : Collision.Boxes)
+	{
+		FKBoxElem Elem(Shape.Extents.X, Shape.Extents.Y, Shape.Extents.Z);
+		SetCommonProperties(Shape, Elem);
+		Elem.Center = Shape.TM.GetTranslation();
+		Elem.Rotation = Shape.TM.Rotator();
+		BodySetup->AggGeom.BoxElems.Add(Elem);
+		bAddedCollision = true;
+	}
+
+	for (const FRigPhysicsCollisionSphere& Shape : Collision.Spheres)
+	{
+		FKSphereElem Elem(Shape.Radius);
+		SetCommonProperties(Shape, Elem);
+		Elem.Center = Shape.TM.GetTranslation();
+		// Note that there is no rotation
+		BodySetup->AggGeom.SphereElems.Add(Elem);
+		bAddedCollision = true;
+	}
+
+	for (const FRigPhysicsCollisionCapsule& Shape : Collision.Capsules)
+	{
+		FKSphylElem Elem(Shape.Radius, Shape.Length);
+		SetCommonProperties(Shape, Elem);
+		Elem.Center = Shape.TM.GetTranslation();
+		Elem.Rotation = Shape.TM.Rotator();
+		BodySetup->AggGeom.SphylElems.Add(Elem);
+		bAddedCollision = true;
+	}
+
+	// This is a bit odd, but in order to create a low-level immediate physics representation, you
+	// have to create a high-level temporary one. Danny TODO avoid using BodyInstance and BodySetup.
+	//
+	// See USkeletalMeshComponent::InstantiatePhysicsAsset_Internal for how it's done in RBAN (which
+	// makes body instances, converts to immediate physics actors, and then discards the body instances).
+	FBodyInstance BodyInstance;
+	if (Dynamics)
+	{
+		BodyInstance.bSimulatePhysics = true;
+		BodyInstance.bEnableGravity = true;
+		BodyInstance.SetMassOverride(Dynamics->Mass);
+		BodyInstance.LinearDamping = 0;
+		BodyInstance.AngularDamping = 0;
+	}
+	else
+	{
+		BodyInstance.bSimulatePhysics = false;
+	}
+	BodyInstance.InitBody(BodySetup, BodyRelSimSpaceTM, nullptr, nullptr);
+
+	ImmediatePhysics::FActorHandle* ActorHandle = Dynamics ?
+		Simulation->CreateDynamicActor(&BodyInstance, BodyRelSimSpaceTM) :
+		Simulation->CreateKinematicActor(&BodyInstance, BodyRelSimSpaceTM);
+
+	// Creation is guaranteed to return something
+	check(ActorHandle);
+
+	ActorHandle->SetName(BodyName);
+
+	Simulation->AddToCollidingPairs(ActorHandle);
+
+	if (bAddedCollision && Dynamics)
+	{
+		// Note that particles are always created disabled. They will simulate when disabled,
+		// but won't collide!
+		ActorHandle->SetEnabled(true);
+	}
+
+	return ActorHandle;
+}
+
+
+//======================================================================================================================
+void FRigPhysicsSimulation::Instantiate(const FRigPhysicsSolverComponent* SolverComponent)
+{
+	if (!bNeedToInstantiate || !OwningControlRig)
+	{
+		return;
+	}
+
+	// We need the simulation space in order to instantiate properly. This is not ideal, as we may
+	// end up updating the simulation space data twice (thus inserting data into the history). This
+	// shouldn't really matter as it will only be on the first step.
+	UpdateSimulationSpaceStateAndCalculateData(SolverComponent, 0.0f);
+
+	InitialiseSimulation(SolverComponent);
+
+	InstantiateSolverCollision(SolverComponent);
+
+	InstantiatePhysicsBodies(SolverComponent);
+
+	InstantiatePhysicsJoints(SolverComponent);
+
+	bNeedToInstantiate = false;
+}
+
+//======================================================================================================================
+void FRigPhysicsSimulation::InstantiateSolverCollision(const FRigPhysicsSolverComponent* SolverComponent)
+{
+	if (!SolverComponent->SolverSettings.Collision.IsEmpty())
+	{
+		// When we make these additional collision shapes, their actors are all considered to be at
+		// the origin, with the offsets being contained in the collision shapes.
+		FTransform BodyRelSimSpaceTM = ConvertCollisionSpaceTransformToSimSpace(
+			SolverComponent->SolverSettings, FTransform());
+
+		ImmediatePhysics::FActorHandle* EnvironmentActorHandle = CreateBody(
+			FName(TEXT("Environment")), SolverComponent->SolverSettings.Collision, nullptr, BodyRelSimSpaceTM);
+		AdditionalActorHandles.Add(EnvironmentActorHandle);
+	}
+}
+
+//======================================================================================================================
+void FRigPhysicsSimulation::InstantiatePhysicsJoints(const FRigPhysicsSolverComponent* SolverComponent)
+{
+	const URigHierarchy* Hierarchy = OwningControlRig->GetHierarchy();
+
+	const FRigPhysicsSolverSettings& SolverSettings = SolverComponent->SolverSettings;
+	const FRigPhysicsSimulationSpaceSettings& SimulationSpaceSettings = SolverComponent->SimulationSpaceSettings;
+
+	// Once all the bodies are created, we can make their physics joints
+	for (TPair<FRigComponentKey, FRigBodyRecord>& BodyRecordPair : BodyRecords)
+	{
+		const FRigComponentKey& ChildBodyComponentKey = BodyRecordPair.Key;
+		FRigBodyRecord& ChildBodyRecord = BodyRecordPair.Value;
+
+		ImmediatePhysics::FActorHandle* ChildActorHandle = ChildBodyRecord.ActorHandle;
+		if (!ChildActorHandle)
+		{
+			continue;
+		}
+
+		if (const FRigPhysicsComponent* ChildPhysicsComponent = Cast<FRigPhysicsComponent>(
+			Hierarchy->FindComponent(ChildBodyComponentKey)))
+		{
+			const FRigPhysicsJoint& PhysicsJoint = ChildPhysicsComponent->Joint;
+			if (!PhysicsJoint.bEnable)
+			{
+				continue;
+			}
+
+			// Now we need to find the parent body (etc) for the physics joint
+			FRigComponentKey ParentBodyComponentKey = PhysicsJoint.ParentBody;
+
+			// And along the way we will discover various other things
+			FRigElementKey ChildBoneKey = ChildBodyComponentKey.ElementKey;
+			if (!ChildBoneKey.IsValid())
+			{
+				continue;
+			}
+
+			if (!ParentBodyComponentKey.IsValid())
+			{
+				// Physics bodies are parented to a bone - we want the parent of the bone
+				FRigElementKey ParentBoneKey = Hierarchy->GetFirstParent(ChildBoneKey);
+				// Now we have the regular bone, we want the physics element that is the child
+				if (ParentBoneKey.IsValid())
+				{
+					TArray<FRigComponentKey> ComponentKeys = Hierarchy->GetComponentKeys(ParentBoneKey);
+					for (const FRigComponentKey& ComponentKey : ComponentKeys)
+					{
+						const FRigBaseComponent* BaseComponent = Hierarchy->FindComponent(ComponentKey);
+						// Danny TODO what's the best way to find/check this?!
+						if (Cast<FRigPhysicsComponent>(BaseComponent))
+						{
+							ParentBodyComponentKey = ComponentKey;
+							break;
+						}
+					}
+				}
+			}
+			const FRigPhysicsComponent* ParentPhysicsComponent =
+				Cast<FRigPhysicsComponent>(Hierarchy->FindComponent(ParentBodyComponentKey));
+			if (!ParentPhysicsComponent)
+			{
+				continue;
+			}
+
+			ChildBodyRecord.ParentBodyComponentKey = ParentBodyComponentKey;
+
+			const FRigBodyRecord* ParentRecord = nullptr;
+			ImmediatePhysics::FActorHandle* ParentActorHandle = nullptr;
+
+			// Find the record that corresponds to the parent, so we can see if there is a parent actor
+			for (const TPair<FRigComponentKey, FRigBodyRecord>& OtherBodyRecordPair : BodyRecords)
+			{
+				const FRigComponentKey& OtherBodyComponentKey = OtherBodyRecordPair.Key;
+				const FRigBodyRecord& OtherBodyRecord = OtherBodyRecordPair.Value;
+				if (OtherBodyComponentKey == ParentPhysicsComponent->GetKey())
+				{
+					ParentRecord = &OtherBodyRecord;
+					ParentActorHandle = OtherBodyRecord.ActorHandle;
+					break;
+				}
+			}
+			if (!ParentRecord)
+			{
+				continue;
+			}
+
+			// Make the physics joint (joint constraint). 
+			// UE likes to treat Body1 (index 0) as the child, and Body2 (index 1) as the parent
+
+			const FTransform ParentCoMTransform = ParentActorHandle->GetLocalCoMTransform();
+			const FTransform ChildCoMTransform = ChildActorHandle->GetLocalCoMTransform();
+
+			Chaos::FPBDJointSettings ConstraintSettings;
+			if (PhysicsJoint.bAutoCalculateChildOffset)
+			{
+				// Danny TODO - is the "non-override" actually always simply the identity TM?
+				ConstraintSettings.ConnectorTransforms[0] = FTransform();
+			}
+			ConstraintSettings.ConnectorTransforms[0] =
+				PhysicsJoint.ExtraChildOffset * ConstraintSettings.ConnectorTransforms[0];
+
+			if (PhysicsJoint.bAutoCalculateParentOffset)
+			{
+				// Danny TODO is it right to use the TM of the child bone relative to parent bone?
+				// What about if the body has got a TM too?
+				ConstraintSettings.ConnectorTransforms[1] = Hierarchy->GetLocalTransform(ChildBoneKey, true);
+			}
+			ConstraintSettings.ConnectorTransforms[1] =
+				PhysicsJoint.ExtraParentOffset * ConstraintSettings.ConnectorTransforms[1];
+
+			// Danny TODO The linear limit is likely implemented as a box limit (check) - it's just
+			// not exposed as one. We should change this.
+			ConstraintSettings.LinearLimit = PhysicsJoint.LinearLimit;
+			Chaos::EJointMotionType LinearMotionType = GetMotionTypeFromLimitValue(ConstraintSettings.LinearLimit);
+			ConstraintSettings.LinearMotionTypes = { LinearMotionType, LinearMotionType, LinearMotionType };
+
+			ConstraintSettings.AngularLimits = FMath::DegreesToRadians(PhysicsJoint.AngularLimit);
+			ConstraintSettings.AngularMotionTypes = {
+				GetMotionTypeFromLimitValue(PhysicsJoint.AngularLimit[0]),
+				GetMotionTypeFromLimitValue(PhysicsJoint.AngularLimit[1]),
+				GetMotionTypeFromLimitValue(PhysicsJoint.AngularLimit[2])
+			};
+
+			// The physics setting is backwards, because we can't enable collision on bodies that
+			// are set to not collide for other reasons.
+			ConstraintSettings.bCollisionEnabled = !PhysicsJoint.bDisableCollision;
+			ConstraintSettings.bProjectionEnabled = 
+				PhysicsJoint.LinearProjectionAmount > 0.0f || PhysicsJoint.AngularProjectionAmount > 0.0f;
+			ConstraintSettings.AngularProjection = PhysicsJoint.AngularProjectionAmount;
+			ConstraintSettings.LinearProjection = PhysicsJoint.LinearProjectionAmount;
+			ConstraintSettings.ParentInvMassScale = PhysicsJoint.ParentInverseMassScale;
+
+			ConstraintSettings.bUseLinearSolver = SolverSettings.bUseLinearJointSolver;
+
+			ChildBodyRecord.JointHandle = Simulation->CreateJoint(
+				ConstraintSettings, ChildActorHandle, ParentActorHandle);
+
+			// Now set up the drive, if requested
+			const FRigPhysicsDrive& Drive = PhysicsJoint.Drive;
+			if (ChildBodyRecord.JointHandle && Drive.bEnable)
+			{
+				float AngularSpring;
+				float AngularDamping;
+				const float MaxTorque = Drive.MaxTorque;
+
+				float LinearSpring;
+				float LinearDamping;
+				const float MaxForce = Drive.MaxForce;
+
+				UE::PhysicsControl::ConvertStrengthToSpringParams(
+					AngularSpring, AngularDamping,
+					Drive.AngularStrength, Drive.AngularDampingRatio, Drive.AngularExtraDamping);
+				UE::PhysicsControl::ConvertStrengthToSpringParams(
+					LinearSpring, LinearDamping,
+					Drive.LinearStrength, Drive.LinearDampingRatio, Drive.LinearExtraDamping);
+
+				if (Chaos::FPBDJointConstraintHandle* Constraint = ChildBodyRecord.JointHandle->GetConstraint())
+				{
+					Chaos::FPBDJointSettings Settings = Constraint->GetSettings();
+
+					// Linear
+					Settings.LinearDriveForceMode = Chaos::EJointForceMode::Acceleration;
+					const bool bLinearPositionDriveEnabled = LinearSpring > 0.0f;
+					const bool bLinearVelocityDriveEnabled = LinearDamping > 0.0f;
+					Settings.bLinearPositionDriveEnabled = { 
+						bLinearPositionDriveEnabled, bLinearPositionDriveEnabled, bLinearPositionDriveEnabled};
+					Settings.bLinearVelocityDriveEnabled = {
+						bLinearVelocityDriveEnabled, bLinearVelocityDriveEnabled, bLinearVelocityDriveEnabled};
+
+					// Angular
+					Settings.AngularDriveForceMode = Chaos::EJointForceMode::Acceleration;
+					// SLERP is generally better, but doesn't work if anything is locked
+					if (ConstraintSettings.AngularMotionTypes[0] == Chaos::EJointMotionType::Locked ||
+						ConstraintSettings.AngularMotionTypes[1] == Chaos::EJointMotionType::Locked ||
+						ConstraintSettings.AngularMotionTypes[2] == Chaos::EJointMotionType::Locked)
+					{
+						Settings.bAngularSwingPositionDriveEnabled = AngularSpring > 0.0f;
+						Settings.bAngularTwistPositionDriveEnabled = AngularSpring > 0.0f;
+						Settings.bAngularSwingVelocityDriveEnabled = AngularDamping > 0.0f;
+						Settings.bAngularTwistVelocityDriveEnabled = AngularDamping > 0.0f;
+					}
+					else
+					{
+						Settings.bAngularSLerpPositionDriveEnabled = AngularSpring > 0.0f;
+						Settings.bAngularSLerpVelocityDriveEnabled = AngularDamping > 0.0f;
+					}
+
+					Constraint->SetSettings(Settings);
+					Constraint->SetDriveParams(
+						Chaos::FVec3(LinearSpring), Chaos::FVec3(LinearDamping), Chaos::FVec3(MaxForce),
+						Chaos::FVec3(AngularSpring), Chaos::FVec3(AngularDamping), Chaos::FVec3(MaxTorque));
+
+					// The target will be set in the pre-physics update
+				}
+			}
+		}
+	}
+}
+
+//======================================================================================================================
+void FRigPhysicsSimulation::InstantiatePhysicsBodies(const FRigPhysicsSolverComponent* SolverComponent)
+{
+	const URigHierarchy* Hierarchy = OwningControlRig->GetHierarchy();
+
+	const FRigPhysicsSolverSettings& SolverSettings = SolverComponent->SolverSettings;
+	const FRigPhysicsSimulationSpaceSettings& SimulationSpaceSettings = SolverComponent->SimulationSpaceSettings;
+
+	for (TPair<FRigComponentKey, FRigBodyRecord>& BodyRecordPair : BodyRecords)
+	{
+		const FRigComponentKey& ComponentKey = BodyRecordPair.Key;
+		FRigBodyRecord& Record = BodyRecordPair.Value;
+
+		if (const FRigPhysicsComponent* PhysicsComponent = Cast<FRigPhysicsComponent>(
+			Hierarchy->FindComponent(ComponentKey)))
+		{
+			const FRigPhysicsCollision& Collision = PhysicsComponent->Collision;
+			const FRigPhysicsDynamics& Dynamics = PhysicsComponent->Dynamics;
+			const FPhysicsControlModifierData& BodyData = PhysicsComponent->BodyData;
+
+			FRigElementKey SourceKey = PhysicsComponent->BodySolverSettings.SourceBone;
+			if (!SourceKey.IsValid())
+			{
+				SourceKey = ComponentKey.ElementKey;
+			}
+
+			// What should we do if the key is not valid? 
+			if (SourceKey.IsValid())
+			{
+				FTransform SourceComponentSpaceTM = Hierarchy->GetGlobalTransform(SourceKey);
+				const FTransform SourceSimulationSpaceTM =
+					ConvertComponentSpaceTransformToSimSpace(SolverComponent->SolverSettings, SourceComponentSpaceTM);
+				Record.ActorHandle = CreateBody(ComponentKey.ElementKey.Name, Collision, &Dynamics, SourceSimulationSpaceTM);
+			}
+		}
+	}
+}
+
+//======================================================================================================================
+void FRigPhysicsSimulation::UpdatePrePhysics(
+	const FRigVMExecuteContext&     ExecuteContext, 
+	const FRigPhysicsSolverComponent* SolverComponent, 
+	const float                     DeltaTime, 
+	const FSimulationSpaceData&     SimulationSpaceData)
+{
+	URigHierarchy* Hierarchy = OwningControlRig->GetHierarchy();
+
+	for (ImmediatePhysics::FActorHandle* ActorHandle : AdditionalActorHandles)
+	{
+		FTransform BodyRelSimSpaceTM = ConvertCollisionSpaceTransformToSimSpace(
+			SolverComponent->SolverSettings, FTransform());
+		ActorHandle->SetKinematicTarget(BodyRelSimSpaceTM);
+	}
+
+	for (TPair<FRigComponentKey, FRigBodyRecord>& BodyRecordPair : BodyRecords)
+	{
+		const FRigComponentKey& ComponentKey = BodyRecordPair.Key;
+		const FRigBodyRecord& Record = BodyRecordPair.Value;
+
+		if (const FRigPhysicsComponent* PhysicsComponent = Cast<FRigPhysicsComponent>(
+			Hierarchy->FindComponent(ComponentKey)))
+		{
+			if (Record.ActorHandle)
+			{
+				UpdateBodyFromModifierData(
+					Record.ActorHandle, Simulation.Get(), PhysicsComponent->CurrentBodyData, SimulationSpaceData.Gravity);
+
+				if (Record.ActorHandle->GetIsKinematic())
+				{
+					// Get the target in component space, and then convert it into sim space if necessary.
+
+					// If the target is already in component space, then that's all we need
+					FTransform KinematicTargetCS;
+					switch (PhysicsComponent->CurrentKinematicTargetSpace)
+					{
+						case ERigPhysicsKinematicTargetSpace::Component:
+						{
+							KinematicTargetCS = PhysicsComponent->CurrentKinematicTarget;
+							break;
+						}
+						case ERigPhysicsKinematicTargetSpace::World:
+						{
+							// Record.KinematicTarget * SimulationSpaceState.ComponentTM.Inverse();
+							KinematicTargetCS = PhysicsComponent->CurrentKinematicTarget.GetRelativeTransform(
+								SimulationSpaceState.ComponentTM);
+							break;
+						}
+						default:
+						{
+							// All the other options are relative to a bone, so the first task is to get
+							// that, which will be in component space
+							FRigElementKey SourceKey = PhysicsComponent->BodySolverSettings.SourceBone;
+							if (!SourceKey.IsValid())
+							{
+								SourceKey = ComponentKey.ElementKey;
+							}
+							if (SourceKey.IsValid())
+							{
+								FTransform SourceComponentSpaceTM = Hierarchy->GetGlobalTransform(SourceKey);
+
+								switch (PhysicsComponent->CurrentKinematicTargetSpace)
+								{
+									case ERigPhysicsKinematicTargetSpace::OffsetInBoneSpace:
+									{
+										KinematicTargetCS = PhysicsComponent->CurrentKinematicTarget * SourceComponentSpaceTM;
+										break;
+									}
+									case ERigPhysicsKinematicTargetSpace::OffsetInWorldSpace:
+									{
+										// Convert the bone to WS, apply the target, and convert back
+										FTransform BoneWS = SourceComponentSpaceTM * SimulationSpaceState.ComponentTM;
+										FTransform KinematicTargetWS = BoneWS * PhysicsComponent->CurrentKinematicTarget;
+										// Danny TODO figure out which of the GetRelativeTransform versions this is
+										KinematicTargetCS = SimulationSpaceState.ComponentTM.Inverse() * KinematicTargetWS;
+										break;
+									}
+									case ERigPhysicsKinematicTargetSpace::OffsetInComponentSpace:
+									{
+										KinematicTargetCS = SourceComponentSpaceTM * PhysicsComponent->CurrentKinematicTarget;
+										break;
+									}
+									default:
+									{
+										UE_CONTROLRIG_RIGUNIT_REPORT_ERROR(TEXT("Kinematic target space is not valid"));
+									}
+								}
+							}
+						}
+					}
+					FTransform KinematicTargetTM = ConvertComponentSpaceTransformToSimSpace(
+						SolverComponent->SolverSettings, KinematicTargetCS);
+					Record.ActorHandle->SetKinematicTarget(KinematicTargetTM);
+				}
+
+				// Now update the joint targets
+				if (Record.JointHandle)
+				{
+					if (Chaos::FPBDJointConstraintHandle* Constraint = Record.JointHandle->GetConstraint())
+					{
+						const Chaos::FPBDJointSettings& Settings = Constraint->GetSettings();
+
+						if (Settings.AngularDriveStiffness.SquaredLength() > 0.0f ||
+							Settings.AngularDriveDamping.SquaredLength() > 0.0f ||
+							Settings.LinearDriveStiffness.SquaredLength() > 0.0f ||
+							Settings.LinearDriveDamping.SquaredLength() > 0.0f)
+						{
+							// Multiplier on the velocity calculated from the current and previous target
+							float TargetVelocityMultiplier = 1.0f;
+
+							FRigElementKey ChildSourceKey;
+							if (const FRigPhysicsComponent* ChildPhysicsComponent = Cast<FRigPhysicsComponent>(
+								Hierarchy->FindComponent(ComponentKey)))
+							{
+								ChildSourceKey = ChildPhysicsComponent->BodySolverSettings.SourceBone;
+								if (!ChildSourceKey.IsValid())
+								{
+									ChildSourceKey = ComponentKey.ElementKey;
+								}
+								TargetVelocityMultiplier = 
+									ChildPhysicsComponent->Joint.Drive.SkeletalAnimationVelocityMultiplier;
+							}
+
+							FRigElementKey ParentSourceKey;
+							if (const FRigPhysicsComponent* ParentPhysicsComponent = Cast<FRigPhysicsComponent>(
+								Hierarchy->FindComponent(Record.ParentBodyComponentKey)))
+							{
+								ParentSourceKey = ParentPhysicsComponent->BodySolverSettings.SourceBone;
+								if (!ParentSourceKey.IsValid())
+								{
+									ParentSourceKey = Record.ParentBodyComponentKey.ElementKey;
+								}
+							}
+
+							if (!ChildSourceKey.IsValid() || !ParentSourceKey.IsValid())
+							{
+								continue;
+							}
+
+							// Note that the drive operates between a parent and child part, so we
+							// don't need to worry about global/component (etc) space.
+							const FTransform ChildTM = Hierarchy->GetGlobalTransform(ChildSourceKey);
+							const FTransform ParentTM = Hierarchy->GetGlobalTransform(ParentSourceKey);
+
+							const FTransform ComponentSpaceParentFrameTM = Settings.ConnectorTransforms[1] * ParentTM;
+							const FTransform ComponentSpaceChildFrameTM = Settings.ConnectorTransforms[0] * ChildTM;
+
+							const FTransform DriveTargetTM = 
+								ComponentSpaceChildFrameTM * ComponentSpaceParentFrameTM.Inverse();
+
+							Constraint->SetLinearDrivePositionTarget(DriveTargetTM.GetTranslation());
+							Constraint->SetAngularDrivePositionTarget(DriveTargetTM.GetRotation());
+
+							if (Record.PreviousDriveTargetUpdateCounter + 1 == UpdateCounter && 
+								TargetVelocityMultiplier > 0.0f)
+							{
+								if (DeltaTime > SMALL_NUMBER)
+								{
+									const FTransform DriveTargetTMDelta = 
+										Record.PreviousDriveTargetTM.Inverse() * DriveTargetTM;
+									const FVector Velocity = DriveTargetTMDelta.GetTranslation() / DeltaTime;
+									const FVector AngularVelocity = 
+										DriveTargetTMDelta.GetRotation().GetShortestArcWith(
+											FQuat::Identity).ToRotationVector() / DeltaTime;
+									if (!Velocity.ContainsNaN() && !AngularVelocity.ContainsNaN())
+									{
+										Constraint->SetLinearDriveVelocityTarget(
+											Velocity * TargetVelocityMultiplier);
+										Constraint->SetAngularDriveVelocityTarget(
+											AngularVelocity * TargetVelocityMultiplier);
+									}
+								}
+							}
+							else
+							{
+								Constraint->SetLinearDriveVelocityTarget(Chaos::FVec3(0));
+								Constraint->SetAngularDriveVelocityTarget(Chaos::FVec3(0));
+							}
+							BodyRecordPair.Value.PreviousDriveTargetUpdateCounter = UpdateCounter;
+							BodyRecordPair.Value.PreviousDriveTargetTM = DriveTargetTM;
+						}
+					}
+				}
+			}
+		}
+	}
+}
+
+//======================================================================================================================
+// Note that we read back into a target bone, which may have been specified explicitly, or will
+// otherwise default to the physics element parent.
+void FRigPhysicsSimulation::UpdatePostPhysics(const FRigPhysicsSolverComponent* SolverComponent, float Alpha)
+{
+	if (Alpha == 0.0f)
+	{
+		return;
+	}
+	const FRigPhysicsSolverSettings& SolverSettings = SolverComponent->SolverSettings;
+
+	bool bGotInvalidSimulationData = false;
+	URigHierarchy* Hierarchy = OwningControlRig->GetHierarchy();
+	for (TPair<FRigComponentKey, FRigBodyRecord>& BodyRecordPair : BodyRecords)
+	{
+		const FRigComponentKey& ComponentKey = BodyRecordPair.Key;
+		const FRigBodyRecord& Record = BodyRecordPair.Value;
+
+		if (Record.ActorHandle)
+		{
+			if (const FRigPhysicsComponent* PhysicsComponent = Cast<FRigPhysicsComponent>(
+				Hierarchy->FindComponent(ComponentKey)))
+			{
+				// We will only have created this body record if it's associated with our solver, so
+				// we don't need to check again that our solver is in its list.
+				FRigElementKey TargetElementKey = PhysicsComponent->BodySolverSettings.TargetBone;
+				if (!TargetElementKey.IsValid())
+				{
+					TargetElementKey = ComponentKey.ElementKey;
+				}
+
+				if (TargetElementKey.IsValid())
+				{
+					FTransform TargetSimSpaceTM = Record.ActorHandle->GetWorldTransform();
+					double DistSq = TargetSimSpaceTM.GetTranslation().SquaredLength();
+					if (!TargetSimSpaceTM.IsValid() || (SolverSettings.DistanceThresholdForReset > 0.0f &&
+						DistSq > FMath::Square(SolverSettings.DistanceThresholdForReset)))
+					{
+						bGotInvalidSimulationData = true;
+					}
+					else
+					{
+						FTransform TargetComponentSpaceTM = ConvertSimSpaceTransformToComponentSpace(
+							SolverSettings, TargetSimSpaceTM);
+						if (Alpha < 0.999f)
+						{
+							// Danny TODO Note that this uses Alpha to blend in component space.
+							// This can cause joint separation. We probably want an option to blend
+							// in local (joint) space, perhaps splitting the alpha into orientation
+							// and position.
+							FTransform CurrentTM = Hierarchy->GetGlobalTransform(TargetElementKey);
+							FQuat TargetQ = FQuat::Slerp(
+								CurrentTM.GetRotation(), TargetComponentSpaceTM.GetRotation(), Alpha);
+							FVector TargetT = FMath::Lerp(
+								CurrentTM.GetTranslation(), TargetComponentSpaceTM.GetTranslation(), Alpha);
+							TargetComponentSpaceTM.SetRotation(TargetQ);
+							TargetComponentSpaceTM.SetTranslation(TargetT);
+						}
+						// Danny TODO should we be "updating" the children here or not? 
+						Hierarchy->SetGlobalTransform(TargetElementKey, TargetComponentSpaceTM);
+					}
+				}
+			}
+		}
+	}
+
+	if (bGotInvalidSimulationData)
+	{
+		for (TPair<FRigComponentKey, FRigBodyRecord>& BodyRecordPair : BodyRecords)
+		{
+			const FRigComponentKey& ComponentKey = BodyRecordPair.Key;
+			const FRigBodyRecord& Record = BodyRecordPair.Value;
+
+			if (const FRigPhysicsComponent* PhysicsComponent = Cast<FRigPhysicsComponent>(
+				Hierarchy->FindComponent(ComponentKey)))
+			{
+				if (Record.ActorHandle && !Record.ActorHandle->GetIsKinematic())
+				{
+					// Get the TM in component space, and then convert it into sim space.
+					FRigElementKey SourceKey = PhysicsComponent->BodySolverSettings.SourceBone;
+					if (!SourceKey.IsValid())
+					{
+						SourceKey = ComponentKey.ElementKey;
+					}
+					if (SourceKey.IsValid())
+					{
+						FTransform SourceComponentSpaceTM = Hierarchy->GetGlobalTransform(SourceKey);
+						const FTransform SourceSimulationSpaceTM =
+							ConvertComponentSpaceTransformToSimSpace(SolverSettings, SourceComponentSpaceTM);
+						Record.ActorHandle->SetWorldTransform(SourceSimulationSpaceTM);
+						Record.ActorHandle->SetLinearVelocity(FVector::ZeroVector);
+						Record.ActorHandle->SetAngularVelocity(FVector::ZeroVector);
+					}
+				}
+			}
+		}
+	}
+}
+
+//======================================================================================================================
+void FRigPhysicsSimulation::StepSimulation(
+	const FRigVMExecuteContext&     ExecuteContext, 
+	const FRigPhysicsSolverComponent* SolverComponent, 
+	const float                     DeltaTimeOverride, 
+	const float                     Alpha)
+{
+	float PhysicsDeltaTime = ExecuteContext.GetDeltaTime();
+	if (DeltaTimeOverride > 0.0f)
+	{
+		PhysicsDeltaTime = DeltaTimeOverride;
+	}
+	else if (DeltaTimeOverride < 0.0f)
+	{
+		PhysicsDeltaTime = 0.0f;
+	}
+
+	// We need to know about the simulation space etc before we can instantiate anything into the right place
+	FSimulationSpaceData SimulationSpaceData = UpdateSimulationSpaceStateAndCalculateData(
+		SolverComponent, PhysicsDeltaTime);
+
+	// We instantiate when we do the first simulation - this makes sure any changes applied by
+	// the user have been made. It also means there is no overhead if physics is never stepped.
+	// However, there may be a hitch due to the creation, so it may also happen during construction.
+	Instantiate(SolverComponent); // There is an early out if it's already been done
+
+	if (!Simulation)
+	{
+		return;
+	}
+
+	const FRigPhysicsSolverSettings& SolverSettings = SolverComponent->SolverSettings;
+	const FRigPhysicsSimulationSpaceSettings& SimulationSpaceSettings = SolverComponent->SimulationSpaceSettings;
+
+	// Set settings that might change
+	Simulation->SetSolverSettings(
+		SolverSettings.FixedTimeStep,
+		999999999.0, // Danny TODO cull distance - expose this? Yuck
+		SolverSettings.MaxDepenetrationVelocity,
+		SolverSettings.bUseLinearJointSolver,
+		SolverSettings.PositionIterations,
+		SolverSettings.VelocityIterations,
+		SolverSettings.ProjectionIterations,
+		SolverSettings.bUseManifolds);
+	// This gets reset to 100 after every simulation step!
+	Simulation->SetMaxNumRollingAverageStepTimes(SolverSettings.MaxNumRollingAverageStepTimes);
+
+	// Other settings - would normally be static (so Danny TODO move this)
+	ChaosJointSolverSettings.bSolvePositionLast = SolverSettings.bSolveJointPositionsLast;
+	ChaosJointSolverSettings.bSortEnabled = true;
+
+	// Simulation space
+	Chaos::FSimulationSpaceSettings ChaosSimulationSpaceSettings;
+	ChaosSimulationSpaceSettings.bEnabled = (SimulationSpaceSettings.SpaceMovementAmount > 0.0f);
+	ChaosSimulationSpaceSettings.ExternalLinearEtherDrag = SimulationSpaceSettings.ExternalLinearDragV;
+	Simulation->SetSimulationSpaceSettings(ChaosSimulationSpaceSettings);
+	Simulation->UpdateSimulationSpace(
+		SimulationSpaceState.SimulationSpaceTM,
+		SimulationSpaceSettings.SpaceMovementAmount * SimulationSpaceData.LinearVelocity,
+		SimulationSpaceSettings.SpaceMovementAmount * SimulationSpaceData.AngularVelocity,
+		SimulationSpaceSettings.SpaceMovementAmount * SimulationSpaceData.LinearAcceleration,
+		SimulationSpaceSettings.SpaceMovementAmount * SimulationSpaceData.AngularAcceleration);
+
+	if (PhysicsDeltaTime > 0.0f)
+	{
+		UpdatePrePhysics(ExecuteContext, SolverComponent, PhysicsDeltaTime, SimulationSpaceData);
+
+		Simulation->Simulate(
+			PhysicsDeltaTime, SolverSettings.MaxDeltaTime, SolverSettings.MaxTimeSteps,
+			SimulationSpaceData.Gravity, &ChaosJointSolverSettings);
+	}
+
+	// Always do a read-back, even for zero Dt
+	UpdatePostPhysics(SolverComponent, Alpha);
+
+	++UpdateCounter;
+}
diff --git a/Engine/Plugins/Experimental/PhysicsControl/Source/ControlRigPhysics/Private/RigPhysicsSimulation.h b/Engine/Plugins/Experimental/PhysicsControl/Source/ControlRigPhysics/Private/RigPhysicsSimulation.h
new file mode 100644
--- /dev/null
+++ b/Engine/Plugins/Experimental/PhysicsControl/Source/ControlRigPhysics/Private/RigPhysicsSimulation.h
@@ -0,0 +1,192 @@
+// Copyright Epic Games, Inc. All Rights Reserved.
+
+#pragma once
+
+#include "RigPhysicsData.h"
+
+#include "PhysicsControlData.h"
+
+#include "Rigs/RigHierarchyComponents.h"
+
+#include "Physics/ImmediatePhysics/ImmediatePhysicsDeclares.h"
+#include "Chaos/PBDJointConstraintTypes.h"
+
+#include "RigPhysicsSimulation.generated.h"
+
+class UControlRig;
+
+struct FRigVMDrawInterface;
+struct FRigVMExecuteContext;
+struct FRigPhysicsSolverComponent;
+
+struct FRigPhysicsCollision;
+struct FRigPhysicsDynamics;
+
+#if UE_BUILD_SHIPPING || UE_BUILD_TEST
+DECLARE_LOG_CATEGORY_EXTERN(LogRigPhysics, Log, Warning);
+#else
+DECLARE_LOG_CATEGORY_EXTERN(LogRigPhysics, Log, All);
+#endif
+
+/**
+ * For internal use only, to keep track of objects we have created
+ */
+USTRUCT()
+struct FRigBodyRecord
+{
+	GENERATED_BODY()
+
+	// Things that are set during instantiation
+	ImmediatePhysics::FActorHandle* ActorHandle = nullptr;
+	ImmediatePhysics::FJointHandle* JointHandle = nullptr;
+	// When we have a joint, this is the key for the parent body. Note that the child body component
+	// key is stored as the lookup key for the records, so we don't need to duplicate it here.
+	FRigComponentKey                ParentBodyComponentKey;
+
+	// Things that are updated as the simulation progresses
+
+	// The drive works with velocities so we store the previous target transform, and when it was stored.
+	FTransform                      PreviousDriveTargetTM;
+	// This is stored from the main solver update counter, marking when the previous drive TM was valid.
+	uint64                          PreviousDriveTargetUpdateCounter = 0;
+};
+
+
+USTRUCT()
+struct FRigPhysicsSimulation : public FRigPhysicsSimulationBase
+{
+public:
+	GENERATED_BODY()
+
+	FRigPhysicsSimulation(
+		TObjectPtr<UControlRig> InOwningControlRig = nullptr);
+
+	// This will initialise/create the simulation and then create everything we need in it. 
+	void Instantiate(const FRigPhysicsSolverComponent* SolverComponent);
+
+	// Integrates the simulation forwards.
+	// If DeltaTimeOverride is +ve, then that value is used.
+	// If it is zero, then delta time is taken from the execute context
+	// If it is negative, then the simulation isn't stepped.
+	void StepSimulation(
+		const FRigVMExecuteContext&       ExecuteContext, 
+		const FRigPhysicsSolverComponent* SolverComponent, 
+		const float                       DeltaTimeOverride, 
+		const float                       Alpha);
+
+	// Draws shapes etc AND (potentially) enables the low-level Chaos debug draw
+	void Draw(
+		FRigVMDrawInterface*                    DI,
+		const FRigPhysicsSolverSettings&        SolverSettings,
+		const FRigPhysicsVisualizationSettings& VisualizationSettings,
+		const UWorld*                           DebugWorld) const;
+
+private:
+	// Used by the world-space to simulation-space motion transfer system in Component- or
+	// Bone-Space sims, and preserved between updates.
+	struct FSimulationSpaceState
+	{
+		FTransform ComponentTM;
+		FTransform BoneRelComponentTM;
+		FVector    ComponentLinearVelocity;
+		FVector    ComponentAngularVelocity;
+		FVector    BoneLinearVelocity;
+		FVector    BoneAngularVelocity;
+
+		// The world transform of the simulation space
+		FTransform SimulationSpaceTM;
+		FTransform PrevSimulationSpaceTM;
+		FTransform PrevPrevSimulationSpaceTM;
+		// The time between SimulationSpaceTM and PrevSimulationSpaceTM
+		float      Dt = 1.0f; 
+		// The time between PrevSimulationSpaceTM and PrevPrevSimulationSpaceTM
+		float      PrevDt = 1.0f;
+	};
+
+	// Represents the properties of the simulation space - calculated near the beginning of the update.
+	// Note that all these are specified in the simulation space itself
+	struct FSimulationSpaceData
+	{
+		FVector LinearVelocity;
+		FVector AngularVelocity;
+		FVector LinearAcceleration;
+		FVector AngularAcceleration;
+		FVector Gravity;
+	};
+
+	// This collects all the bodies associated with the solver and makes records for them
+	void InitialiseSimulation(const FRigPhysicsSolverComponent* SolverComponent);
+
+	void InstantiateSolverCollision(const FRigPhysicsSolverComponent* SolverComponent);
+
+	void InstantiatePhysicsJoints(const FRigPhysicsSolverComponent* SolverComponent);
+
+	void InstantiatePhysicsBodies(const FRigPhysicsSolverComponent* SolverComponent);
+
+	// Creates an actor with collision. This will by dynamic if Dynamics is valid, or otherwise kinematic
+	ImmediatePhysics::FActorHandle* CreateBody(
+		const FName                 BodyName,
+		const FRigPhysicsCollision& Collision,
+		const FRigPhysicsDynamics*  Dynamics,
+		const FTransform&           BodyRelSimSpaceTM) const;
+
+	// This releases the actors etc in the simulation, and destroys the simulation, but it doesn't
+	// change our records (apart from resetting pointers to actors etc).
+	void DestroyPhysicsSimulation();
+
+	void UpdatePrePhysics(
+		const FRigVMExecuteContext&       ExecuteContext, 
+		const FRigPhysicsSolverComponent* SolverComponent, 
+		const float                       DeltaTime, 
+		const FSimulationSpaceData&       SimulationSpaceData);
+	void UpdatePostPhysics(const FRigPhysicsSolverComponent* SolverComponent, const float Alpha);
+
+	void InitSimulationSpace(const FTransform& ComponentToWorld, const FTransform& BoneToComponent);
+
+	FSimulationSpaceData UpdateSimulationSpaceStateAndCalculateData(
+		const FRigPhysicsSolverComponent* SolverComponent, const float Dt);
+
+	// Returns the simulation space transform, in world space
+	FTransform GetSimulationSpaceTransform(const FRigPhysicsSolverSettings& SolverSettings) const;
+
+	// Converts a transform from component space (e.g. coming from the owning control rig) into the
+	// simulation space
+	FTransform ConvertComponentSpaceTransformToSimSpace(
+		const FRigPhysicsSolverSettings& SolverSettings, const FTransform& TM) const;
+
+	// Converts a transform from the simulation space to component space (e.g. for writing back to
+	// the owning control rig)
+	FTransform ConvertSimSpaceTransformToComponentSpace(
+		const FRigPhysicsSolverSettings& SolverSettings, const FTransform& TM) const;
+
+	// Converts a vector specified in world space to the simulation space (e.g. converting gravity)
+	FVector ConvertWorldVectorToSimSpaceNoScale(
+		const FRigPhysicsSolverSettings& SolverSettings, const FVector& WorldVector) const;
+
+	FTransform ConvertCollisionSpaceTransformToSimSpace(
+		const FRigPhysicsSolverSettings& SolverSettings, const FTransform& TM) const;
+
+private:
+	TObjectPtr<UControlRig> OwningControlRig;
+	FRigComponentKey        PhysicsSolverComponentKey;
+
+	TMap<FRigComponentKey, FRigBodyRecord> BodyRecords;
+	TSharedPtr<ImmediatePhysics::FSimulation> Simulation;
+
+	// Used to store things like collision shapes.
+	TArray<ImmediatePhysics::FActorHandle*> AdditionalActorHandles;
+
+	Chaos::FPBDJointSolverSettings ChaosJointSolverSettings;
+
+	FSimulationSpaceState SimulationSpaceState;
+
+	// This is incremented at the end of each simulation step, used to identify when previously
+	// calculated values are valid
+	// Note that the universe will roll over before a uint64 does
+	uint64 UpdateCounter = 0;
+
+	// This will be set to false after instantiation, which can happen manually during construction,
+	// or be postponed until there is a simulation step.
+	bool bNeedToInstantiate = true;
+};
+
diff --git a/Engine/Plugins/Experimental/PhysicsControl/Source/ControlRigPhysics/Private/RigPhysicsSimulationSpace.cpp b/Engine/Plugins/Experimental/PhysicsControl/Source/ControlRigPhysics/Private/RigPhysicsSimulationSpace.cpp
new file mode 100644
--- /dev/null
+++ b/Engine/Plugins/Experimental/PhysicsControl/Source/ControlRigPhysics/Private/RigPhysicsSimulationSpace.cpp
@@ -0,0 +1,292 @@
+// Copyright Epic Games, Inc. All Rights Reserved.
+
+#include "RigPhysicsSimulation.h"
+#include "RigPhysicsSolverComponent.h"
+
+#include "Rigs/RigHierarchy.h"
+#include "Rigs/RigHierarchyController.h"
+#include "ControlRig.h"
+
+#include "Physics/ImmediatePhysics/ImmediatePhysicsActorHandle.h"
+#include "Physics/ImmediatePhysics/ImmediatePhysicsJointHandle.h"
+#include "Physics/ImmediatePhysics/ImmediatePhysicsSimulation.h"
+
+#include "Physics/Experimental/PhysScene_Chaos.h"
+#include "Chaos/ParticleHandle.h"
+#include "Chaos/PBDJointConstraints.h"
+#include "Chaos/ImplicitObject.h"
+#include "Chaos/ShapeInstance.h"
+#include "Chaos/Capsule.h"
+#include "Chaos/PBDJointConstraintUtilities.h"
+
+#include "PhysicsEngine/BodyInstance.h"
+#include "PhysicsEngine/ConstraintInstance.h"
+#include "PhysicsEngine/BodySetup.h"
+
+//======================================================================================================================
+FORCEINLINE_DEBUGGABLE FTransform GetSpaceTransform(
+	ERigPhysicsSimulationSpace Space, const FTransform& ComponentTM, const FTransform& BoneTM)
+{
+	switch (Space)
+	{
+	case ERigPhysicsSimulationSpace::Component: return ComponentTM;
+	case ERigPhysicsSimulationSpace::World: return FTransform::Identity;
+	case ERigPhysicsSimulationSpace::SpaceBone: return BoneTM * ComponentTM;
+	default: ensureMsgf(false, TEXT("Unsupported Simulation Space")); return FTransform::Identity;
+	}
+}
+
+//======================================================================================================================
+FTransform FRigPhysicsSimulation::GetSimulationSpaceTransform(const FRigPhysicsSolverSettings& SolverSettings) const
+{
+	switch (SolverSettings.SimulationSpace)
+	{
+	case ERigPhysicsSimulationSpace::Component: return SimulationSpaceState.ComponentTM;
+	case ERigPhysicsSimulationSpace::World: return FTransform::Identity;
+	case ERigPhysicsSimulationSpace::SpaceBone: return SimulationSpaceState.BoneRelComponentTM * SimulationSpaceState.ComponentTM;
+	default: ensureMsgf(false, TEXT("Unsupported Simulation Space")); return FTransform::Identity;
+	}
+}
+
+//======================================================================================================================
+FTransform FRigPhysicsSimulation::ConvertComponentSpaceTransformToSimSpace(
+	const FRigPhysicsSolverSettings& SolverSettings, const FTransform& TM) const
+{
+	switch (SolverSettings.SimulationSpace)
+	{
+	case ERigPhysicsSimulationSpace::Component: return TM;
+	case ERigPhysicsSimulationSpace::World:  return TM * SimulationSpaceState.ComponentTM;
+	case ERigPhysicsSimulationSpace::SpaceBone: return TM.GetRelativeTransform(SimulationSpaceState.BoneRelComponentTM);
+	default: ensureMsgf(false, TEXT("Unsupported Simulation Space")); return TM;
+	}
+}
+
+//======================================================================================================================
+FTransform FRigPhysicsSimulation::ConvertSimSpaceTransformToComponentSpace(
+	const FRigPhysicsSolverSettings& SolverSettings, const FTransform& TM) const
+{
+	switch (SolverSettings.SimulationSpace)
+	{
+	case ERigPhysicsSimulationSpace::Component: return TM;
+	case ERigPhysicsSimulationSpace::World:  return TM.GetRelativeTransform(SimulationSpaceState.ComponentTM);
+	case ERigPhysicsSimulationSpace::SpaceBone: return TM * SimulationSpaceState.BoneRelComponentTM;
+	default: ensureMsgf(false, TEXT("Unsupported Simulation Space")); return TM;
+	}
+}
+
+//======================================================================================================================
+FORCEINLINE_DEBUGGABLE FVector ConvertWorldVectorToSimSpaceNoScale(
+	ERigPhysicsSimulationSpace Space, const FVector& WorldVector, 
+	const FTransform& ComponentTM, const FTransform& BoneTM)
+{
+	switch (Space)
+	{
+	case ERigPhysicsSimulationSpace::Component: return ComponentTM.InverseTransformVectorNoScale(WorldVector);
+	case ERigPhysicsSimulationSpace::World: return WorldVector;
+	case ERigPhysicsSimulationSpace::SpaceBone:
+		return BoneTM.InverseTransformVectorNoScale(ComponentTM.InverseTransformVectorNoScale(WorldVector));
+	default: ensureMsgf(false, TEXT("Unsupported Simulation Space")); return WorldVector;
+	}
+}
+
+//======================================================================================================================
+FVector FRigPhysicsSimulation::ConvertWorldVectorToSimSpaceNoScale(
+	const FRigPhysicsSolverSettings& SolverSettings, const FVector& WorldVector) const
+{
+	switch (SolverSettings.SimulationSpace)
+	{
+	case ERigPhysicsSimulationSpace::Component: return SimulationSpaceState.ComponentTM.InverseTransformVectorNoScale(WorldVector);
+	case ERigPhysicsSimulationSpace::World: return WorldVector;
+	case ERigPhysicsSimulationSpace::SpaceBone:
+		return SimulationSpaceState.BoneRelComponentTM.InverseTransformVectorNoScale(
+			SimulationSpaceState.ComponentTM.InverseTransformVectorNoScale(WorldVector));
+	default: ensureMsgf(false, TEXT("Unsupported Simulation Space")); return WorldVector;
+	}
+}
+
+//======================================================================================================================
+FTransform FRigPhysicsSimulation::ConvertCollisionSpaceTransformToSimSpace(
+	const FRigPhysicsSolverSettings& SolverSettings, const FTransform& TM) const
+{
+	FTransform SimSpaceTM = GetSpaceTransform(
+		SolverSettings.SimulationSpace, 
+		SimulationSpaceState.ComponentTM, 
+		SimulationSpaceState.BoneRelComponentTM);
+	FTransform CollisionSpaceTM = GetSpaceTransform(
+		SolverSettings.CollisionSpace, 
+		SimulationSpaceState.ComponentTM, 
+		SimulationSpaceState.BoneRelComponentTM);
+
+	FTransform WorldSpaceTM = TM * CollisionSpaceTM;
+	return WorldSpaceTM.GetRelativeTransform(SimSpaceTM);
+}
+
+//======================================================================================================================
+void FRigPhysicsSimulation::InitSimulationSpace(
+	const FTransform& ComponentTM,
+	const FTransform& BoneRelComponentTM)
+{
+	SimulationSpaceState.ComponentTM = ComponentTM;
+	SimulationSpaceState.BoneRelComponentTM = BoneRelComponentTM;
+	SimulationSpaceState.ComponentLinearVelocity = FVector::ZeroVector;
+	SimulationSpaceState.ComponentAngularVelocity = FVector::ZeroVector;
+	SimulationSpaceState.BoneLinearVelocity = FVector::ZeroVector;
+	SimulationSpaceState.BoneAngularVelocity = FVector::ZeroVector;
+}
+
+//======================================================================================================================
+FVector CalculateAngularVelocity(const FQuat& PrevQ, const FQuat& CurrentQ, float Dt)
+{
+	const FQuat DeltaQ = (CurrentQ * PrevQ.Inverse()).GetShortestArcWith(FQuat::Identity);
+	return DeltaQ.ToRotationVector() / Dt;
+}
+
+//======================================================================================================================
+// Note - don't use the space conversion functions here as the state won't have been set yet.
+FRigPhysicsSimulation::FSimulationSpaceData FRigPhysicsSimulation::UpdateSimulationSpaceStateAndCalculateData(
+	const FRigPhysicsSolverComponent* SolverComponent, const float Dt)
+{
+	const FRigPhysicsSolverSettings& SolverSettings = SolverComponent->SolverSettings;
+	const FRigPhysicsSimulationSpaceSettings& SimulationSpaceSettings = SolverComponent->SimulationSpaceSettings;
+
+	if (const USceneComponent* OwningSceneComponent = OwningControlRig->GetOwningSceneComponent())
+	{
+		SimulationSpaceState.ComponentTM = OwningSceneComponent->GetComponentTransform();
+	}
+	else 
+	{
+		SimulationSpaceState.ComponentTM.SetIdentity();
+	}
+
+	if (SolverSettings.SimulationSpace == ERigPhysicsSimulationSpace::SpaceBone &&  SolverSettings.SpaceBone.IsValid())
+	{
+		URigHierarchy* Hierarchy = OwningControlRig->GetHierarchy();
+		SimulationSpaceState.BoneRelComponentTM = Hierarchy->GetGlobalTransform(SolverSettings.SpaceBone);
+	}
+
+	// Record the history - but avoid polluting it with zero Dt updates. What that means is - if we
+	// get a zero-Dt update, then just update our current sim space TM, which means the time delta
+	// from the previous state is actually the current Dt (i.e. don't overwrite the current Dt).
+	if (Dt > SMALL_NUMBER)
+	{
+		SimulationSpaceState.PrevDt = SimulationSpaceState.Dt;
+		SimulationSpaceState.Dt = Dt;
+
+		SimulationSpaceState.PrevPrevSimulationSpaceTM = SimulationSpaceState.PrevSimulationSpaceTM;
+		SimulationSpaceState.PrevSimulationSpaceTM = SimulationSpaceState.SimulationSpaceTM;
+	}
+	SimulationSpaceState.SimulationSpaceTM = GetSpaceTransform(
+		SolverSettings.SimulationSpace, SimulationSpaceState.ComponentTM, SimulationSpaceState.BoneRelComponentTM);
+
+	FSimulationSpaceData SimulationSpaceData;
+	SimulationSpaceData.Gravity = ::ConvertWorldVectorToSimSpaceNoScale(
+		SolverSettings.SimulationSpace, SolverSettings.Gravity, 
+		SimulationSpaceState.ComponentTM, SimulationSpaceState.BoneRelComponentTM);
+
+	if (SolverSettings.SimulationSpace == ERigPhysicsSimulationSpace::World)
+	{
+		// Danny TODO This is probably redundant unless we support runtime switching of the space
+		InitSimulationSpace(SimulationSpaceState.ComponentTM, SimulationSpaceState.BoneRelComponentTM);
+		SimulationSpaceData.LinearVelocity = SimulationSpaceSettings.ExternalLinearVelocity;
+		SimulationSpaceData.AngularVelocity = SimulationSpaceSettings.ExternalAngularVelocity;
+		return SimulationSpaceData;
+	}
+
+	// If the timestep is zero, then it doesn't actually matter what the velocity is - but make sure
+	// it doesn't corrupt anything.
+	if (SimulationSpaceState.Dt < SMALL_NUMBER || SimulationSpaceState.PrevDt < SMALL_NUMBER)
+	{
+		SimulationSpaceData.LinearVelocity = FVector::ZeroVector;
+		SimulationSpaceData.AngularVelocity = FVector::ZeroVector;
+		SimulationSpaceData.LinearAcceleration = FVector::ZeroVector;
+		SimulationSpaceData.AngularAcceleration = FVector::ZeroVector;
+		return SimulationSpaceData;
+	}
+
+	// We calculate velocities etc in world space first, and then subsequently convert them into
+	// simulation space.
+
+	// Note that the velocity/accel calculations are intended to track the world/simulation behavior
+	// - not necessarily be the most accurate calculations! For example, we could use one-sided
+	// finite difference approximations, but this wouldn't necessarily be correct.
+
+	// World-space component linear velocity and acceleration
+	const FVector PrevSpaceLinearVel = Chaos::FVec3::CalculateVelocity(
+		SimulationSpaceState.PrevPrevSimulationSpaceTM.GetTranslation(),
+		SimulationSpaceState.PrevSimulationSpaceTM.GetTranslation(), SimulationSpaceState.PrevDt);
+	SimulationSpaceData.LinearVelocity = Chaos::FVec3::CalculateVelocity(
+		SimulationSpaceState.PrevSimulationSpaceTM.GetTranslation(),
+		SimulationSpaceState.SimulationSpaceTM.GetTranslation(), SimulationSpaceState.Dt);
+	SimulationSpaceData.LinearAcceleration = 
+		(SimulationSpaceData.LinearVelocity - PrevSpaceLinearVel) / SimulationSpaceState.Dt;
+
+	// World-space component angular velocity and acceleration
+	const FVector PrevSpaceAngularVel = CalculateAngularVelocity(
+		SimulationSpaceState.PrevPrevSimulationSpaceTM.GetRotation(),
+		SimulationSpaceState.PrevSimulationSpaceTM.GetRotation(), SimulationSpaceState.PrevDt);
+	SimulationSpaceData.AngularVelocity = CalculateAngularVelocity(
+		SimulationSpaceState.PrevSimulationSpaceTM.GetRotation(),
+		SimulationSpaceState.SimulationSpaceTM.GetRotation(), SimulationSpaceState.Dt);
+	SimulationSpaceData.AngularAcceleration = 
+		(SimulationSpaceData.AngularVelocity - PrevSpaceAngularVel) / SimulationSpaceState.Dt;
+
+	// Apply Z scale
+	SimulationSpaceData.LinearVelocity.Z *= SimulationSpaceSettings.VelocityScaleZ;
+	SimulationSpaceData.LinearAcceleration.Z *= SimulationSpaceSettings.VelocityScaleZ;
+
+	// Clamped world-space motion of the simulation space
+	if (
+		(SimulationSpaceSettings.LinearAccelerationThresholdForTeleport > 0 &&
+			SimulationSpaceData.LinearAcceleration.SquaredLength() >
+			FMath::Square(SimulationSpaceSettings.LinearAccelerationThresholdForTeleport)) ||
+		(SimulationSpaceSettings.AngularAccelerationThresholdForTeleport > 0 &&
+			SimulationSpaceData.AngularAcceleration.SquaredLength() >
+			FMath::Square(SimulationSpaceSettings.AngularAccelerationThresholdForTeleport))
+		)
+	{
+		SimulationSpaceData.LinearVelocity = FVector::ZeroVector;
+		SimulationSpaceData.AngularVelocity = FVector::ZeroVector;
+		SimulationSpaceData.LinearAcceleration = FVector::ZeroVector;
+		SimulationSpaceData.AngularAcceleration = FVector::ZeroVector;
+	}
+	else
+	{
+		if (SimulationSpaceSettings.bClampLinearVelocity)
+		{
+			SimulationSpaceData.LinearVelocity =
+				SimulationSpaceData.LinearVelocity.GetClampedToMaxSize(SimulationSpaceSettings.MaxLinearVelocity);
+		}
+		if (SimulationSpaceSettings.bClampAngularVelocity)
+		{
+			SimulationSpaceData.AngularVelocity =
+				SimulationSpaceData.AngularVelocity.GetClampedToMaxSize(SimulationSpaceSettings.MaxAngularVelocity);
+		}
+		if (SimulationSpaceSettings.bClampLinearAcceleration)
+		{
+			SimulationSpaceData.LinearAcceleration =
+				SimulationSpaceData.LinearAcceleration.GetClampedToMaxSize(SimulationSpaceSettings.MaxLinearAcceleration);
+		}
+		if (SimulationSpaceSettings.bClampAngularAcceleration)
+		{
+			SimulationSpaceData.AngularAcceleration =
+				SimulationSpaceData.AngularAcceleration.GetClampedToMaxSize(SimulationSpaceSettings.MaxAngularAcceleration);
+		}
+	}
+
+	SimulationSpaceData.LinearVelocity += SimulationSpaceSettings.ExternalLinearVelocity;
+	SimulationSpaceData.AngularVelocity += SimulationSpaceSettings.ExternalAngularVelocity;
+
+	// Transform world-space motion into simulation space Danny TODO note that this matches the code
+	// in RBAN, and is doing what the interface requires (i.e. movement of the space in the space of
+	// the space!). However, the resultant motion is incorrect for rotation of the space.
+	SimulationSpaceData.LinearVelocity =
+		SimulationSpaceState.SimulationSpaceTM.InverseTransformVector(SimulationSpaceData.LinearVelocity);
+	SimulationSpaceData.AngularVelocity =
+		SimulationSpaceState.SimulationSpaceTM.InverseTransformVector(SimulationSpaceData.AngularVelocity);
+	SimulationSpaceData.LinearAcceleration =
+		SimulationSpaceState.SimulationSpaceTM.InverseTransformVector(SimulationSpaceData.LinearAcceleration);
+	SimulationSpaceData.AngularAcceleration =
+		SimulationSpaceState.SimulationSpaceTM.InverseTransformVector(SimulationSpaceData.AngularAcceleration);
+
+	return SimulationSpaceData;
+}
diff --git a/Engine/Plugins/Experimental/PhysicsControl/Source/ControlRigPhysics/Private/RigPhysicsSimulationVisualization.cpp b/Engine/Plugins/Experimental/PhysicsControl/Source/ControlRigPhysics/Private/RigPhysicsSimulationVisualization.cpp
new file mode 100644
--- /dev/null
+++ b/Engine/Plugins/Experimental/PhysicsControl/Source/ControlRigPhysics/Private/RigPhysicsSimulationVisualization.cpp
@@ -0,0 +1,349 @@
+// Copyright Epic Games, Inc. All Rights Reserved.
+
+#include "RigPhysicsSimulation.h"
+
+#include "Rigs/RigHierarchy.h"
+#include "Rigs/RigHierarchyController.h"
+#include "ControlRig.h"
+
+#include "RigVMCore/RigVMDrawInterface.h"
+
+#include "Physics/ImmediatePhysics/ImmediatePhysicsActorHandle.h"
+#include "Physics/ImmediatePhysics/ImmediatePhysicsJointHandle.h"
+#include "Physics/ImmediatePhysics/ImmediatePhysicsSimulation.h"
+
+#include "Physics/Experimental/PhysScene_Chaos.h"
+#include "Chaos/ParticleHandle.h"
+#include "Chaos/PBDJointConstraints.h"
+#include "Chaos/ImplicitObject.h"
+#include "Chaos/ShapeInstance.h"
+#include "Chaos/Capsule.h"
+#include "Chaos/PBDJointConstraintUtilities.h"
+
+#include "Engine/Engine.h"
+#include "Materials/MaterialInstanceDynamic.h"
+
+#include "Math/Color.h"
+
+namespace RigPhysicsSolverDraw
+{
+	FLinearColor DynamicColor = FColor::Yellow;
+	FLinearColor KinematicColor = FColor::Blue;
+}
+
+//======================================================================================================================
+static void DrawShapes(
+	FRigVMDrawInterface*                    DI,
+	const FRigPhysicsVisualizationSettings& VisualizationSettings,
+	const bool                              bIsKinematic,
+	const Chaos::FGeometryParticleHandle*   Particle,
+	const FTransform&                       ShapeTransform, 
+	const Chaos::FImplicitObject*           ImplicitObject,
+	const Chaos::FPerShapeData*             Shape)
+{
+	using namespace Chaos;
+
+	if (!ImplicitObject)
+	{
+		return;
+	}
+	const EImplicitObjectType PackedType = ImplicitObject->GetType(); // Type includes scaling and instancing data
+	const EImplicitObjectType InnerType = GetInnerType(ImplicitObject->GetType());
+
+	// For simplicity, we're going to assume no scaling, instancing, transforming etc
+
+	if (!ImplicitObject)
+	{
+		return;
+	}
+
+	switch (InnerType)
+	{
+	case ImplicitObjectType::Transformed:
+	{
+		const TImplicitObjectTransformed<FReal, 3>* Transformed = 
+			ImplicitObject->GetObject<TImplicitObjectTransformed<FReal, 3>>();
+		FTransform TransformedTransform = FTransform(
+			ShapeTransform.GetRotation() * Transformed->GetTransform().GetRotation(),
+			ShapeTransform.TransformPosition(Transformed->GetTransform().GetLocation()));
+		DrawShapes(
+			DI, VisualizationSettings, bIsKinematic, Particle, 
+			TransformedTransform, Transformed->GetTransformedObject(), Shape);
+		return;
+	}
+	case ImplicitObjectType::Sphere:
+	{
+		const Chaos::FSphere* Sphere = ImplicitObject->GetObject<Chaos::FSphere>();
+		DI->DrawSphere(ShapeTransform, FTransform(FVec3(Sphere->GetCenterOfMass())), Sphere->GetRadiusf(),
+			bIsKinematic ? RigPhysicsSolverDraw::KinematicColor : RigPhysicsSolverDraw::DynamicColor,
+			VisualizationSettings.LineThickness, VisualizationSettings.ShapeDetail);
+		break;
+	}
+	case ImplicitObjectType::Box:
+	{
+		const TBox<FReal, 3>* Box = ImplicitObject->GetObject<TBox<FReal, 3>>();
+		DI->DrawBox(ShapeTransform, FTransform(FQuat::Identity, Box->GetCenter(), Box->Extents()),
+			bIsKinematic ? RigPhysicsSolverDraw::KinematicColor : RigPhysicsSolverDraw::DynamicColor,
+			VisualizationSettings.LineThickness);
+		break;
+	}
+	case ImplicitObjectType::Capsule:
+	{
+		const Chaos::FCapsule* Capsule = ImplicitObject->GetObject<Chaos::FCapsule>();
+		const FQuat Q = FRotationMatrix::MakeFromZ(Capsule->GetAxis()).ToQuat();
+		DI->DrawCapsule(ShapeTransform, FTransform(Q, Capsule->GetCenterOfMass()),
+			Capsule->GetRadiusf(), Capsule->GetHeightf(),
+			bIsKinematic ? RigPhysicsSolverDraw::KinematicColor : RigPhysicsSolverDraw::DynamicColor,
+			VisualizationSettings.LineThickness);
+		break;
+	}
+	default:
+		// If we don't know what it is, don't draw it.
+		break;
+	}
+
+}
+
+//======================================================================================================================
+static void DrawActor(
+	FRigVMDrawInterface*                    DI,
+	const FRigPhysicsVisualizationSettings& VisualizationSettings,
+	const FTransform&                       SpaceTransform,
+	const ImmediatePhysics::FActorHandle&   ActorHandle)
+{
+	bool bIsKinematic = ActorHandle.GetIsKinematic();
+	const Chaos::FGeometryParticleHandle* GeometryParticleHandle = ActorHandle.GetParticle();
+	const FTransform ParticleTransform = GeometryParticleHandle->GetTransformXR() * SpaceTransform;
+	for (const Chaos::FShapeInstancePtr& ShapeInstance : GeometryParticleHandle->ShapeInstances())
+	{
+		const Chaos::FImplicitObject* ImplicitObject = ShapeInstance->GetGeometry();
+		DrawShapes(
+			DI, VisualizationSettings, bIsKinematic, GeometryParticleHandle,
+			ParticleTransform, ImplicitObject, ShapeInstance.Get());
+	}
+}
+
+enum
+{
+	TwistIndex = 0,
+	Swing1Index = 1,
+	Swing2Index = 2
+};
+
+// Returns a number to use for the limit for limit based on limit type
+float GetLimitAngleRadians(float LimitAngle, Chaos::EJointMotionType LimitType)
+{
+	switch (LimitType)
+	{
+	case Chaos::EJointMotionType::Free: return UE_PI;
+	case Chaos::EJointMotionType::Locked: return 0.f;
+	default: return LimitAngle;
+	}
+}
+
+//======================================================================================================================
+static void DrawJoint(
+	FRigVMDrawInterface*                    DI,
+	const FRigPhysicsVisualizationSettings& VisualizationSettings,
+	const FTransform&                       SpaceTransform,
+	const ImmediatePhysics::FJointHandle&   JointHandle)
+{
+	const Chaos::FPBDJointConstraintHandle* ConstraintHandle = JointHandle.GetConstraint();
+	if (!ConstraintHandle)
+	{
+		return;
+	}
+
+	const Chaos::TVec2<const ImmediatePhysics_Chaos::FActorHandle*>& ActorHandles = JointHandle.GetActorHandles();
+	const ImmediatePhysics_Chaos::FActorHandle* ChildActor = ActorHandles[0];
+	const ImmediatePhysics_Chaos::FActorHandle* ParentActor = ActorHandles[1];
+
+	FTransform ChildActorTM = (ChildActor ? ChildActor->GetWorldTransform() : FTransform()) * SpaceTransform;
+	FTransform ParentActorTM = (ParentActor ? ParentActor->GetWorldTransform() : FTransform()) * SpaceTransform;
+
+	bool bChildIsKinematic = ChildActor ? ChildActor->GetIsKinematic() : true;
+	bool bParentIsKinematic = ParentActor ? ParentActor->GetIsKinematic() : true;
+
+	float Size = 5.0f;
+
+	const Chaos::FPBDJointSettings& JointSettings = ConstraintHandle->GetSettings();
+
+	FTransform DialFrame = JointSettings.ConnectorTransforms[0] * ChildActorTM;
+	FTransform LimitFrame = JointSettings.ConnectorTransforms[1] * ParentActorTM;
+
+	static bool bDrawAsAxes = false;
+	if (bDrawAsAxes)
+	{
+		DI->DrawAxes(LimitFrame, FTransform(), Size, VisualizationSettings.LineThickness * 2);
+		DI->DrawAxes(DialFrame, FTransform(), Size, VisualizationSettings.LineThickness);
+	}
+
+	{
+		// See FConstraintInstance::DrawConstraintImp for inspiration
+
+		// There seems to be a swap between swing1 and swing2 compared to GetSwingTwistAngles
+		FVector LimitAngleRadians(
+			GetLimitAngleRadians(JointSettings.AngularLimits[0], JointSettings.AngularMotionTypes[0]),
+			GetLimitAngleRadians(JointSettings.AngularLimits[2], JointSettings.AngularMotionTypes[2]),
+			GetLimitAngleRadians(JointSettings.AngularLimits[1], JointSettings.AngularMotionTypes[1])
+		);
+
+		Chaos::FReal TwistAngle, Swing1Angle, Swing2Angle;
+		const FQuat LimitQ = LimitFrame.GetRotation();
+		FQuat DialQ = DialFrame.GetRotation();
+		DialQ.EnforceShortestArcWith(LimitQ);
+		Chaos::FPBDJointUtilities::GetSwingTwistAngles(LimitQ, DialQ, TwistAngle, Swing1Angle, Swing2Angle);
+
+		bool bDrawViolatedLimits = true;
+
+		bool bTwistViolated = false;
+		bool bSwing1Violated = false;
+		bool bSwing2Violated = false;
+
+		if (bDrawViolatedLimits)
+		{
+			bTwistViolated = JointSettings.AngularMotionTypes[TwistIndex] == Chaos::EJointMotionType::Limited &&
+				FMath::Abs(TwistAngle) > LimitAngleRadians[0];
+
+			bSwing1Violated = JointSettings.AngularMotionTypes[Swing1Index] == Chaos::EJointMotionType::Limited &&
+				FMath::Abs(Swing1Angle) > LimitAngleRadians[1];
+
+			bSwing2Violated = JointSettings.AngularMotionTypes[Swing2Index] == Chaos::EJointMotionType::Limited &&
+				FMath::Abs(Swing2Angle) > LimitAngleRadians[2];
+		}
+
+		const bool bLockSwing1 = JointSettings.AngularMotionTypes[Swing1Index] == Chaos::EJointMotionType::Locked;
+		const bool bLockSwing2 = JointSettings.AngularMotionTypes[Swing2Index] == Chaos::EJointMotionType::Locked;
+		const bool bLockAllSwing = bLockSwing1 && bLockSwing2;
+
+		check(GEngine->ConstraintLimitMaterialX && GEngine->ConstraintLimitMaterialY && GEngine->ConstraintLimitMaterialZ);
+		static UMaterialInterface* LimitMaterialX = GEngine->ConstraintLimitMaterialX;
+		static UMaterialInterface* LimitMaterialXAxis = GEngine->ConstraintLimitMaterialXAxis;
+		static UMaterialInterface* LimitMaterialY = GEngine->ConstraintLimitMaterialY;
+		static UMaterialInterface* LimitMaterialYAxis = GEngine->ConstraintLimitMaterialYAxis;
+		static UMaterialInterface* LimitMaterialZ = GEngine->ConstraintLimitMaterialZ;
+		static UMaterialInterface* LimitMaterialZAxis = GEngine->ConstraintLimitMaterialZAxis;
+
+		float ArrowLength = Size * 1.05f; // stick out a little bit 
+
+		// If swing is limited (but not locked) - draw the swing limit cone.
+		if (!bLockAllSwing)
+		{
+			if (JointSettings.AngularMotionTypes[Swing1Index] == Chaos::EJointMotionType::Free && 
+				JointSettings.AngularMotionTypes[Swing2Index] == Chaos::EJointMotionType::Free)
+			{
+				DI->DrawSphere(LimitFrame, FTransform(), Size * 0.2f, FColor::White, 
+					VisualizationSettings.LineThickness, VisualizationSettings.ShapeDetail);
+			}
+			else
+			{
+				FTransform ConeLimitTM = LimitFrame;
+				ConeLimitTM.SetScale3D(FVector(Size));
+				const float Swing1LimitAngle = LimitAngleRadians[Swing1Index];
+				const float Swing2LimitAngle = LimitAngleRadians[Swing2Index];
+				DI->DrawCone(
+					ConeLimitTM, FTransform(), Swing1LimitAngle, Swing2LimitAngle, VisualizationSettings.ShapeDetail, 
+					true, FColor::Green, LimitMaterialX->GetRenderProxy(), VisualizationSettings.LineThickness);
+			}
+			// Draw the swing Dial indicator - shows the current orientation of the child frame
+			// relative to the parent frame on the swing axis.
+			// Start the arrow at the limit position as it can be confusing if there is joint separation
+			DI->DrawArrow(FTransform(LimitFrame.GetTranslation()), 
+				DialFrame.GetUnitAxis(EAxis::X) * ArrowLength,
+				DialFrame.GetUnitAxis(EAxis::Y), FColor::Red, VisualizationSettings.LineThickness);
+
+			if (bSwing1Violated || bSwing2Violated)
+			{
+				// Drawing a sphere seems a bit heavy... but it appears that if DrawPoint is used then they get
+				// culled when you get close!
+				DI->DrawSphere(FTransform(
+					DialFrame.GetRotation(),
+					LimitFrame.GetTranslation() + DialFrame.GetUnitAxis(EAxis::X) * ArrowLength),
+					FTransform(),
+					ArrowLength * 0.01f, FColor::Orange, VisualizationSettings.LineThickness * 4);
+			}
+
+		}
+
+		// Draw the twist limit
+		if (JointSettings.AngularMotionTypes[TwistIndex] != Chaos::EJointMotionType::Locked)
+		{
+			// Draw as a flat cone
+			FTransform ConeLimitTM = LimitFrame;
+			ConeLimitTM.SetScale3D(FVector(Size));
+			const float TwistLimitAngle = LimitAngleRadians[TwistIndex];
+			DI->DrawCone(ConeLimitTM, FTransform(FQuat::MakeFromRotationVector(FVector(0, 1, 0) * -HALF_PI)),
+				TwistLimitAngle, 0.0f, VisualizationSettings.ShapeDetail, true, FColor::Green,
+				LimitMaterialYAxis->GetRenderProxy(), VisualizationSettings.LineThickness);
+
+			// Draw the twist Dial indicator - shows the current orientation of the child frame
+			// relative to the parent frame on the twist axis.
+			FQuat Rot(LimitFrame.GetUnitAxis(EAxis::X), TwistAngle);
+			FVector TwistArrow = Rot * LimitFrame.GetUnitAxis(EAxis::Z);
+			FVector TwistSideDir = FVector::CrossProduct(LimitFrame.GetUnitAxis(EAxis::X), TwistArrow);
+
+			DI->DrawArrow(FTransform(LimitFrame.GetTranslation()),
+				TwistArrow * ArrowLength,
+				TwistSideDir, FColor::Blue, VisualizationSettings.LineThickness);
+
+			if (bTwistViolated)
+			{
+				DI->DrawSphere(FTransform(
+					DialFrame.GetRotation(),
+					LimitFrame.GetTranslation() + TwistArrow * ArrowLength),
+					FTransform(),
+					ArrowLength * 0.01f, FColor::Orange, VisualizationSettings.LineThickness * 4);
+			}
+		}
+	}
+}
+
+//======================================================================================================================
+void FRigPhysicsSimulation::Draw(
+	FRigVMDrawInterface*                    DI,
+	const FRigPhysicsSolverSettings&        SolverSettings,
+	const FRigPhysicsVisualizationSettings& VisualizationSettings,
+	const UWorld*                           DebugWorld) const
+{
+#if CHAOS_DEBUG_DRAW
+	// Danny TODO enable/disable chaos debug draw with a cvar or flag
+	if (DebugWorld && DebugWorld->GetPhysicsScene() && DebugWorld->GetPhysicsScene()->GetDebugDrawScene())
+	{
+		Simulation->SetDebugDrawScene(FString("ControlRig"), DebugWorld->GetPhysicsScene()->GetDebugDrawScene());
+		Simulation->DebugDraw();
+	}
+#endif
+
+	const URigHierarchy* Hierarchy = OwningControlRig->GetHierarchy();
+	if (!DI || !Hierarchy)
+	{
+		return;
+	}
+
+	// All rendering is done relative to the component, so convert the sim space (identity) into the
+	// component space.
+	const FTransform SpaceTransform = ConvertSimSpaceTransformToComponentSpace(SolverSettings, FTransform());
+
+	for (const ImmediatePhysics::FActorHandle* ActorHandle : AdditionalActorHandles)
+	{
+		DrawActor(DI, VisualizationSettings, SpaceTransform, *ActorHandle);
+	}
+
+	for (const TPair<FRigComponentKey, FRigBodyRecord>& BodyRecordPair : BodyRecords)
+	{
+		const FRigBodyRecord& Record = BodyRecordPair.Value;
+		const ImmediatePhysics::FActorHandle* ActorHandle = Record.ActorHandle;
+		const ImmediatePhysics::FJointHandle* JointHandle = Record.JointHandle;
+
+		if (JointHandle)
+		{
+			DrawJoint(DI, VisualizationSettings, SpaceTransform, *JointHandle);
+		}
+
+		if (ActorHandle)
+		{
+			DrawActor(DI, VisualizationSettings, SpaceTransform, *ActorHandle);
+		}
+	}
+}
+
diff --git a/Engine/Plugins/Experimental/PhysicsControl/Source/ControlRigPhysics/Private/RigPhysicsSolver.cpp b/Engine/Plugins/Experimental/PhysicsControl/Source/ControlRigPhysics/Private/RigPhysicsSolver.cpp
deleted file mode 100644
diff --git a/Engine/Plugins/Experimental/PhysicsControl/Source/ControlRigPhysics/Private/RigPhysicsSolver.h b/Engine/Plugins/Experimental/PhysicsControl/Source/ControlRigPhysics/Private/RigPhysicsSolver.h
deleted file mode 100644
diff --git a/Engine/Plugins/Experimental/PhysicsControl/Source/ControlRigPhysics/Private/RigPhysicsSolverComponent.cpp b/Engine/Plugins/Experimental/PhysicsControl/Source/ControlRigPhysics/Private/RigPhysicsSolverComponent.cpp
new file mode 100644
--- /dev/null
+++ b/Engine/Plugins/Experimental/PhysicsControl/Source/ControlRigPhysics/Private/RigPhysicsSolverComponent.cpp
@@ -0,0 +1,26 @@
+// Copyright Epic Games, Inc. All Rights Reserved.
+
+#include "RigPhysicsSolverComponent.h"
+
+//======================================================================================================================
+void FRigPhysicsSolverComponent::Save(FArchive& Ar)
+{
+	FRigBaseComponent::Save(Ar);
+	Ar << SolverSettings;
+	Ar << SimulationSpaceSettings;
+}
+
+//======================================================================================================================
+void FRigPhysicsSolverComponent::Load(FArchive& Ar)
+{
+	FRigBaseComponent::Load(Ar);
+	Ar << SolverSettings;
+	Ar << SimulationSpaceSettings;
+}
+
+//======================================================================================================================
+bool FRigPhysicsSolverComponent::CanBeAddedTo(
+	const FRigElementKey& InElementKey, const URigHierarchy* InHierarchy) const
+{
+	return InElementKey.IsTypeOf(ERigElementType::Bone);
+}
diff --git a/Engine/Plugins/Experimental/PhysicsControl/Source/ControlRigPhysics/Private/RigPhysicsSolverSimulationSpace.cpp b/Engine/Plugins/Experimental/PhysicsControl/Source/ControlRigPhysics/Private/RigPhysicsSolverSimulationSpace.cpp
deleted file mode 100644
diff --git a/Engine/Plugins/Experimental/PhysicsControl/Source/ControlRigPhysics/Private/RigPhysicsSolverVisualization.cpp b/Engine/Plugins/Experimental/PhysicsControl/Source/ControlRigPhysics/Private/RigPhysicsSolverVisualization.cpp
deleted file mode 100644
diff --git a/Engine/Plugins/Experimental/PhysicsControl/Source/ControlRigPhysics/Public/ControlRigPhysicsModule.h b/Engine/Plugins/Experimental/PhysicsControl/Source/ControlRigPhysics/Public/ControlRigPhysicsModule.h
--- a/Engine/Plugins/Experimental/PhysicsControl/Source/ControlRigPhysics/Public/ControlRigPhysicsModule.h
+++ b/Engine/Plugins/Experimental/PhysicsControl/Source/ControlRigPhysics/Public/ControlRigPhysicsModule.h
@@ -2,8 +2,15 @@
 
 #pragma once
 
+#include "CoreMinimal.h"
 #include "Modules/ModuleManager.h"
 
+#if UE_BUILD_SHIPPING || UE_BUILD_TEST
+DECLARE_LOG_CATEGORY_EXTERN(LogControlRigPhysics, Log, Warning);
+#else
+DECLARE_LOG_CATEGORY_EXTERN(LogControlRigPhysics, Log, All);
+#endif
+
 class FControlRigPhysicsModule : public IModuleInterface
 {
 public:
diff --git a/Engine/Plugins/Experimental/PhysicsControl/Source/ControlRigPhysics/Public/RigPhysicsComponent.h b/Engine/Plugins/Experimental/PhysicsControl/Source/ControlRigPhysics/Public/RigPhysicsComponent.h
--- a/Engine/Plugins/Experimental/PhysicsControl/Source/ControlRigPhysics/Public/RigPhysicsComponent.h
+++ b/Engine/Plugins/Experimental/PhysicsControl/Source/ControlRigPhysics/Public/RigPhysicsComponent.h
@@ -11,12 +11,6 @@
 
 #include "RigPhysicsComponent.generated.h"
 
-struct FRigVMExecuteContext;
-struct FRigBaseElement;
-struct FRigControlElement;
-class URigHierarchy;
-class FRigElementKeyRedirector;
-
 /**
  * A component that can be added to a joint/element that defines how a physical body can be "attached" to it.
  * The body supports dynamic movement, collision, and a physics joint with this body's parent in the hierarchy.
@@ -29,6 +23,9 @@ public:
 	GENERATED_BODY()
 	DECLARE_RIG_COMPONENT_METHODS(FRigPhysicsComponent)
 
+	//==================================================================================================================
+	// Initial data
+	//==================================================================================================================
 	UPROPERTY(BlueprintReadOnly, EditAnywhere, Category = Physics, meta = (ShowOnlyInnerProperties))
 	FRigPhysicsBodySolverSettings BodySolverSettings;
 
@@ -44,7 +41,22 @@ public:
 	UPROPERTY(BlueprintReadOnly, EditAnywhere, Category = Physics)
 	FPhysicsControlModifierData BodyData;
 
-	virtual void Save(FArchive& A) override;
+	//==================================================================================================================
+	// Current data
+	//==================================================================================================================
+
+	// The target for when this body is kinematic
+	UPROPERTY(BlueprintReadOnly, EditAnywhere, Category = CurrentPhysics)
+	FTransform CurrentKinematicTarget;
+
+	UPROPERTY(BlueprintReadOnly, EditAnywhere, Category = CurrentPhysics)
+	ERigPhysicsKinematicTargetSpace CurrentKinematicTargetSpace = ERigPhysicsKinematicTargetSpace::OffsetInBoneSpace;
+
+	// Desired state of the physics body - this will be written to the simulation during UpdatePrePhysics
+	UPROPERTY(BlueprintReadOnly, EditAnywhere, Category = CurrentPhysics)
+	FPhysicsControlModifierData CurrentBodyData;
+
+	virtual void Save(FArchive& Ar) override;
 	virtual void Load(FArchive& Ar) override;
 };
 
diff --git a/Engine/Plugins/Experimental/PhysicsControl/Source/ControlRigPhysics/Public/RigPhysicsControlComponent.h b/Engine/Plugins/Experimental/PhysicsControl/Source/ControlRigPhysics/Public/RigPhysicsControlComponent.h
--- a/Engine/Plugins/Experimental/PhysicsControl/Source/ControlRigPhysics/Public/RigPhysicsControlComponent.h
+++ b/Engine/Plugins/Experimental/PhysicsControl/Source/ControlRigPhysics/Public/RigPhysicsControlComponent.h
@@ -20,7 +20,7 @@ public:
 	DECLARE_RIG_COMPONENT_METHODS(FRigPhysicsControlComponent)
 
 	UPROPERTY(BlueprintReadOnly, EditAnywhere, Category = Control)
-	FRigComponentKey ParentBody;
+	FRigComponentKey ParentBodyComponentKey;
 
 	/** Describes the initial strength etc of the new control */
 	UPROPERTY(BlueprintReadWrite, EditAnywhere, Category = Control)
diff --git a/Engine/Plugins/Experimental/PhysicsControl/Source/ControlRigPhysics/Public/RigPhysicsControlExecution.h b/Engine/Plugins/Experimental/PhysicsControl/Source/ControlRigPhysics/Public/RigPhysicsControlExecution.h
--- a/Engine/Plugins/Experimental/PhysicsControl/Source/ControlRigPhysics/Public/RigPhysicsControlExecution.h
+++ b/Engine/Plugins/Experimental/PhysicsControl/Source/ControlRigPhysics/Public/RigPhysicsControlExecution.h
@@ -3,7 +3,7 @@
 #pragma once
 
 #include "RigPhysicsComponent.h"
-#include "RigPhysicsSolver.h"
+#include "RigPhysicsSimulation.h"
 #include "PhysicsControlData.h"
 #include "RigPhysicsExecution.h"
 
@@ -20,12 +20,6 @@ struct CONTROLRIGPHYSICS_API FRigUnit_AddPhysicsControl : public FRigUnitMutable
 {
 	GENERATED_BODY()
 
-	FRigUnit_AddPhysicsControl()
-	{
-		Owner = FRigElementKey();
-		Result = FRigComponentKey();
-	}
-
 	RIGVM_METHOD()
 	virtual void Execute() override;
 
@@ -36,16 +30,16 @@ struct CONTROLRIGPHYSICS_API FRigUnit_AddPhysicsControl : public FRigUnitMutable
 	FRigElementKey Owner;
 
 	UPROPERTY(meta = (Output))
-	FRigComponentKey Result;
+	FRigComponentKey ControlComponentKey;
 
 	/*
 	 * The solver to relate this new physics element to
 	 */
 	UPROPERTY(meta = (Input))
-	FRigPhysicsSolverID SolverID;
+	FRigComponentKey PhysicsSolverComponentKey;
 
 	UPROPERTY(meta = (Input))
-	FRigComponentKey ParentBody;
+	FRigComponentKey ParentBodyComponentKey;
 
 	/** Describes the initial strength etc of the new control */
 	UPROPERTY(meta = (Input))
diff --git a/Engine/Plugins/Experimental/PhysicsControl/Source/ControlRigPhysics/Public/RigPhysicsData.h b/Engine/Plugins/Experimental/PhysicsControl/Source/ControlRigPhysics/Public/RigPhysicsData.h
--- a/Engine/Plugins/Experimental/PhysicsControl/Source/ControlRigPhysics/Public/RigPhysicsData.h
+++ b/Engine/Plugins/Experimental/PhysicsControl/Source/ControlRigPhysics/Public/RigPhysicsData.h
@@ -212,6 +212,26 @@ struct FRigPhysicsSimulationSpaceSettings
 {
 	GENERATED_BODY()
 
+	friend FArchive& operator <<(FArchive& Ar, FRigPhysicsSimulationSpaceSettings& Data)
+	{
+		Ar << Data.SpaceMovementAmount;
+		Ar << Data.VelocityScaleZ;
+		Ar << Data.bClampLinearVelocity;
+		Ar << Data.MaxLinearVelocity;
+		Ar << Data.bClampAngularVelocity;
+		Ar << Data.MaxAngularVelocity;
+		Ar << Data.bClampLinearAcceleration;
+		Ar << Data.MaxLinearAcceleration;
+		Ar << Data.bClampAngularAcceleration;
+		Ar << Data.MaxAngularAcceleration;
+		Ar << Data.LinearAccelerationThresholdForTeleport;
+		Ar << Data.AngularAccelerationThresholdForTeleport;
+		Ar << Data.ExternalLinearDragV;
+		Ar << Data.ExternalLinearVelocity;
+		Ar << Data.ExternalAngularVelocity;
+		return Ar;
+	}
+
 	// Overall multiplier on the effects of simulation space movement on the simulation
 	UPROPERTY(EditAnywhere, BlueprintReadWrite, Category = Settings, meta = (ClampMin = "0.0", ClampMax = "1.0"))
 	float SpaceMovementAmount = 1.0f;
@@ -313,6 +333,28 @@ struct CONTROLRIGPHYSICS_API FRigPhysicsSolverSettings
 {
 	GENERATED_BODY()
 
+	friend FArchive& operator <<(FArchive& Ar, FRigPhysicsSolverSettings& Data)
+	{
+		Ar << Data.SimulationSpace;
+		Ar << Data.CollisionSpace;
+		Ar << Data.SpaceBone;
+		Ar << Data.Collision;
+		Ar << Data.Gravity;
+		Ar << Data.PositionIterations;
+		Ar << Data.VelocityIterations;
+		Ar << Data.ProjectionIterations;
+		Ar << Data.MaxNumRollingAverageStepTimes;
+		Ar << Data.MaxDepenetrationVelocity;
+		Ar << Data.FixedTimeStep;
+		Ar << Data.MaxTimeSteps;
+		Ar << Data.MaxDeltaTime;
+		Ar << Data.bUseLinearJointSolver;
+		Ar << Data.bSolveJointPositionsLast;
+		Ar << Data.bUseManifolds;
+		Ar << Data.DistanceThresholdForReset;
+		return Ar;
+	}
+
 	UPROPERTY(BlueprintReadOnly, EditAnywhere, Category = SolverSettings)
 	ERigPhysicsSimulationSpace SimulationSpace = ERigPhysicsSimulationSpace::Component;
 
@@ -636,19 +678,19 @@ struct CONTROLRIGPHYSICS_API FRigPhysicsBodySolverSettings
 	GENERATED_BODY()
 
 	FRigPhysicsBodySolverSettings(
-		FRigPhysicsSolverID InSolverID = FRigPhysicsSolverID(), FRigElementKey InTargetBone = FRigElementKey())
-		: SolverID(InSolverID), TargetBone(InTargetBone) {}
+		FRigComponentKey InPhysicsSolverComponentKey = FRigComponentKey(), FRigElementKey InTargetBone = FRigElementKey())
+		: PhysicsSolverComponentKey(InPhysicsSolverComponentKey), TargetBone(InTargetBone) {}
 
 	friend FArchive& operator <<(FArchive& Ar, FRigPhysicsBodySolverSettings& Data)
 	{
-		Ar << Data.SolverID;
+		Ar << Data.PhysicsSolverComponentKey;
 		Ar << Data.TargetBone;
 		Ar << Data.SourceBone;
 		return Ar;
 	}
 
 	UPROPERTY(BlueprintReadWrite, EditAnywhere, Category = Solver)
-	FRigPhysicsSolverID SolverID;
+	FRigComponentKey PhysicsSolverComponentKey;
 
 	/// The bone that is used to initialize physics, as well as what to track when the body is set to be kinematic.
 	/// Note that if this is unset, the it will default to the parent of our owner.
diff --git a/Engine/Plugins/Experimental/PhysicsControl/Source/ControlRigPhysics/Public/RigPhysicsExecution.h b/Engine/Plugins/Experimental/PhysicsControl/Source/ControlRigPhysics/Public/RigPhysicsExecution.h
--- a/Engine/Plugins/Experimental/PhysicsControl/Source/ControlRigPhysics/Public/RigPhysicsExecution.h
+++ b/Engine/Plugins/Experimental/PhysicsControl/Source/ControlRigPhysics/Public/RigPhysicsExecution.h
@@ -3,7 +3,7 @@
 #pragma once
 
 #include "RigPhysicsComponent.h"
-#include "RigPhysicsSolver.h"
+#include "RigPhysicsSimulation.h"
 #include "PhysicsControlData.h"
 
 #include "Units/Execution/RigUnit_DynamicHierarchy.h"
@@ -11,40 +11,31 @@
 #include "RigPhysicsExecution.generated.h"
 
 /**
- * Adds a new physics solver
+ * Adds a new physics solver as a component on the owner element.
  * Note: This node only runs as part of the construction event.
  */
-USTRUCT(meta=(Category = "RigPhysics", NodeColor = "1.0 0.6 0.3", DisplayName="Spawn Physics Solver", Keywords="Construction,Create,New,Simulation,Physics", Varying))
+USTRUCT(meta=(Category = "RigPhysics", NodeColor = "1.0 0.6 0.3", DisplayName="Add Physics Solver Component", Keywords="Construction,Create,New,Simulation,Physics", Varying))
 struct CONTROLRIGPHYSICS_API FRigUnit_AddPhysicsSolver : public FRigUnitMutable
 {
 	GENERATED_BODY()
 
-	FRigUnit_AddPhysicsSolver()
-	{
-		Name = TEXT("Solver");
-		SolverID = FRigPhysicsSolverID();
-	}
+	RIGVM_METHOD()
+	virtual void Execute() override;
 
-	/*
-	 * The name of the new solver to spawn
+	/**
+	 * The owner of the newly created component (must be set/valid)
 	 */
-	UPROPERTY(meta = (Input))
-	FName Name;
+	UPROPERTY(meta = (Input, BoneName))
+	FRigElementKey Owner;
+
+	UPROPERTY(meta = (Output))
+	FRigComponentKey PhysicsSolverComponentKey;
 
 	UPROPERTY(meta = (Input))
 	FRigPhysicsSolverSettings SolverSettings;
 
 	UPROPERTY(meta = (Input))
 	FRigPhysicsSimulationSpaceSettings SimulationSpaceSettings;
-
-	/*
-	 * The identifier of the solver
-	 */
-	UPROPERTY(meta = (Output))
-	FRigPhysicsSolverID SolverID;
-
-	RIGVM_METHOD()
-	virtual void Execute() override;
 };
 
 /**
@@ -59,13 +50,13 @@ struct CONTROLRIGPHYSICS_API FRigUnit_InstantiatePhysics : public FRigUnitMutabl
 	GENERATED_BODY()
 
 	RIGVM_METHOD()
-		virtual void Execute() override;
+	virtual void Execute() override;
 
 	/*
 	 * The solver to relate this new physics element to
 	 */
 	UPROPERTY(meta = (Input))
-	FRigPhysicsSolverID SolverID;
+	FRigComponentKey PhysicsSolverComponentKey;
 };
 
 
@@ -82,7 +73,7 @@ struct CONTROLRIGPHYSICS_API FRigUnit_StepPhysicsSolver : public FRigUnitMutable
 
 	// The solver to relate this new physics element to
 	UPROPERTY(meta = (Input))
-	FRigPhysicsSolverID SolverID;
+	FRigComponentKey PhysicsSolverComponentKey;
 	
 	// If this is zero, then the execute context time will be used. If this is positive then it will
 	// override the delta time. A negative value will prevent the solver from stepping, but there will
@@ -108,15 +99,8 @@ struct CONTROLRIGPHYSICS_API FRigUnit_AddPhysicsBody : public FRigUnitMutable
 {
 	GENERATED_BODY()
 
-	FRigUnit_AddPhysicsBody()
-	{
-		Owner = FRigElementKey();
-		Result = FRigComponentKey();
-		SolverID = FRigPhysicsSolverID();
-	}
-
 	RIGVM_METHOD()
-		virtual void Execute() override;
+	virtual void Execute() override;
 
 	/**
 	 * The owner of the newly created component (must be set/valid)
@@ -125,13 +109,13 @@ struct CONTROLRIGPHYSICS_API FRigUnit_AddPhysicsBody : public FRigUnitMutable
 	FRigElementKey Owner;
 
 	UPROPERTY(meta = (Output))
-	FRigComponentKey Result;
+	FRigComponentKey PhysicsBodyComponentKey;
 
 	/*
 	 * The solver to relate this new physics element to
 	 */
 	UPROPERTY(meta = (Input))
-	FRigPhysicsSolverID SolverID;
+	FRigComponentKey PhysicsSolverComponentKey;
 
 	/*
 	 * The dynamics properties of the new physics element
@@ -171,7 +155,7 @@ struct CONTROLRIGPHYSICS_API FRigUnit_HierarchySetPhysicsBodySourceBone : public
 	virtual void Execute() override;
 
 	UPROPERTY(meta = (Input))
-	FRigComponentKey PhysicsBody;
+	FRigComponentKey PhysicsBodyComponentKey;
 
 	UPROPERTY(meta = (Input))
 	FRigElementKey SourceBone;
@@ -189,7 +173,7 @@ struct CONTROLRIGPHYSICS_API FRigUnit_HierarchySetPhysicsBodyTargetBone : public
 	virtual void Execute() override;
 
 	UPROPERTY(meta = (Input))
-	FRigComponentKey PhysicsBody;
+	FRigComponentKey PhysicsBodyComponentKey;
 
 	UPROPERTY(meta = (Input))
 	FRigElementKey TargetBone;
@@ -210,10 +194,13 @@ struct CONTROLRIGPHYSICS_API FRigUnit_HierarchySetPhysicsBodySparseData : public
 	virtual void Execute() override;
 
 	UPROPERTY(meta = (Input))
-	FRigComponentKey PhysicsBody;
+	FRigComponentKey PhysicsBodyComponentKey;
 
 	UPROPERTY(meta = (Input))
 	FPhysicsControlModifierSparseData Data;
+
+	UPROPERTY(meta = (Input))
+	bool bInitial = true;
 };
 
 /**
@@ -228,13 +215,16 @@ struct CONTROLRIGPHYSICS_API FRigUnit_HierarchySetPhysicsBodyKinematicTarget : p
 	virtual void Execute() override;
 
 	UPROPERTY(meta = (Input))
-	FRigComponentKey PhysicsBody;
+	FRigComponentKey PhysicsBodyComponentKey;
 
 	UPROPERTY(meta = (Input))
 	ERigPhysicsKinematicTargetSpace KinematicTargetSpace = ERigPhysicsKinematicTargetSpace::OffsetInBoneSpace;
 
 	UPROPERTY(meta = (Input))
 	FTransform KinematicTarget;
+
+	UPROPERTY(meta = (Input))
+	bool bInitial = true;
 };
 
 /**
@@ -250,10 +240,13 @@ struct CONTROLRIGPHYSICS_API FRigUnit_HierarchySetPhysicsBodyMovementType : publ
 	virtual void Execute() override;
 
 	UPROPERTY(meta = (Input))
-	FRigComponentKey PhysicsBody;
+	FRigComponentKey PhysicsBodyComponentKey;
 
 	UPROPERTY(meta = (Input))
 	EPhysicsMovementType MovementType = EPhysicsMovementType::Simulated;
+
+	UPROPERTY(meta = (Input))
+	bool bInitial = true;
 };
 
 /**
@@ -268,10 +261,13 @@ struct CONTROLRIGPHYSICS_API FRigUnit_HierarchySetPhysicsBodyCollisionType : pub
 	virtual void Execute() override;
 
 	UPROPERTY(meta = (Input))
-	FRigComponentKey PhysicsBody;
+	FRigComponentKey PhysicsBodyComponentKey;
 
 	UPROPERTY(meta = (Input))
 	TEnumAsByte<ECollisionEnabled::Type> CollisionType = ECollisionEnabled::QueryAndPhysics;
+
+	UPROPERTY(meta = (Input))
+	bool bInitial = true;
 };
 
 /**
@@ -286,10 +282,13 @@ struct CONTROLRIGPHYSICS_API FRigUnit_HierarchySetPhysicsBodyGravityMultiplier :
 	virtual void Execute() override;
 
 	UPROPERTY(meta = (Input))
-	FRigComponentKey PhysicsBody;
+	FRigComponentKey PhysicsBodyComponentKey;
 
 	UPROPERTY(meta = (Input))
 	float GravityMultiplier = 1.0f;
+
+	UPROPERTY(meta = (Input))
+	bool bInitial = true;
 };
 
 /**
@@ -305,10 +304,13 @@ struct CONTROLRIGPHYSICS_API FRigUnit_HierarchySetPhysicsBodyPhysicsBlendWeight
 	virtual void Execute() override;
 
 	UPROPERTY(meta = (Input))
-	FRigComponentKey PhysicsBody;
+	FRigComponentKey PhysicsBodyComponentKey;
 
 	UPROPERTY(meta = (Input))
 	float PhysicsBlendWeight = 1.0f;
+
+	UPROPERTY(meta = (Input))
+	bool bInitial = true;
 };
 
 /**
@@ -324,10 +326,13 @@ struct CONTROLRIGPHYSICS_API FRigUnit_HierarchySetPhysicsBodyUseSkeletalAnimatio
 	virtual void Execute() override;
 
 	UPROPERTY(meta = (Input))
-	FRigComponentKey PhysicsBody;
+	FRigComponentKey PhysicsBodyComponentKey;
 
 	UPROPERTY(meta = (Input))
 	bool bUseSkeletalAnimation = true;
+
+	UPROPERTY(meta = (Input))
+	bool bInitial = true;
 };
 
 /**
@@ -344,10 +349,13 @@ struct CONTROLRIGPHYSICS_API FRigUnit_HierarchySetPhysicsBodyUpdateKinematicFrom
 	virtual void Execute() override;
 
 	UPROPERTY(meta = (Input))
-	FRigComponentKey PhysicsBody;
+	FRigComponentKey PhysicsBodyComponentKey;
 
 	UPROPERTY(meta = (Input))
 	bool bUpdateKinematicFromSimulation = true;
+
+	UPROPERTY(meta = (Input))
+	bool bInitial = true;
 };
 
 
diff --git a/Engine/Plugins/Experimental/PhysicsControl/Source/ControlRigPhysics/Public/RigPhysicsSolverComponent.h b/Engine/Plugins/Experimental/PhysicsControl/Source/ControlRigPhysics/Public/RigPhysicsSolverComponent.h
new file mode 100644
--- /dev/null
+++ b/Engine/Plugins/Experimental/PhysicsControl/Source/ControlRigPhysics/Public/RigPhysicsSolverComponent.h
@@ -0,0 +1,33 @@
+// Copyright Epic Games, Inc. All Rights Reserved.
+
+#pragma once
+
+#include "RigPhysicsData.h"
+
+#include "Rigs/RigPhysics.h"
+#include "Rigs/RigHierarchyDefines.h"
+#include "Rigs/RigHierarchyComponents.h"
+
+#include "RigPhysicsSolverComponent.generated.h"
+
+// A solver coordinates the physical movement of bodies
+USTRUCT(BlueprintType)
+struct CONTROLRIGPHYSICS_API FRigPhysicsSolverComponent : public FRigBaseComponent
+{
+public:
+
+	GENERATED_BODY()
+	DECLARE_RIG_COMPONENT_METHODS(FRigPhysicsSolverComponent)
+
+	UPROPERTY(BlueprintReadOnly, EditAnywhere, Category = Physics, meta = (ShowOnlyInnerProperties))
+	FRigPhysicsSolverSettings SolverSettings;
+
+	UPROPERTY(BlueprintReadOnly, EditAnywhere, Category = Physics, meta = (ShowOnlyInnerProperties))
+	FRigPhysicsSimulationSpaceSettings SimulationSpaceSettings;
+
+	virtual bool CanBeAddedTo(const FRigElementKey& InElementKey, const URigHierarchy* InHierarchy) const;
+
+	virtual void Save(FArchive& Ar) override;
+	virtual void Load(FArchive& Ar) override;
+};
+
