Ensure physics components get fully archived

diff --git a/Engine/Plugins/Experimental/PhysicsControl/Source/ControlRigPhysics/Private/RigPhysicsComponent.cpp b/Engine/Plugins/Experimental/PhysicsControl/Source/ControlRigPhysics/Private/RigPhysicsComponent.cpp
--- a/Engine/Plugins/Experimental/PhysicsControl/Source/ControlRigPhysics/Private/RigPhysicsComponent.cpp
+++ b/Engine/Plugins/Experimental/PhysicsControl/Source/ControlRigPhysics/Private/RigPhysicsComponent.cpp
@@ -1,15 +1,22 @@
 // Copyright Epic Games, Inc. All Rights Reserved.
 
 #include "RigPhysicsComponent.h"
+#include "PhysicsControlObjectVersion.h"
 
 //======================================================================================================================
 void FRigPhysicsComponent::Save(FArchive& Ar)
 {
+	Ar.UsingCustomVersion(FPhysicsControlObjectVersion::GUID);
+
 	FRigBaseComponent::Save(Ar);
 	Ar << BodySolverSettings;
 	Ar << Dynamics;
 	Ar << Collision;
 	Ar << Joint;
+	Ar << BodyData;
+	Ar << CurrentKinematicTarget;
+	Ar << CurrentKinematicTargetSpace;
+	Ar << CurrentBodyData;
 }
 
 //======================================================================================================================
@@ -20,6 +27,10 @@ void FRigPhysicsComponent::Load(FArchive& Ar)
 	Ar << Dynamics;
 	Ar << Collision;
 	Ar << Joint;
+	Ar << BodyData;
+	Ar << CurrentKinematicTarget;
+	Ar << CurrentKinematicTargetSpace;
+	Ar << CurrentBodyData;
 }
 
 //======================================================================================================================
diff --git a/Engine/Plugins/Experimental/PhysicsControl/Source/PhysicsControl/Private/PhysicsControlModule.cpp b/Engine/Plugins/Experimental/PhysicsControl/Source/PhysicsControl/Private/PhysicsControlModule.cpp
--- a/Engine/Plugins/Experimental/PhysicsControl/Source/PhysicsControl/Private/PhysicsControlModule.cpp
+++ b/Engine/Plugins/Experimental/PhysicsControl/Source/PhysicsControl/Private/PhysicsControlModule.cpp
@@ -1,9 +1,17 @@
 // Copyright Epic Games, Inc. All Rights Reserved.
 
 #include "PhysicsControlModule.h"
+#include "PhysicsControlObjectVersion.h"
+#include "UObject/DevObjectVersion.h"
 
 #define LOCTEXT_NAMESPACE "FPhysicsControlModule"
 
+// Unique version id
+const FGuid FPhysicsControlObjectVersion::GUID(0xE68A5C1A, 0x42C3CA0D, 0xB01441B2, 0x28C1F2DD);
+// Register custom version with Core
+FDevVersionRegistration GRegisterPhysicsControlObjectVersion(
+	FPhysicsControlObjectVersion::GUID, FPhysicsControlObjectVersion::LatestVersion, TEXT("Dev-PhysicsControl"));
+
 void FPhysicsControlModule::StartupModule()
 {
 }
diff --git a/Engine/Plugins/Experimental/PhysicsControl/Source/PhysicsControl/Public/PhysicsControlData.h b/Engine/Plugins/Experimental/PhysicsControl/Source/PhysicsControl/Public/PhysicsControlData.h
--- a/Engine/Plugins/Experimental/PhysicsControl/Source/PhysicsControl/Public/PhysicsControlData.h
+++ b/Engine/Plugins/Experimental/PhysicsControl/Source/PhysicsControl/Public/PhysicsControlData.h
@@ -229,6 +229,17 @@ struct PHYSICSCONTROL_API FPhysicsControlModifierData
 
 	bool operator==(const FPhysicsControlModifierData& other) const = default;
 
+	friend FArchive& operator <<(FArchive& Ar, FPhysicsControlModifierData& Data)
+	{
+		Ar << Data.MovementType;
+		Ar << Data.CollisionType;
+		Ar << Data.GravityMultiplier;
+		Ar << Data.PhysicsBlendWeight;
+		Ar << Data.bUseSkeletalAnimation;
+		Ar << Data.bUpdateKinematicFromSimulation;
+		return Ar;
+	}
+
 	void UpdateFromSparseData(const FPhysicsControlModifierSparseData& SparseData);
 
 	/**
@@ -263,7 +274,7 @@ struct PHYSICSCONTROL_API FPhysicsControlModifierData
 	 * NOTE This is currently only used in the PhysicsControlComponent implementation
 	 */
 	UPROPERTY(EditAnywhere, BlueprintReadWrite, Category = PhysicsControl)
-	uint8 bUseSkeletalAnimation : 1;
+	bool bUseSkeletalAnimation;
 
 	/**
 	 * If true then the associated actor's transform will be updated from the simulation when it is
@@ -272,7 +283,7 @@ struct PHYSICSCONTROL_API FPhysicsControlModifierData
 	 * NOTE This is currently only used in the PhysicsControlComponent implementation (is it needed in RBWC?)
 	 */
 	UPROPERTY(EditAnywhere, BlueprintReadWrite, Category = PhysicsControl)
-	uint8 bUpdateKinematicFromSimulation : 1;
+	bool bUpdateKinematicFromSimulation;
 };
 
 
diff --git a/Engine/Plugins/Experimental/PhysicsControl/Source/PhysicsControl/Public/PhysicsControlObjectVersion.h b/Engine/Plugins/Experimental/PhysicsControl/Source/PhysicsControl/Public/PhysicsControlObjectVersion.h
new file mode 100644
--- /dev/null
+++ b/Engine/Plugins/Experimental/PhysicsControl/Source/PhysicsControl/Public/PhysicsControlObjectVersion.h
@@ -0,0 +1,23 @@
+// Copyright Epic Games, Inc. All Rights Reserved.
+#pragma once
+
+#include "Misc/Guid.h"
+
+struct PHYSICSCONTROL_API FPhysicsControlObjectVersion
+{
+	enum Type
+	{
+		// Before any version changes were made
+		BeforeCustomVersionWasAdded,
+
+		// -----<new versions can be added above this line>-------------------------------------------------
+		VersionPlusOne,
+		LatestVersion = VersionPlusOne - 1
+	};
+
+	// The GUID for this custom version number
+	const static FGuid GUID;
+
+private:
+	FPhysicsControlObjectVersion() {}
+};
