Control Rig: Component call back to react on spawning, adapting physics component / unit for better user flow

Preflight was https://horde.devtools.epicgames.com/job/67ae0b1f6170813d743fbaaf?step=b8ee - which is failing for GDKHandle only - so I'll marked as skipped preflight to be able to submit.

#rnx

diff --git a/Engine/Plugins/Animation/ControlRig/Source/ControlRig/Private/Rigs/RigHierarchy.cpp b/Engine/Plugins/Animation/ControlRig/Source/ControlRig/Private/Rigs/RigHierarchy.cpp
--- a/Engine/Plugins/Animation/ControlRig/Source/ControlRig/Private/Rigs/RigHierarchy.cpp
+++ b/Engine/Plugins/Animation/ControlRig/Source/ControlRig/Private/Rigs/RigHierarchy.cpp
@@ -1648,6 +1648,22 @@ FRigBaseComponent* URigHierarchy::FindComponent(const FRigComponentKey& InKey)
 	return nullptr;
 }
 
+TArray<const FRigBaseComponent*> URigHierarchy::GetComponents(const UScriptStruct* InComponentStruct) const
+{
+	check(InComponentStruct);
+	check(InComponentStruct->IsChildOf(FRigBaseComponent::StaticStruct()));
+	
+	TArray<const FRigBaseComponent*> Result;
+	for(const FInstancedStruct& ComponentInstancedStruct : ElementComponents)
+	{
+		if(ComponentInstancedStruct.GetScriptStruct() == InComponentStruct)
+		{
+			Result.Add(reinterpret_cast<const FRigBaseComponent*>(ComponentInstancedStruct.GetMemory()));
+		}
+	}
+	return Result;
+}
+
 const FRigBaseComponent* URigHierarchy::FindTopLevelComponent(const FName& InName) const
 {
 	return FindComponent(FRigComponentKey(GetTopLevelComponentElementKey(), InName));
@@ -1725,6 +1741,25 @@ FRigBaseComponent* URigHierarchy::GetComponent(int32 InIndex)
 	return nullptr;
 }
 
+int32 URigHierarchy::NumComponents() const
+{
+	return Algo::CountIf(ElementComponents, [](const FInstancedStruct& InComponentInstancedStruct)
+	{
+		return InComponentInstancedStruct.IsValid();
+	});
+}
+
+int32 URigHierarchy::NumComponents(const UScriptStruct* InComponentStruct) const
+{
+	check(InComponentStruct);
+	check(InComponentStruct->IsChildOf(FRigBaseComponent::StaticStruct()));
+
+	return Algo::CountIf(ElementComponents, [InComponentStruct](const FInstancedStruct& InComponentInstancedStruct)
+	{
+		return InComponentInstancedStruct.GetScriptStruct() == InComponentStruct;
+	});
+}
+
 int32 URigHierarchy::NumComponents(FRigElementKey InElement) const
 {
 	if(InElement == GetTopLevelComponentElementKey())
diff --git a/Engine/Plugins/Animation/ControlRig/Source/ControlRig/Private/Rigs/RigHierarchyComponents.cpp b/Engine/Plugins/Animation/ControlRig/Source/ControlRig/Private/Rigs/RigHierarchyComponents.cpp
--- a/Engine/Plugins/Animation/ControlRig/Source/ControlRig/Private/Rigs/RigHierarchyComponents.cpp
+++ b/Engine/Plugins/Animation/ControlRig/Source/ControlRig/Private/Rigs/RigHierarchyComponents.cpp
@@ -7,6 +7,8 @@
 #include "UObject/UObjectIterator.h"
 #include "Serialization/MemoryReader.h"
 #include "Serialization/MemoryWriter.h"
+#include "UObject/AnimObjectVersion.h"
+#include "ControlRigObjectVersion.h"
 
 #include UE_INLINE_GENERATED_CPP_BY_NAME(RigHierarchyComponents)
 
@@ -80,6 +82,9 @@ TArray<UScriptStruct*> FRigBaseComponent::GetAllComponentScriptStructs(bool bSor
 
 bool FRigBaseComponent::Serialize(FArchive& Ar)
 {
+	Ar.UsingCustomVersion(FAnimObjectVersion::GUID);
+	Ar.UsingCustomVersion(FControlRigObjectVersion::GUID);
+
 	if(Ar.IsLoading())
 	{
 		Load(Ar);
@@ -108,6 +113,7 @@ FRigComponentState FRigBaseComponent::GetState() const
 	State.ComponentStruct = GetScriptStruct();
 	FMemoryWriter Writer(State.Data);
 	const_cast<FRigBaseComponent*>(this)->Save(Writer);
+	State.Versions = Writer.GetCustomVersions();
 	return State;
 }
 
@@ -116,6 +122,7 @@ bool FRigBaseComponent::SetState(const FRigComponentState& InState)
 	if(InState.IsValid() && InState.GetComponentStruct() == GetScriptStruct())
 	{
 		FMemoryReader Reader(InState.Data);
+		Reader.SetCustomVersions(InState.Versions);
 		Load(Reader);
 		return true;
 	}
diff --git a/Engine/Plugins/Animation/ControlRig/Source/ControlRig/Private/Rigs/RigHierarchyController.cpp b/Engine/Plugins/Animation/ControlRig/Source/ControlRig/Private/Rigs/RigHierarchyController.cpp
--- a/Engine/Plugins/Animation/ControlRig/Source/ControlRig/Private/Rigs/RigHierarchyController.cpp
+++ b/Engine/Plugins/Animation/ControlRig/Source/ControlRig/Private/Rigs/RigHierarchyController.cpp
@@ -992,6 +992,9 @@ FRigComponentKey URigHierarchyController::AddComponent(UScriptStruct* InComponen
 	Hierarchy->IncrementTopologyVersion();
 	Notify(ERigHierarchyNotification::ComponentAdded, NewComponent);
 
+	// allow the component to react to it being spawned
+	NewComponent->OnAddedToHierarchy(Hierarchy, this);
+
 #if WITH_EDITOR
 	TransactionPtr.Reset();
 
diff --git a/Engine/Plugins/Animation/ControlRig/Source/ControlRig/Public/Rigs/RigHierarchy.h b/Engine/Plugins/Animation/ControlRig/Source/ControlRig/Public/Rigs/RigHierarchy.h
--- a/Engine/Plugins/Animation/ControlRig/Source/ControlRig/Public/Rigs/RigHierarchy.h
+++ b/Engine/Plugins/Animation/ControlRig/Source/ControlRig/Public/Rigs/RigHierarchy.h
@@ -816,7 +816,28 @@ public:
 	{
 		return CastChecked<T>(GetComponent(InIndex));
 	}
-	
+
+	/**
+	 * Returns all components of a given type
+	 * @param InComponentStruct The component type to look up
+	 * @return All components of type InComponentStruct
+	 */
+	TArray<const FRigBaseComponent*> GetComponents(const UScriptStruct* InComponentStruct) const;
+
+	/**
+	 * Returns all components of the template type
+	 * @return All components of the template type
+	 */
+	template<typename T>
+	TArray<const T*> GetComponents() const
+	{
+		TArray<const FRigBaseComponent*> Components = GetComponents(T::StaticStruct());
+		TArray<const T*> Result;
+		Result.SetNumZeroed(Components.Num());
+		FMemory::Memcpy(Result.GetData(), Components.GetData(), Components.GetTypeSize() * Components.Num());
+		return Result;
+	}
+
 	/**
 	 * Returns a component for a given top level name
 	 * @param InName The name of the component to retrieve.
@@ -897,6 +918,29 @@ public:
 		return CastChecked<T>(GetTopLevelComponent(InTopLevelComponentIndex));
 	}
 
+	/**
+	 * Returns the count of all components
+	 * @return The count of all components
+	 */
+	int32 NumComponents() const;
+
+	/**
+	 * Returns the count of all components of a given type
+	 * @param InComponentStruct The component type to look up
+	 * @return The count of all components of type InComponentStruct
+	 */
+	int32 NumComponents(const UScriptStruct* InComponentStruct) const;
+
+	/**
+	 * Returns the count of all components of the template type
+	 * @return The count of all components of the template type
+	 */
+	template<typename T>
+	int32 NumComponents() const
+	{
+		return NumComponents(T::StaticStruct());
+	}
+
 	UFUNCTION(BlueprintPure, Category = URigHierarchy)
 	int32 NumComponents(FRigElementKey InElement) const;
 
diff --git a/Engine/Plugins/Animation/ControlRig/Source/ControlRig/Public/Rigs/RigHierarchyComponents.h b/Engine/Plugins/Animation/ControlRig/Source/ControlRig/Public/Rigs/RigHierarchyComponents.h
--- a/Engine/Plugins/Animation/ControlRig/Source/ControlRig/Public/Rigs/RigHierarchyComponents.h
+++ b/Engine/Plugins/Animation/ControlRig/Source/ControlRig/Public/Rigs/RigHierarchyComponents.h
@@ -6,8 +6,12 @@
 #include "RigHierarchyDefines.h"
 #include "Textures/SlateIcon.h"
 #include "Styling/SlateColor.h"
+#include "Serialization/CustomVersion.h"
 #include "RigHierarchyComponents.generated.h"
 
+class URigHierarchy;
+class URigHierarchyController;
+
 #define DECLARE_RIG_COMPONENT_METHODS(ComponentType) \
 virtual UScriptStruct* GetScriptStruct() const override { return ComponentType::StaticStruct(); } \
 template<typename T> \
@@ -46,6 +50,7 @@ private:
 
 	const UScriptStruct* ComponentStruct = nullptr;
 	TArray<uint8> Data;
+	FCustomVersionContainer Versions;
 
 	friend struct FRigBaseComponent;
 };
@@ -91,6 +96,9 @@ public:
 	// you need to pass the URigHierarchy::GetTopLevelComponentElementKey() as InElementKey
 	virtual bool CanBeAddedTo(const FRigElementKey& InElementKey, const URigHierarchy* InHierarchy, FString* OutFailureReason = nullptr) const { return true; }
 
+	// allows the component to react to it being spawned
+	virtual void OnAddedToHierarchy(URigHierarchy* InHierarchy, URigHierarchyController* InController) {}
+
 	// return the icon to use for this component in the UI
 	virtual const FSlateIcon& GetIconForUI() const;
 
diff --git a/Engine/Plugins/Experimental/AnimNext/Source/AnimNextEditor/Private/Common/GraphEditorSchemaActions.cpp b/Engine/Plugins/Experimental/AnimNext/Source/AnimNextEditor/Private/Common/GraphEditorSchemaActions.cpp
--- a/Engine/Plugins/Experimental/AnimNext/Source/AnimNextEditor/Private/Common/GraphEditorSchemaActions.cpp
+++ b/Engine/Plugins/Experimental/AnimNext/Source/AnimNextEditor/Private/Common/GraphEditorSchemaActions.cpp
@@ -52,7 +52,7 @@ UEdGraphNode* FAnimNextSchemaAction_RigUnit::PerformAction(UEdGraph* ParentGraph
 		Controller->OpenUndoBracket(FString::Printf(TEXT("Add '%s' Node"), *Name.ToString()));
 
 		FRigVMUnitNodeCreatedContext& UnitNodeCreatedContext = Controller->GetUnitNodeCreatedContext();
-		FRigVMUnitNodeCreatedContext::FScope ReasonScope(UnitNodeCreatedContext, ERigVMNodeCreatedReason::NodeSpawner);
+		FRigVMUnitNodeCreatedContext::FScope ReasonScope(UnitNodeCreatedContext, ERigVMNodeCreatedReason::NodeSpawner, Host);
 
 		if (URigVMUnitNode* ModelNode = Controller->AddUnitNode(StructTemplate, NodeClass, FRigVMStruct::ExecuteName, Location, Name.ToString(), true, false))
 		{
diff --git a/Engine/Plugins/Experimental/PhysicsControl/Source/ControlRigPhysics/ControlRigPhysics.Build.cs b/Engine/Plugins/Experimental/PhysicsControl/Source/ControlRigPhysics/ControlRigPhysics.Build.cs
--- a/Engine/Plugins/Experimental/PhysicsControl/Source/ControlRigPhysics/ControlRigPhysics.Build.cs
+++ b/Engine/Plugins/Experimental/PhysicsControl/Source/ControlRigPhysics/ControlRigPhysics.Build.cs
@@ -29,5 +29,16 @@ public class ControlRigPhysics : ModuleRules
 				"RigVM"
 			}
 		);
+		
+		if (Target.bBuildEditor == true)
+		{
+			PrivateDependencyModuleNames.AddRange(
+				new string[]
+				{
+					"RigVMDeveloper",
+					"ControlRigDeveloper",
+				}
+			);
+		}
 	}
 }
diff --git a/Engine/Plugins/Experimental/PhysicsControl/Source/ControlRigPhysics/Private/ControlRigPhysicsModule.cpp b/Engine/Plugins/Experimental/PhysicsControl/Source/ControlRigPhysics/Private/ControlRigPhysicsModule.cpp
--- a/Engine/Plugins/Experimental/PhysicsControl/Source/ControlRigPhysics/Private/ControlRigPhysicsModule.cpp
+++ b/Engine/Plugins/Experimental/PhysicsControl/Source/ControlRigPhysics/Private/ControlRigPhysicsModule.cpp
@@ -20,3 +20,4 @@ void FControlRigPhysicsModule::ShutdownModule()
 	
 IMPLEMENT_MODULE(FControlRigPhysicsModule, ControlRigPhysics)
 
+TAutoConsoleVariable<bool> CVarControlRigPhysicsAutomaticallySetSolver(TEXT("ControlRig.Physics.AutomaticallySetSolver"), false, TEXT("if this is set to true component and nodes will pick the solver automatically."));
diff --git a/Engine/Plugins/Experimental/PhysicsControl/Source/ControlRigPhysics/Private/RigPhysicsComponent.cpp b/Engine/Plugins/Experimental/PhysicsControl/Source/ControlRigPhysics/Private/RigPhysicsComponent.cpp
--- a/Engine/Plugins/Experimental/PhysicsControl/Source/ControlRigPhysics/Private/RigPhysicsComponent.cpp
+++ b/Engine/Plugins/Experimental/PhysicsControl/Source/ControlRigPhysics/Private/RigPhysicsComponent.cpp
@@ -2,6 +2,10 @@
 
 #include "RigPhysicsComponent.h"
 #include "PhysicsControlObjectVersion.h"
+#include "RigPhysicsSolverComponent.h"
+#include "Rigs/RigHierarchy.h"
+#include "Rigs/RigHierarchyController.h"
+#include "ControlRigPhysicsModule.h"
 
 //======================================================================================================================
 void FRigPhysicsComponent::Save(FArchive& Ar)
@@ -46,3 +50,25 @@ bool FRigPhysicsComponent::CanBeAddedTo(const FRigElementKey& InElementKey, cons
 	}
 	return true;
 }
+
+//======================================================================================================================
+void FRigPhysicsComponent::OnAddedToHierarchy(URigHierarchy* InHierarchy, URigHierarchyController* InController)
+{
+	if(!CVarControlRigPhysicsAutomaticallySetSolver.GetValueOnAnyThread())
+	{
+		return;
+	}
+	
+	// add a solver if there isn't any
+	TArray<const FRigPhysicsSolverComponent*> SolverComponents = InHierarchy->GetComponents<FRigPhysicsSolverComponent>();
+	if(SolverComponents.IsEmpty())
+	{
+		BodySolverSettings.PhysicsSolverComponentKey = InController->AddTopLevelComponent(FRigPhysicsSolverComponent::StaticStruct(), NAME_None);
+	}
+	else
+	{
+		BodySolverSettings.PhysicsSolverComponentKey = SolverComponents[0]->GetKey();
+	}
+	// let the system know we changed something
+	InHierarchy->Notify(ERigHierarchyNotification::ComponentContentChanged, this);
+}
diff --git a/Engine/Plugins/Experimental/PhysicsControl/Source/ControlRigPhysics/Private/RigPhysicsExecution.cpp b/Engine/Plugins/Experimental/PhysicsControl/Source/ControlRigPhysics/Private/RigPhysicsExecution.cpp
--- a/Engine/Plugins/Experimental/PhysicsControl/Source/ControlRigPhysics/Private/RigPhysicsExecution.cpp
+++ b/Engine/Plugins/Experimental/PhysicsControl/Source/ControlRigPhysics/Private/RigPhysicsExecution.cpp
@@ -8,11 +8,84 @@
 #include "Rigs/RigHierarchyController.h"
 #include "Units/RigUnitContext.h"
 #include "Units/Execution/RigUnit_PrepareForExecution.h"
+#include "ControlRigPhysicsModule.h"
 
 #include "ControlRig.h"
+#if WITH_EDITOR
+#include "RigVMModel/RigVMController.h"
+#include "ControlRigBlueprint.h"
+#endif
 
 #include UE_INLINE_GENERATED_CPP_BY_NAME(RigPhysicsExecution)
 
+//======================================================================================================================
+#if WITH_EDITOR
+void FRigUnit_PhysicsBaseMutable::OnUnitNodeCreated(FRigVMUnitNodeCreatedContext& InContext) const
+{
+	FRigUnitMutable::OnUnitNodeCreated(InContext);
+	
+	if(!CVarControlRigPhysicsAutomaticallySetSolver.GetValueOnAnyThread())
+	{
+		return;
+	}
+
+	// if the user has dropped the node manually
+	if(InContext.GetReason() != ERigVMNodeCreatedReason::NodeSpawner)
+	{
+		return;
+	}
+
+	const UControlRigBlueprint* Blueprint = Cast<UControlRigBlueprint>(InContext.GetHost());
+	if(Blueprint == nullptr)
+	{
+		return;
+	}
+
+	const URigHierarchy* Hierarchy = Blueprint->GetHierarchy();
+	if(Hierarchy == nullptr)
+	{
+		return;
+	}
+
+	URigVMController* Controller = InContext.GetController();
+	if(Controller == nullptr)
+	{
+		return;
+	}
+
+	const URigVMGraph* Graph = Controller->GetGraph();
+	if(Graph == nullptr)
+	{
+		return;
+	}
+
+	const URigVMNode* MyNode = Graph->FindNodeByName(InContext.GetNodeName());
+	if(MyNode == nullptr)
+	{
+		return;
+	}
+
+	const URigVMPin* SolverKeyPin = MyNode->FindPin(GET_MEMBER_NAME_STRING_CHECKED(FRigUnit_AddPhysicsSolver, PhysicsSolverComponentKey));
+	if(SolverKeyPin == nullptr)
+	{
+		return;
+	}
+
+	if(SolverKeyPin->GetDirection() != ERigVMPinDirection::Input)
+	{
+		return;
+	}
+	
+	TArray<const FRigPhysicsSolverComponent*> SolverComponents = Hierarchy->GetComponents<FRigPhysicsSolverComponent>();
+	if(!SolverComponents.IsEmpty())
+	{
+		FString ContentAsString;
+		FRigComponentKey::StaticStruct()->ExportText(ContentAsString, &SolverComponents[0]->GetKey(), nullptr, nullptr, PPF_ExternalEditor, nullptr);
+		Controller->SetPinDefaultValue(SolverKeyPin->GetPinPath(), ContentAsString);
+	}
+}
+#endif
+
 //======================================================================================================================
 FRigUnit_AddPhysicsSolver_Execute()
 {
@@ -341,4 +414,3 @@ FRigUnit_HierarchySetPhysicsBodyUpdateKinematicFromSimulation_Execute()
 		}
 	}
 }
-
diff --git a/Engine/Plugins/Experimental/PhysicsControl/Source/ControlRigPhysics/Public/ControlRigPhysicsModule.h b/Engine/Plugins/Experimental/PhysicsControl/Source/ControlRigPhysics/Public/ControlRigPhysicsModule.h
--- a/Engine/Plugins/Experimental/PhysicsControl/Source/ControlRigPhysics/Public/ControlRigPhysicsModule.h
+++ b/Engine/Plugins/Experimental/PhysicsControl/Source/ControlRigPhysics/Public/ControlRigPhysicsModule.h
@@ -3,6 +3,7 @@
 #pragma once
 
 #include "CoreMinimal.h"
+#include "HAL/IConsoleManager.h"
 #include "Modules/ModuleManager.h"
 
 #if UE_BUILD_SHIPPING || UE_BUILD_TEST
@@ -19,3 +20,5 @@ public:
 	virtual void StartupModule() override;
 	virtual void ShutdownModule() override;
 };
+
+extern CONTROLRIGPHYSICS_API TAutoConsoleVariable<bool> CVarControlRigPhysicsAutomaticallySetSolver;
diff --git a/Engine/Plugins/Experimental/PhysicsControl/Source/ControlRigPhysics/Public/RigPhysicsComponent.h b/Engine/Plugins/Experimental/PhysicsControl/Source/ControlRigPhysics/Public/RigPhysicsComponent.h
--- a/Engine/Plugins/Experimental/PhysicsControl/Source/ControlRigPhysics/Public/RigPhysicsComponent.h
+++ b/Engine/Plugins/Experimental/PhysicsControl/Source/ControlRigPhysics/Public/RigPhysicsComponent.h
@@ -60,5 +60,7 @@ public:
 	virtual void Load(FArchive& Ar) override;
 
 	virtual bool CanBeAddedTo(const FRigElementKey& InElementKey, const URigHierarchy* InHierarchy, FString* OutFailureReason) const override;
+
+	virtual void OnAddedToHierarchy(URigHierarchy* InHierarchy, URigHierarchyController* InController) override;
 };
 
diff --git a/Engine/Plugins/Experimental/PhysicsControl/Source/ControlRigPhysics/Public/RigPhysicsExecution.h b/Engine/Plugins/Experimental/PhysicsControl/Source/ControlRigPhysics/Public/RigPhysicsExecution.h
--- a/Engine/Plugins/Experimental/PhysicsControl/Source/ControlRigPhysics/Public/RigPhysicsExecution.h
+++ b/Engine/Plugins/Experimental/PhysicsControl/Source/ControlRigPhysics/Public/RigPhysicsExecution.h
@@ -10,12 +10,25 @@
 
 #include "RigPhysicsExecution.generated.h"
 
+/**
+ * Base struct for all other mutable physics nodes
+ */
+USTRUCT(meta = (Category = "RigPhysics", NodeColor = "1.0 0.6 0.3", Keywords = "Physics"))
+struct CONTROLRIGPHYSICS_API FRigUnit_PhysicsBaseMutable: public FRigUnitMutable
+{
+	GENERATED_BODY()
+	
+#if WITH_EDITOR
+	virtual void OnUnitNodeCreated(FRigVMUnitNodeCreatedContext& InContext) const override;
+#endif
+};
+
 /**
  * Adds a new physics solver as a component on the owner element.
  * Note: This node only runs as part of the construction event.
  */
-USTRUCT(meta=(Category = "RigPhysics", NodeColor = "1.0 0.6 0.3", DisplayName="Add Physics Solver Component", Keywords="Construction,Create,New,Simulation,Physics", Varying))
-struct CONTROLRIGPHYSICS_API FRigUnit_AddPhysicsSolver : public FRigUnitMutable
+USTRUCT(meta=(DisplayName="Add Physics Solver Component", Keywords="Construction,Create,New,Simulation", Varying))
+struct CONTROLRIGPHYSICS_API FRigUnit_AddPhysicsSolver : public FRigUnit_PhysicsBaseMutable
 {
 	GENERATED_BODY()
 
@@ -44,8 +57,8 @@ struct CONTROLRIGPHYSICS_API FRigUnit_AddPhysicsSolver : public FRigUnitMutable
  * requested. Explicit instantiation allows the timing to be controlled, as allocations etc may cause some
  * delays.
  */
-USTRUCT(meta = (Category = "RigPhysics", NodeColor = "1.0 0.6 0.3", DisplayName = "Instantiate physics", Keywords = "Physics"))
-struct CONTROLRIGPHYSICS_API FRigUnit_InstantiatePhysics : public FRigUnitMutable
+USTRUCT(meta = (DisplayName = "Instantiate physics"))
+struct CONTROLRIGPHYSICS_API FRigUnit_InstantiatePhysics : public FRigUnit_PhysicsBaseMutable
 {
 	GENERATED_BODY()
 
@@ -63,8 +76,8 @@ struct CONTROLRIGPHYSICS_API FRigUnit_InstantiatePhysics : public FRigUnitMutabl
 /**
  * Steps the specified physics solver
  */
-USTRUCT(meta = (Category = "RigPhysics", NodeColor = "1.0 0.6 0.3", DisplayName = "Step Physics Solver", Keywords = "Simulate,Physics"))
-struct CONTROLRIGPHYSICS_API FRigUnit_StepPhysicsSolver : public FRigUnitMutable
+USTRUCT(meta = (DisplayName = "Step Physics Solver"))
+struct CONTROLRIGPHYSICS_API FRigUnit_StepPhysicsSolver : public FRigUnit_PhysicsBaseMutable
 {
 	GENERATED_BODY()
 
@@ -94,8 +107,8 @@ struct CONTROLRIGPHYSICS_API FRigUnit_StepPhysicsSolver : public FRigUnitMutable
  * Adds a new physics body as a component on the owner element.
  * Note: This node only runs as part of the construction event.
  */
-USTRUCT(meta = (Category = "RigPhysics", NodeColor = "1.0 0.6 0.3", DisplayName = "Add Physics Body Component", Keywords = "Construction,Create,New,Joint,Body,Skeleton,Physics", Varying))
-struct CONTROLRIGPHYSICS_API FRigUnit_AddPhysicsBody : public FRigUnitMutable
+USTRUCT(meta = (DisplayName = "Add Physics Body Component", Keywords = "Construction,Create,New,Joint,Body,Skeleton", Varying))
+struct CONTROLRIGPHYSICS_API FRigUnit_AddPhysicsBody : public FRigUnit_PhysicsBaseMutable
 {
 	GENERATED_BODY()
 
@@ -146,8 +159,8 @@ struct CONTROLRIGPHYSICS_API FRigUnit_AddPhysicsBody : public FRigUnitMutable
  * Sets what bone is used as a source transform for the physics body. This is used as a kinematic target, and when
  * initializing the simulation.
  */
-USTRUCT(meta = (Category = "RigPhysics", NodeColor = "1.0 0.6 0.3", DisplayName = "Set Source Bone", Keywords = "Physics", Varying))
-struct CONTROLRIGPHYSICS_API FRigUnit_HierarchySetPhysicsBodySourceBone : public FRigUnitMutable
+USTRUCT(meta = (DisplayName = "Set Source Bone", Varying))
+struct CONTROLRIGPHYSICS_API FRigUnit_HierarchySetPhysicsBodySourceBone : public FRigUnit_PhysicsBaseMutable
 {
 	GENERATED_BODY()
 
@@ -164,8 +177,8 @@ struct CONTROLRIGPHYSICS_API FRigUnit_HierarchySetPhysicsBodySourceBone : public
 /**
  * Sets what bone is targeted by the simulation - i.e. where the simulation output is written to.
  */
-USTRUCT(meta = (Category = "RigPhysics", NodeColor = "1.0 0.6 0.3", DisplayName = "Set Target Bone", Keywords = "Physics", Varying))
-struct CONTROLRIGPHYSICS_API FRigUnit_HierarchySetPhysicsBodyTargetBone : public FRigUnitMutable
+USTRUCT(meta = (DisplayName = "Set Target Bone", Varying))
+struct CONTROLRIGPHYSICS_API FRigUnit_HierarchySetPhysicsBodyTargetBone : public FRigUnit_PhysicsBaseMutable
 {
 	GENERATED_BODY()
 
@@ -185,8 +198,8 @@ struct CONTROLRIGPHYSICS_API FRigUnit_HierarchySetPhysicsBodyTargetBone : public
  * Danny TODO - Note that the sparse data does not get displayed correctly, so this is largely 
  * unusable - the flags that enable/disable all end up getting reset if the user attempts to change them.
  */
-USTRUCT(meta = (Category = "RigPhysics", NodeColor = "1.0 0.6 0.3", DisplayName = "Set Physics Body Data", Keywords = "Physics", Varying))
-struct CONTROLRIGPHYSICS_API FRigUnit_HierarchySetPhysicsBodySparseData : public FRigUnitMutable
+USTRUCT(meta = (DisplayName = "Set Physics Body Data", Varying))
+struct CONTROLRIGPHYSICS_API FRigUnit_HierarchySetPhysicsBodySparseData : public FRigUnit_PhysicsBaseMutable
 {
 	GENERATED_BODY()
 
@@ -206,8 +219,8 @@ struct CONTROLRIGPHYSICS_API FRigUnit_HierarchySetPhysicsBodySparseData : public
 /**
  * Sets the kinematic target for a body - note that this won't actually make the body kinematic
  */
-USTRUCT(meta = (Category = "RigPhysics", NodeColor = "1.0 0.6 0.3", DisplayName = "Set Kinematic Target", Keywords = "Physics", Varying))
-struct CONTROLRIGPHYSICS_API FRigUnit_HierarchySetPhysicsBodyKinematicTarget : public FRigUnitMutable
+USTRUCT(meta = (DisplayName = "Set Kinematic Target", Varying))
+struct CONTROLRIGPHYSICS_API FRigUnit_HierarchySetPhysicsBodyKinematicTarget : public FRigUnit_PhysicsBaseMutable
 {
 	GENERATED_BODY()
 
@@ -231,8 +244,8 @@ struct CONTROLRIGPHYSICS_API FRigUnit_HierarchySetPhysicsBodyKinematicTarget : p
  * Sets the movement type for this body. 
  * Danny TODO explain the different between kinematic and static
  */
-USTRUCT(meta = (Category = "RigPhysics", NodeColor = "1.0 0.6 0.3", DisplayName = "Set Physics Body Movement Type", Keywords = "Physics", Varying))
-struct CONTROLRIGPHYSICS_API FRigUnit_HierarchySetPhysicsBodyMovementType : public FRigUnitMutable
+USTRUCT(meta = (DisplayName = "Set Physics Body Movement Type", Varying))
+struct CONTROLRIGPHYSICS_API FRigUnit_HierarchySetPhysicsBodyMovementType : public FRigUnit_PhysicsBaseMutable
 {
 	GENERATED_BODY()
 
@@ -252,8 +265,8 @@ struct CONTROLRIGPHYSICS_API FRigUnit_HierarchySetPhysicsBodyMovementType : publ
 /**
  * Sets what collision is used for this body
  */
-USTRUCT(meta = (Category = "RigPhysics", NodeColor = "1.0 0.6 0.3", DisplayName = "Set Physics Body Collision Type", Keywords = "Physics", Varying))
-struct CONTROLRIGPHYSICS_API FRigUnit_HierarchySetPhysicsBodyCollisionType : public FRigUnitMutable
+USTRUCT(meta = (DisplayName = "Set Physics Body Collision Type", Varying))
+struct CONTROLRIGPHYSICS_API FRigUnit_HierarchySetPhysicsBodyCollisionType : public FRigUnit_PhysicsBaseMutable
 {
 	GENERATED_BODY()
 
@@ -273,8 +286,8 @@ struct CONTROLRIGPHYSICS_API FRigUnit_HierarchySetPhysicsBodyCollisionType : pub
 /**
  * Sets the multiplier on gravity that should be applied to the body.
  */
-USTRUCT(meta = (Category = "RigPhysics", NodeColor = "1.0 0.6 0.3", DisplayName = "Set Physics Body Gravity Multiplier", Keywords = "Physics", Varying))
-struct CONTROLRIGPHYSICS_API FRigUnit_HierarchySetPhysicsBodyGravityMultiplier : public FRigUnitMutable
+USTRUCT(meta = (DisplayName = "Set Physics Body Gravity Multiplier", Varying))
+struct CONTROLRIGPHYSICS_API FRigUnit_HierarchySetPhysicsBodyGravityMultiplier : public FRigUnit_PhysicsBaseMutable
 {
 	GENERATED_BODY()
 
@@ -295,8 +308,8 @@ struct CONTROLRIGPHYSICS_API FRigUnit_HierarchySetPhysicsBodyGravityMultiplier :
  * Controls the amount that the simulation is blended back into the target bones. 
  * Danny TODO implement
  */
-USTRUCT(meta = (Category = "RigPhysics", NodeColor = "1.0 0.6 0.3", DisplayName = "Set Physics Body Physics Blend Weight", Keywords = "Physics", Varying))
-struct CONTROLRIGPHYSICS_API FRigUnit_HierarchySetPhysicsBodyPhysicsBlendWeight : public FRigUnitMutable
+USTRUCT(meta = (DisplayName = "Set Physics Body Physics Blend Weight", Varying))
+struct CONTROLRIGPHYSICS_API FRigUnit_HierarchySetPhysicsBodyPhysicsBlendWeight : public FRigUnit_PhysicsBaseMutable
 {
 	GENERATED_BODY()
 
@@ -317,8 +330,8 @@ struct CONTROLRIGPHYSICS_API FRigUnit_HierarchySetPhysicsBodyPhysicsBlendWeight
  * Danny TODO I think this is redundant given the ability to specify the space when kinematic 
  * targets are set. Investigate propagating that into the rest of Physics Control.
  */
-USTRUCT(meta = (Category = "RigPhysics", NodeColor = "1.0 0.6 0.3", DisplayName = "Set Physics Body Use Skeletal Animation", Keywords = "Physics", Varying))
-struct CONTROLRIGPHYSICS_API FRigUnit_HierarchySetPhysicsBodyUseSkeletalAnimation : public FRigUnitMutable
+USTRUCT(meta = (DisplayName = "Set Physics Body Use Skeletal Animation", Varying))
+struct CONTROLRIGPHYSICS_API FRigUnit_HierarchySetPhysicsBodyUseSkeletalAnimation : public FRigUnit_PhysicsBaseMutable
 {
 	GENERATED_BODY()
 
@@ -340,8 +353,8 @@ struct CONTROLRIGPHYSICS_API FRigUnit_HierarchySetPhysicsBodyUseSkeletalAnimatio
  * necessary when either kinematic targets are being used, or when the target bone differs from the source bone.
  * Danny TODO - check the implementation, and also can we auto-calculate this at runtime?
  */
-USTRUCT(meta = (Category = "RigPhysics", NodeColor = "1.0 0.6 0.3", DisplayName = "Set Physics Body Update Kinematic From Simulation", Keywords = "Physics", Varying))
-struct CONTROLRIGPHYSICS_API FRigUnit_HierarchySetPhysicsBodyUpdateKinematicFromSimulation : public FRigUnitMutable
+USTRUCT(meta = (DisplayName = "Set Physics Body Update Kinematic From Simulation", Varying))
+struct CONTROLRIGPHYSICS_API FRigUnit_HierarchySetPhysicsBodyUpdateKinematicFromSimulation : public FRigUnit_PhysicsBaseMutable
 {
 	GENERATED_BODY()
 
diff --git a/Engine/Plugins/Runtime/RigVM/Source/RigVM/Public/RigVMCore/RigVMStruct.h b/Engine/Plugins/Runtime/RigVM/Source/RigVM/Public/RigVMCore/RigVMStruct.h
--- a/Engine/Plugins/Runtime/RigVM/Source/RigVM/Public/RigVMCore/RigVMStruct.h
+++ b/Engine/Plugins/Runtime/RigVM/Source/RigVM/Public/RigVMCore/RigVMStruct.h
@@ -36,6 +36,7 @@ DECLARE_DELEGATE_RetVal_TwoParams(bool, FRigVMBindPinToExternalVariableDelegate,
 
 struct FRigVMStruct;
 class URigVMController;
+class IRigVMClientHost;
 
 /** Context as of why the node was created */
 enum class ERigVMNodeCreatedReason : uint8
@@ -57,21 +58,33 @@ public:
 	struct FScope
 	{
 	public:
-		FScope(FRigVMUnitNodeCreatedContext& InContext, ERigVMNodeCreatedReason InReason)
+		FScope(FRigVMUnitNodeCreatedContext& InContext, ERigVMNodeCreatedReason InReason, IRigVMClientHost* InHost)
 			: Context(InContext)
 			, PreviousReason(InContext.GetReason())
+#if WITH_EDITOR
+			, PreviousHost(InContext.GetHost())
+#endif
 		{
 			Context.Reason = InReason;
+#if WITH_EDITOR
+			Context.Host = InHost;
+#endif
 		}
 
 		~FScope()
 		{
 			Context.Reason = PreviousReason;
+#if WITH_EDITOR
+			Context.Host = PreviousHost;
+#endif
 		}
 
 	private:
 		FRigVMUnitNodeCreatedContext& Context;
 		ERigVMNodeCreatedReason PreviousReason;
+#if WITH_EDITOR
+		IRigVMClientHost* PreviousHost = nullptr;
+#endif
 	};
 
 	/** Returns true if this context is valid to use */
@@ -82,6 +95,7 @@ public:
 
 #if WITH_EDITOR
 	URigVMController* GetController() const { return Controller; }
+	IRigVMClientHost* GetHost() const { return Host; }
 #endif
 
 	/** Get the name of this node */
@@ -149,7 +163,8 @@ public:
 private:
 
 #if WITH_EDITOR
-	URigVMController* Controller;
+	URigVMController* Controller = nullptr;
+	IRigVMClientHost* Host = nullptr;
 #endif
 	FName NodeName = NAME_None;
 	ERigVMNodeCreatedReason Reason = ERigVMNodeCreatedReason::Unknown;
diff --git a/Engine/Plugins/Runtime/RigVM/Source/RigVMDeveloper/Private/RigVMModel/RigVMController.cpp b/Engine/Plugins/Runtime/RigVM/Source/RigVMDeveloper/Private/RigVMModel/RigVMController.cpp
--- a/Engine/Plugins/Runtime/RigVM/Source/RigVMDeveloper/Private/RigVMModel/RigVMController.cpp
+++ b/Engine/Plugins/Runtime/RigVM/Source/RigVMDeveloper/Private/RigVMModel/RigVMController.cpp
@@ -4371,7 +4371,7 @@ TArray<FName> URigVMController::ImportNodesFromText(const FString& InText, bool
 		}
 	}
 
-	FRigVMUnitNodeCreatedContext::FScope UnitNodeCreatedScope(UnitNodeCreatedContext, ERigVMNodeCreatedReason::Paste);
+	FRigVMUnitNodeCreatedContext::FScope UnitNodeCreatedScope(UnitNodeCreatedContext, ERigVMNodeCreatedReason::Paste, GetClientHost());
 	for (URigVMNode* CreatedNode : Factory.CreatedNodes)
 	{
 		URigVMGraph* Graph = CreatedNode->GetTypedOuter<URigVMGraph>();
diff --git a/Engine/Plugins/Runtime/RigVM/Source/RigVMEditor/Private/EdGraph/NodeSpawners/RigVMEdGraphUnitNodeSpawner.cpp b/Engine/Plugins/Runtime/RigVM/Source/RigVMEditor/Private/EdGraph/NodeSpawners/RigVMEdGraphUnitNodeSpawner.cpp
--- a/Engine/Plugins/Runtime/RigVM/Source/RigVMEditor/Private/EdGraph/NodeSpawners/RigVMEdGraphUnitNodeSpawner.cpp
+++ b/Engine/Plugins/Runtime/RigVM/Source/RigVMEditor/Private/EdGraph/NodeSpawners/RigVMEdGraphUnitNodeSpawner.cpp
@@ -252,7 +252,7 @@ URigVMEdGraphNode* URigVMEdGraphUnitNodeSpawner::SpawnNode(UEdGraph* ParentGraph
 		Controller->OpenUndoBracket(FString::Printf(TEXT("Add '%s' Node"), *Name.ToString()));
 
 		FRigVMUnitNodeCreatedContext& UnitNodeCreatedContext = Controller->GetUnitNodeCreatedContext();
-		FRigVMUnitNodeCreatedContext::FScope ReasonScope(UnitNodeCreatedContext, ERigVMNodeCreatedReason::NodeSpawner);
+		FRigVMUnitNodeCreatedContext::FScope ReasonScope(UnitNodeCreatedContext, ERigVMNodeCreatedReason::NodeSpawner, RigBlueprint);
 
 		if (URigVMUnitNode* ModelNode = Controller->AddUnitNode(StructTemplate, InMethodName, Location, Name.ToString(), true, true))
 		{
