Control Rig Physics Ensure that the "DisableCollision" flag in the joint is honored

diff --git a/Engine/Plugins/Experimental/PhysicsControl/Source/ControlRigPhysics/Private/RigPhysicsSimulation.cpp b/Engine/Plugins/Experimental/PhysicsControl/Source/ControlRigPhysics/Private/RigPhysicsSimulation.cpp
--- a/Engine/Plugins/Experimental/PhysicsControl/Source/ControlRigPhysics/Private/RigPhysicsSimulation.cpp
+++ b/Engine/Plugins/Experimental/PhysicsControl/Source/ControlRigPhysics/Private/RigPhysicsSimulation.cpp
@@ -278,6 +278,8 @@ void FRigPhysicsSimulation::InstantiatePhysicsJoints(const FRigPhysicsSolverComp
 	const FRigPhysicsSolverSettings& SolverSettings = SolverComponent->SolverSettings;
 	const FRigPhysicsSimulationSpaceSettings& SimulationSpaceSettings = SolverComponent->SimulationSpaceSettings;
 
+	TArray<ImmediatePhysics::FSimulation::FIgnorePair> IgnorePairs;
+
 	// Once all the bodies are created, we can make their physics joints
 	for (TPair<FRigComponentKey, FRigBodyRecord>& BodyRecordPair : BodyRecords)
 	{
@@ -464,11 +466,21 @@ void FRigPhysicsSimulation::InstantiatePhysicsJoints(const FRigPhysicsSolverComp
 						Chaos::FVec3(LinearSpring), Chaos::FVec3(LinearDamping), Chaos::FVec3(MaxForce),
 						Chaos::FVec3(AngularSpring), Chaos::FVec3(AngularDamping), Chaos::FVec3(MaxTorque));
 
+					if (!Settings.bCollisionEnabled)
+					{
+						// Immediate physics doesn't actually honor the flag in settings, so do it manually
+						ImmediatePhysics::FSimulation::FIgnorePair Pair;
+						Pair.A = ParentActorHandle;
+						Pair.B = ChildActorHandle;
+						IgnorePairs.Add(Pair);
+					}
+
 					// The target will be set in the pre-physics update
 				}
 			}
 		}
 	}
+	Simulation->SetIgnoreCollisionPairTable(IgnorePairs);
 }
 
 //======================================================================================================================
