Control Rig Physics - ensure the Drive settings are included in the physics joint archive.

diff --git a/Engine/Plugins/Experimental/PhysicsControl/Source/ControlRigPhysics/Private/RigPhysicsData.cpp b/Engine/Plugins/Experimental/PhysicsControl/Source/ControlRigPhysics/Private/RigPhysicsData.cpp
new file mode 100644
--- /dev/null
+++ b/Engine/Plugins/Experimental/PhysicsControl/Source/ControlRigPhysics/Private/RigPhysicsData.cpp
@@ -0,0 +1,98 @@
+// Copyright Epic Games, Inc. All Rights Reserved.
+
+#include "RigPhysicsData.h"
+#include "PhysicsControlObjectVersion.h"
+
+FArchive& operator <<(FArchive& Ar, FRigPhysicsJoint& Data)
+{
+	Ar << Data.bEnable;
+	Ar << Data.ParentBody;
+	Ar << Data.bAutoCalculateParentOffset;
+	Ar << Data.bAutoCalculateChildOffset;
+	Ar << Data.ExtraParentOffset;
+	Ar << Data.ExtraChildOffset;
+	Ar << Data.LinearLimit;
+	Ar << Data.AngularLimit;
+	Ar << Data.bDisableCollision;
+	Ar << Data.LinearProjectionAmount;
+	Ar << Data.AngularProjectionAmount;
+	Ar << Data.ParentInverseMassScale;
+	if (Ar.CustomVer(FPhysicsControlObjectVersion::GUID) >= FPhysicsControlObjectVersion::IncludeDriveInJoint)
+	{
+		Ar << Data.Drive;
+	}
+	return Ar;
+}
+
+FArchive& operator <<(FArchive& Ar, FRigPhysicsSimulationSpaceSettings& Data)
+{
+	Ar << Data.SpaceMovementAmount;
+	Ar << Data.VelocityScaleZ;
+	Ar << Data.bClampLinearVelocity;
+	Ar << Data.MaxLinearVelocity;
+	Ar << Data.bClampAngularVelocity;
+	Ar << Data.MaxAngularVelocity;
+	Ar << Data.bClampLinearAcceleration;
+	Ar << Data.MaxLinearAcceleration;
+	Ar << Data.bClampAngularAcceleration;
+	Ar << Data.MaxAngularAcceleration;
+	Ar << Data.LinearAccelerationThresholdForTeleport;
+	Ar << Data.AngularAccelerationThresholdForTeleport;
+	Ar << Data.ExternalLinearDragV;
+	Ar << Data.ExternalLinearVelocity;
+	Ar << Data.ExternalAngularVelocity;
+	return Ar;
+}
+
+FArchive& operator <<(FArchive& Ar, FRigPhysicsSolverSettings& Data)
+{
+	Ar << Data.SimulationSpace;
+	Ar << Data.CollisionSpace;
+	Ar << Data.SpaceBone;
+	Ar << Data.Collision;
+	Ar << Data.Gravity;
+	Ar << Data.PositionIterations;
+	Ar << Data.VelocityIterations;
+	Ar << Data.ProjectionIterations;
+	Ar << Data.MaxNumRollingAverageStepTimes;
+	Ar << Data.MaxDepenetrationVelocity;
+	Ar << Data.FixedTimeStep;
+	Ar << Data.MaxTimeSteps;
+	Ar << Data.MaxDeltaTime;
+	Ar << Data.bUseLinearJointSolver;
+	Ar << Data.bSolveJointPositionsLast;
+	Ar << Data.bUseManifolds;
+	Ar << Data.DistanceThresholdForReset;
+	return Ar;
+}
+
+FArchive& operator <<(FArchive& Ar, FRigPhysicsDrive& Data)
+{
+	Ar << Data.bEnable;
+	Ar << Data.LinearStrength;
+	Ar << Data.LinearDampingRatio;
+	Ar << Data.LinearExtraDamping;
+	Ar << Data.MaxForce;
+	Ar << Data.AngularStrength;
+	Ar << Data.AngularDampingRatio;
+	Ar << Data.AngularExtraDamping;
+	Ar << Data.MaxTorque;
+	Ar << Data.SkeletalAnimationVelocityMultiplier;
+	return Ar;
+}
+
+FArchive& operator <<(FArchive& Ar, FRigPhysicsBodySolverSettings& Data)
+{
+	Ar << Data.PhysicsSolverComponentKey;
+	Ar << Data.TargetBone;
+	Ar << Data.SourceBone;
+	return Ar;
+}
+
+FArchive& operator <<(FArchive& Ar, FRigPhysicsCollision& Data)
+{
+	Ar << Data.Boxes;
+	Ar << Data.Spheres;
+	Ar << Data.Capsules;
+	return Ar;
+}
diff --git a/Engine/Plugins/Experimental/PhysicsControl/Source/ControlRigPhysics/Public/RigPhysicsData.h b/Engine/Plugins/Experimental/PhysicsControl/Source/ControlRigPhysics/Public/RigPhysicsData.h
--- a/Engine/Plugins/Experimental/PhysicsControl/Source/ControlRigPhysics/Public/RigPhysicsData.h
+++ b/Engine/Plugins/Experimental/PhysicsControl/Source/ControlRigPhysics/Public/RigPhysicsData.h
@@ -183,13 +183,7 @@ struct CONTROLRIGPHYSICS_API FRigPhysicsCollision
 {
 	GENERATED_BODY()
 
-	friend FArchive& operator <<(FArchive& Ar, FRigPhysicsCollision& Data)
-	{
-		Ar << Data.Boxes;
-		Ar << Data.Spheres;
-		Ar << Data.Capsules;
-		return Ar;
-	}
+	friend FArchive& operator <<(FArchive& Ar, FRigPhysicsCollision& Data);
 
 	UPROPERTY(BlueprintReadWrite, EditAnywhere, Category = Shapes)
 	TArray<FRigPhysicsCollisionBox> Boxes;
@@ -212,25 +206,7 @@ struct FRigPhysicsSimulationSpaceSettings
 {
 	GENERATED_BODY()
 
-	friend FArchive& operator <<(FArchive& Ar, FRigPhysicsSimulationSpaceSettings& Data)
-	{
-		Ar << Data.SpaceMovementAmount;
-		Ar << Data.VelocityScaleZ;
-		Ar << Data.bClampLinearVelocity;
-		Ar << Data.MaxLinearVelocity;
-		Ar << Data.bClampAngularVelocity;
-		Ar << Data.MaxAngularVelocity;
-		Ar << Data.bClampLinearAcceleration;
-		Ar << Data.MaxLinearAcceleration;
-		Ar << Data.bClampAngularAcceleration;
-		Ar << Data.MaxAngularAcceleration;
-		Ar << Data.LinearAccelerationThresholdForTeleport;
-		Ar << Data.AngularAccelerationThresholdForTeleport;
-		Ar << Data.ExternalLinearDragV;
-		Ar << Data.ExternalLinearVelocity;
-		Ar << Data.ExternalAngularVelocity;
-		return Ar;
-	}
+	friend FArchive& operator <<(FArchive& Ar, FRigPhysicsSimulationSpaceSettings& Data);
 
 	// Overall multiplier on the effects of simulation space movement on the simulation
 	UPROPERTY(EditAnywhere, BlueprintReadWrite, Category = Settings, meta = (ClampMin = "0.0", ClampMax = "1.0"))
@@ -333,27 +309,7 @@ struct CONTROLRIGPHYSICS_API FRigPhysicsSolverSettings
 {
 	GENERATED_BODY()
 
-	friend FArchive& operator <<(FArchive& Ar, FRigPhysicsSolverSettings& Data)
-	{
-		Ar << Data.SimulationSpace;
-		Ar << Data.CollisionSpace;
-		Ar << Data.SpaceBone;
-		Ar << Data.Collision;
-		Ar << Data.Gravity;
-		Ar << Data.PositionIterations;
-		Ar << Data.VelocityIterations;
-		Ar << Data.ProjectionIterations;
-		Ar << Data.MaxNumRollingAverageStepTimes;
-		Ar << Data.MaxDepenetrationVelocity;
-		Ar << Data.FixedTimeStep;
-		Ar << Data.MaxTimeSteps;
-		Ar << Data.MaxDeltaTime;
-		Ar << Data.bUseLinearJointSolver;
-		Ar << Data.bSolveJointPositionsLast;
-		Ar << Data.bUseManifolds;
-		Ar << Data.DistanceThresholdForReset;
-		return Ar;
-	}
+	friend FArchive& operator <<(FArchive& Ar, FRigPhysicsSolverSettings& Data);
 
 	UPROPERTY(BlueprintReadOnly, EditAnywhere, Category = SolverSettings)
 	ERigPhysicsSimulationSpace SimulationSpace = ERigPhysicsSimulationSpace::Component;
@@ -500,20 +456,7 @@ struct CONTROLRIGPHYSICS_API FRigPhysicsDrive
 {
 	GENERATED_BODY()
 
-	friend FArchive& operator <<(FArchive& Ar, FRigPhysicsDrive& Data)
-	{
-		Ar << Data.bEnable;
-		Ar << Data.LinearStrength;
-		Ar << Data.LinearDampingRatio;
-		Ar << Data.LinearExtraDamping;
-		Ar << Data.MaxForce;
-		Ar << Data.AngularStrength;
-		Ar << Data.AngularDampingRatio;
-		Ar << Data.AngularExtraDamping;
-		Ar << Data.MaxTorque;
-		Ar << Data.SkeletalAnimationVelocityMultiplier;
-		return Ar;
-	}
+	friend FArchive& operator <<(FArchive& Ar, FRigPhysicsDrive& Data);
 
 	UPROPERTY(BlueprintReadWrite, EditAnywhere, Category = PhysicsDrive)
 	bool bEnable = true;
@@ -591,22 +534,7 @@ struct CONTROLRIGPHYSICS_API FRigPhysicsJoint
 {
 	GENERATED_BODY()
 
-	friend FArchive& operator <<(FArchive& Ar, FRigPhysicsJoint& Data)
-	{
-		Ar << Data.bEnable;
-		Ar << Data.ParentBody;
-		Ar << Data.bAutoCalculateParentOffset;
-		Ar << Data.bAutoCalculateChildOffset;
-		Ar << Data.ExtraParentOffset;
-		Ar << Data.ExtraChildOffset;
-		Ar << Data.LinearLimit;
-		Ar << Data.AngularLimit;
-		Ar << Data.bDisableCollision;
-		Ar << Data.LinearProjectionAmount;
-		Ar << Data.AngularProjectionAmount;
-		Ar << Data.ParentInverseMassScale;
-		return Ar;
-	}
+	friend FArchive& operator<<(FArchive& Ar, FRigPhysicsJoint& Data);
 
 	UPROPERTY(BlueprintReadWrite, EditAnywhere, Category = PhysicsJoint)
 	bool bEnable = true;
@@ -681,13 +609,7 @@ struct CONTROLRIGPHYSICS_API FRigPhysicsBodySolverSettings
 		FRigComponentKey InPhysicsSolverComponentKey = FRigComponentKey(), FRigElementKey InTargetBone = FRigElementKey())
 		: PhysicsSolverComponentKey(InPhysicsSolverComponentKey), TargetBone(InTargetBone) {}
 
-	friend FArchive& operator <<(FArchive& Ar, FRigPhysicsBodySolverSettings& Data)
-	{
-		Ar << Data.PhysicsSolverComponentKey;
-		Ar << Data.TargetBone;
-		Ar << Data.SourceBone;
-		return Ar;
-	}
+	friend FArchive& operator <<(FArchive& Ar, FRigPhysicsBodySolverSettings& Data);
 
 	UPROPERTY(BlueprintReadWrite, EditAnywhere, Category = Solver)
 	FRigComponentKey PhysicsSolverComponentKey;
diff --git a/Engine/Plugins/Experimental/PhysicsControl/Source/PhysicsControl/Public/PhysicsControlObjectVersion.h b/Engine/Plugins/Experimental/PhysicsControl/Source/PhysicsControl/Public/PhysicsControlObjectVersion.h
--- a/Engine/Plugins/Experimental/PhysicsControl/Source/PhysicsControl/Public/PhysicsControlObjectVersion.h
+++ b/Engine/Plugins/Experimental/PhysicsControl/Source/PhysicsControl/Public/PhysicsControlObjectVersion.h
@@ -10,6 +10,8 @@ struct PHYSICSCONTROL_API FPhysicsControlObjectVersion
 		// Before any version changes were made
 		BeforeCustomVersionWasAdded,
 
+		IncludeDriveInJoint,
+
 		// -----<new versions can be added above this line>-------------------------------------------------
 		VersionPlusOne,
 		LatestVersion = VersionPlusOne - 1
