ControlRigPhysics: Add icons for components

#rnx

diff --git a/Engine/Plugins/Animation/ControlRig/Source/ControlRigEditor/Private/Editor/SRigHierarchy.cpp b/Engine/Plugins/Animation/ControlRig/Source/ControlRigEditor/Private/Editor/SRigHierarchy.cpp
--- a/Engine/Plugins/Animation/ControlRig/Source/ControlRigEditor/Private/Editor/SRigHierarchy.cpp
+++ b/Engine/Plugins/Animation/ControlRig/Source/ControlRigEditor/Private/Editor/SRigHierarchy.cpp
@@ -906,6 +906,22 @@ void SRigHierarchy::OnHierarchyModified(ERigHierarchyNotification InNotif, URigH
 			}
 			break;
 		}
+		case ERigHierarchyNotification::ComponentContentChanged:
+		{
+			if(InComponent)
+			{
+				for (int32 RootIndex = 0; RootIndex < TreeView->RootElements.Num(); ++RootIndex)
+				{
+					TSharedPtr<FRigTreeElement> TreeElement = TreeView->FindElement(InComponent->GetKey(), TreeView->RootElements[RootIndex]);
+					if (TreeElement.IsValid())
+					{
+						const FRigTreeDisplaySettings& Settings = TreeView->GetRigTreeDelegates().GetDisplaySettings();
+						TreeElement->RefreshDisplaySettings(InHierarchy, Settings);
+					}
+				}
+			}
+			break;
+		}
 		default:
 		{
 			break;
diff --git a/Engine/Plugins/Animation/ControlRig/Source/ControlRigEditor/Private/Editor/SRigHierarchyTreeView.cpp b/Engine/Plugins/Animation/ControlRig/Source/ControlRigEditor/Private/Editor/SRigHierarchyTreeView.cpp
--- a/Engine/Plugins/Animation/ControlRig/Source/ControlRigEditor/Private/Editor/SRigHierarchyTreeView.cpp
+++ b/Engine/Plugins/Animation/ControlRig/Source/ControlRigEditor/Private/Editor/SRigHierarchyTreeView.cpp
@@ -1338,8 +1338,10 @@ TPair<const FSlateBrush*, FSlateColor> SRigHierarchyItem::GetBrushForElementType
 	}
 	else if(InKey.IsComponent())
 	{
-		// todo: offer virtuals on the components to customize their look.
-		Brush = PhysicsBrush; // for now
+		if(const FRigBaseComponent* Component = InHierarchy->FindComponent(InKey.GetComponent()))
+		{
+			Brush = Component->GetIconForUI().GetIcon();
+		}
 	}
 
 	return TPair<const FSlateBrush*, FSlateColor>(Brush, Color);
diff --git a/Engine/Plugins/Experimental/PhysicsControl/Source/ControlRigPhysics/ControlRigPhysics.Build.cs b/Engine/Plugins/Experimental/PhysicsControl/Source/ControlRigPhysics/ControlRigPhysics.Build.cs
--- a/Engine/Plugins/Experimental/PhysicsControl/Source/ControlRigPhysics/ControlRigPhysics.Build.cs
+++ b/Engine/Plugins/Experimental/PhysicsControl/Source/ControlRigPhysics/ControlRigPhysics.Build.cs
@@ -37,6 +37,9 @@ public class ControlRigPhysics : ModuleRules
 				{
 					"RigVMDeveloper",
 					"ControlRigDeveloper",
+					"Slate",
+					"SlateCore",
+					"EditorStyle",
 				}
 			);
 		}
diff --git a/Engine/Plugins/Experimental/PhysicsControl/Source/ControlRigPhysics/Private/ControlRigPhysicsEditorStyle.h b/Engine/Plugins/Experimental/PhysicsControl/Source/ControlRigPhysics/Private/ControlRigPhysicsEditorStyle.h
new file mode 100644
--- /dev/null
+++ b/Engine/Plugins/Experimental/PhysicsControl/Source/ControlRigPhysics/Private/ControlRigPhysicsEditorStyle.h
@@ -0,0 +1,49 @@
+// Copyright Epic Games, Inc. All Rights Reserved.
+
+#pragma once
+
+#include "Styling/CoreStyle.h"
+#include "Styling/SlateStyle.h"
+#include "Styling/SlateStyleRegistry.h"
+#include "Styling/SlateStyleMacros.h"
+#include "Styling/SlateTypes.h"
+#include "Styling/StyleColors.h"
+#include "Brushes/SlateRoundedBoxBrush.h"
+
+class FControlRigPhysicsEditorStyle final
+	: public FSlateStyleSet
+{
+public:
+	FControlRigPhysicsEditorStyle()
+		: FSlateStyleSet("ControlRigPhysicsEditorStyle")
+	{
+		const FVector2D Icon16x16(16.0f, 16.0f);
+
+		const FString RigPluginContentDir = FPaths::EnginePluginsDir() / TEXT("Experimental/PhysicsControl/Content");
+		SetContentRoot(RigPluginContentDir);
+
+		// Component Icons
+		{
+			Set("ControlRigPhysics.Component.Solver", new IMAGE_BRUSH_SVG("Slate/Solver_16", Icon16x16));
+			Set("ControlRigPhysics.Component.BodyMultipleDefault", new IMAGE_BRUSH_SVG("Slate/BodyMultipleDefault_16", Icon16x16));
+			Set("ControlRigPhysics.Component.BodyMultipleKinematic", new IMAGE_BRUSH_SVG("Slate/BodyMultipleKinematic_16", Icon16x16));
+			Set("ControlRigPhysics.Component.BodyMultipleSimulated", new IMAGE_BRUSH_SVG("Slate/BodyMultipleSimulated_16", Icon16x16));
+			Set("ControlRigPhysics.Component.BodySingleDefault", new IMAGE_BRUSH_SVG("Slate/BodySingleDefault_16", Icon16x16));
+			Set("ControlRigPhysics.Component.BodySingleKinematic", new IMAGE_BRUSH_SVG("Slate/BodySingleKinematic_16", Icon16x16));
+			Set("ControlRigPhysics.Component.BodySingleSimulated", new IMAGE_BRUSH_SVG("Slate/BodySingleSimulated_16", Icon16x16));
+		}
+		
+		FSlateStyleRegistry::RegisterSlateStyle(*this);
+	}
+
+	static FControlRigPhysicsEditorStyle& Get()
+	{
+		static FControlRigPhysicsEditorStyle Inst;
+		return Inst;
+	}
+	
+	~FControlRigPhysicsEditorStyle()
+	{
+		FSlateStyleRegistry::UnRegisterSlateStyle(*this);
+	}
+};
diff --git a/Engine/Plugins/Experimental/PhysicsControl/Source/ControlRigPhysics/Private/ControlRigPhysicsModule.cpp b/Engine/Plugins/Experimental/PhysicsControl/Source/ControlRigPhysics/Private/ControlRigPhysicsModule.cpp
--- a/Engine/Plugins/Experimental/PhysicsControl/Source/ControlRigPhysics/Private/ControlRigPhysicsModule.cpp
+++ b/Engine/Plugins/Experimental/PhysicsControl/Source/ControlRigPhysics/Private/ControlRigPhysicsModule.cpp
@@ -1,6 +1,9 @@
 // Copyright Epic Games, Inc. All Rights Reserved.
 
 #include "ControlRigPhysicsModule.h"
+#if WITH_EDITOR
+#include "ControlRigPhysicsEditorStyle.h"
+#endif
 
 #define LOCTEXT_NAMESPACE "FControlRigPhysicsModule"
 
@@ -9,6 +12,10 @@ DEFINE_LOG_CATEGORY(LogControlRigPhysics);
 //======================================================================================================================
 void FControlRigPhysicsModule::StartupModule()
 {
+#if WITH_EDITOR
+	// register the editor style
+	FControlRigPhysicsEditorStyle::Get();
+#endif
 }
 
 //======================================================================================================================
diff --git a/Engine/Plugins/Experimental/PhysicsControl/Source/ControlRigPhysics/Private/RigPhysicsComponent.cpp b/Engine/Plugins/Experimental/PhysicsControl/Source/ControlRigPhysics/Private/RigPhysicsComponent.cpp
--- a/Engine/Plugins/Experimental/PhysicsControl/Source/ControlRigPhysics/Private/RigPhysicsComponent.cpp
+++ b/Engine/Plugins/Experimental/PhysicsControl/Source/ControlRigPhysics/Private/RigPhysicsComponent.cpp
@@ -6,6 +6,9 @@
 #include "Rigs/RigHierarchy.h"
 #include "Rigs/RigHierarchyController.h"
 #include "ControlRigPhysicsModule.h"
+#if WITH_EDITOR
+#include "ControlRigPhysicsEditorStyle.h"
+#endif
 
 //======================================================================================================================
 void FRigPhysicsComponent::Save(FArchive& Ar)
@@ -37,6 +40,33 @@ void FRigPhysicsComponent::Load(FArchive& Ar)
 	Ar << CurrentBodyData;
 }
 
+//======================================================================================================================
+#if WITH_EDITOR
+const FSlateIcon& FRigPhysicsComponent::GetIconForUI() const
+{
+	static const FSlateIcon BodyMultipleDefaultIcon = FSlateIcon(FControlRigPhysicsEditorStyle::Get().GetStyleSetName(), "ControlRigPhysics.Component.BodyMultipleDefault");
+	static const FSlateIcon BodyMultipleKinematicIcon = FSlateIcon(FControlRigPhysicsEditorStyle::Get().GetStyleSetName(), "ControlRigPhysics.Component.BodyMultipleKinematic");
+	static const FSlateIcon BodyMultipleSimulatedIcon = FSlateIcon(FControlRigPhysicsEditorStyle::Get().GetStyleSetName(), "ControlRigPhysics.Component.BodyMultipleSimulated");
+	static const FSlateIcon BodySingleDefaultIcon = FSlateIcon(FControlRigPhysicsEditorStyle::Get().GetStyleSetName(), "ControlRigPhysics.Component.BodySingleDefault");
+	static const FSlateIcon BodySingleKinematicIcon = FSlateIcon(FControlRigPhysicsEditorStyle::Get().GetStyleSetName(), "ControlRigPhysics.Component.BodySingleKinematic");
+	static const FSlateIcon BodySingleSimulatedIcon = FSlateIcon(FControlRigPhysicsEditorStyle::Get().GetStyleSetName(), "ControlRigPhysics.Component.BodySingleSimulated");
+
+	const int32 NumShapes = Collision.NumShapes();;
+	if(NumShapes == 0)
+	{
+		return BodySingleDefaultIcon;
+	}
+	
+	if(BodyData.MovementType == EPhysicsMovementType::Simulated)
+	{
+		return NumShapes <= 1 ? BodySingleSimulatedIcon : BodyMultipleSimulatedIcon;
+	}
+
+	// todo: do we want another set of icons for the static movement type?
+	return NumShapes <= 1 ? BodySingleKinematicIcon : BodyMultipleKinematicIcon;
+}
+#endif
+
 //======================================================================================================================
 bool FRigPhysicsComponent::CanBeAddedTo(const FRigElementKey& InElementKey, const URigHierarchy* InHierarchy, FString* OutFailureReason) const
 {
diff --git a/Engine/Plugins/Experimental/PhysicsControl/Source/ControlRigPhysics/Private/RigPhysicsSolverComponent.cpp b/Engine/Plugins/Experimental/PhysicsControl/Source/ControlRigPhysics/Private/RigPhysicsSolverComponent.cpp
--- a/Engine/Plugins/Experimental/PhysicsControl/Source/ControlRigPhysics/Private/RigPhysicsSolverComponent.cpp
+++ b/Engine/Plugins/Experimental/PhysicsControl/Source/ControlRigPhysics/Private/RigPhysicsSolverComponent.cpp
@@ -1,8 +1,10 @@
 // Copyright Epic Games, Inc. All Rights Reserved.
 
 #include "RigPhysicsSolverComponent.h"
-
 #include "Rigs/RigHierarchy.h"
+#if WITH_EDITOR
+#include "ControlRigPhysicsEditorStyle.h"
+#endif
 
 //======================================================================================================================
 void FRigPhysicsSolverComponent::Save(FArchive& Ar)
@@ -19,3 +21,12 @@ void FRigPhysicsSolverComponent::Load(FArchive& Ar)
 	Ar << SolverSettings;
 	Ar << SimulationSpaceSettings;
 }
+
+//======================================================================================================================
+#if WITH_EDITOR
+const FSlateIcon& FRigPhysicsSolverComponent::GetIconForUI() const
+{
+	static const FSlateIcon SolverIcon = FSlateIcon(FControlRigPhysicsEditorStyle::Get().GetStyleSetName(), "ControlRigPhysics.Component.Solver");
+	return SolverIcon;
+}
+#endif
\ No newline at end of file
diff --git a/Engine/Plugins/Experimental/PhysicsControl/Source/ControlRigPhysics/Public/RigPhysicsComponent.h b/Engine/Plugins/Experimental/PhysicsControl/Source/ControlRigPhysics/Public/RigPhysicsComponent.h
--- a/Engine/Plugins/Experimental/PhysicsControl/Source/ControlRigPhysics/Public/RigPhysicsComponent.h
+++ b/Engine/Plugins/Experimental/PhysicsControl/Source/ControlRigPhysics/Public/RigPhysicsComponent.h
@@ -59,6 +59,10 @@ public:
 	virtual void Save(FArchive& Ar) override;
 	virtual void Load(FArchive& Ar) override;
 
+#if WITH_EDITOR
+	virtual const FSlateIcon& GetIconForUI() const override;
+#endif
+
 	virtual bool CanBeAddedTo(const FRigElementKey& InElementKey, const URigHierarchy* InHierarchy, FString* OutFailureReason) const override;
 
 	virtual void OnAddedToHierarchy(URigHierarchy* InHierarchy, URigHierarchyController* InController) override;
diff --git a/Engine/Plugins/Experimental/PhysicsControl/Source/ControlRigPhysics/Public/RigPhysicsData.h b/Engine/Plugins/Experimental/PhysicsControl/Source/ControlRigPhysics/Public/RigPhysicsData.h
--- a/Engine/Plugins/Experimental/PhysicsControl/Source/ControlRigPhysics/Public/RigPhysicsData.h
+++ b/Engine/Plugins/Experimental/PhysicsControl/Source/ControlRigPhysics/Public/RigPhysicsData.h
@@ -195,6 +195,7 @@ struct CONTROLRIGPHYSICS_API FRigPhysicsCollision
 	TArray<FRigPhysicsCollisionCapsule> Capsules;
 
 	bool IsEmpty() const { return Boxes.IsEmpty() && Spheres.IsEmpty() && Capsules.IsEmpty(); }
+	int32 NumShapes() const { return Boxes.Num() + Capsules.Num() + Spheres.Num(); }
 };
 
 /**
diff --git a/Engine/Plugins/Experimental/PhysicsControl/Source/ControlRigPhysics/Public/RigPhysicsSolverComponent.h b/Engine/Plugins/Experimental/PhysicsControl/Source/ControlRigPhysics/Public/RigPhysicsSolverComponent.h
--- a/Engine/Plugins/Experimental/PhysicsControl/Source/ControlRigPhysics/Public/RigPhysicsSolverComponent.h
+++ b/Engine/Plugins/Experimental/PhysicsControl/Source/ControlRigPhysics/Public/RigPhysicsSolverComponent.h
@@ -27,5 +27,9 @@ public:
 
 	virtual void Save(FArchive& Ar) override;
 	virtual void Load(FArchive& Ar) override;
+
+#if WITH_EDITOR
+	virtual const FSlateIcon& GetIconForUI() const override;
+#endif
 };
 
