Control Rig Physics - set transforms in the hierarchy correctly

diff --git a/Engine/Plugins/Experimental/PhysicsControl/Source/ControlRigPhysics/Private/RigPhysicsSimulation.cpp b/Engine/Plugins/Experimental/PhysicsControl/Source/ControlRigPhysics/Private/RigPhysicsSimulation.cpp
--- a/Engine/Plugins/Experimental/PhysicsControl/Source/ControlRigPhysics/Private/RigPhysicsSimulation.cpp
+++ b/Engine/Plugins/Experimental/PhysicsControl/Source/ControlRigPhysics/Private/RigPhysicsSimulation.cpp
@@ -36,6 +36,8 @@
 #include "PhysicsEngine/ConstraintInstance.h"
 #include "PhysicsEngine/BodySetup.h"
 
+#include "Logging/LogMacros.h"
+
 DEFINE_LOG_CATEGORY(LogRigPhysics);
 
 //======================================================================================================================
@@ -69,6 +71,10 @@ void FRigPhysicsSimulation::InitialiseSimulation(const FRigPhysicsSolverComponen
 	FRigComponentKey SolverComponentKey = SolverComponent->GetKey();
 
 	TArray<FRigComponentKey> AllComponentKeys = Hierarchy->GetAllComponentKeys();
+
+	// All the components in this simulation
+	TArray<FRigComponentKey> UnsortedBodyComponentKeys;
+
 	for (FRigComponentKey RigComponentKey : AllComponentKeys)
 	{
 		if (const FRigPhysicsComponent* PhysicsComponent = Cast<FRigPhysicsComponent>(
@@ -78,9 +84,24 @@ void FRigPhysicsSimulation::InitialiseSimulation(const FRigPhysicsSolverComponen
 			{
 				// Just make the record for now - it will be instantiated later
 				BodyRecords.Add(RigComponentKey);
+				UnsortedBodyComponentKeys.Add(RigComponentKey);
 			}
 		}
 	}
+
+	// Sort the component keys according to the traversal of their element (i.e. from root to leaf)
+	SortedBodyComponentKeys.Empty(UnsortedBodyComponentKeys.Num());
+	Hierarchy->Traverse([this, &UnsortedBodyComponentKeys](FRigBaseElement* Element, bool& bContinue) 
+		{
+			const FRigElementKey& Key = Element->GetKey();
+			for (const FRigComponentKey& ComponentKey : UnsortedBodyComponentKeys)
+			{
+				if (ComponentKey.ElementKey == Key)
+				{
+					SortedBodyComponentKeys.Push(ComponentKey);
+				}
+			}
+		});
 }
 
 //======================================================================================================================
@@ -229,7 +250,6 @@ ImmediatePhysics::FActorHandle* FRigPhysicsSimulation::CreateBody(
 	return ActorHandle;
 }
 
-
 //======================================================================================================================
 void FRigPhysicsSimulation::Instantiate(const FRigPhysicsSolverComponent* SolverComponent)
 {
@@ -523,10 +543,10 @@ void FRigPhysicsSimulation::InstantiatePhysicsBodies(const FRigPhysicsSolverComp
 
 //======================================================================================================================
 void FRigPhysicsSimulation::UpdatePrePhysics(
-	const FRigVMExecuteContext&     ExecuteContext, 
+	const FRigVMExecuteContext&       ExecuteContext, 
 	const FRigPhysicsSolverComponent* SolverComponent, 
-	const float                     DeltaTime, 
-	const FSimulationSpaceData&     SimulationSpaceData)
+	const float                       DeltaTime, 
+	const FSimulationSpaceData&       SimulationSpaceData)
 {
 	URigHierarchy* Hierarchy = OwningControlRig->GetHierarchy();
 
@@ -536,7 +556,6 @@ void FRigPhysicsSimulation::UpdatePrePhysics(
 			SolverComponent->SolverSettings, FTransform());
 		ActorHandle->SetKinematicTarget(BodyRelSimSpaceTM);
 	}
-
 	for (TPair<FRigComponentKey, FRigBodyRecord>& BodyRecordPair : BodyRecords)
 	{
 		const FRigComponentKey& ComponentKey = BodyRecordPair.Key;
@@ -720,14 +739,13 @@ void FRigPhysicsSimulation::UpdatePostPhysics(const FRigPhysicsSolverComponent*
 		return;
 	}
 	const FRigPhysicsSolverSettings& SolverSettings = SolverComponent->SolverSettings;
-
 	bool bGotInvalidSimulationData = false;
 	URigHierarchy* Hierarchy = OwningControlRig->GetHierarchy();
-	for (TPair<FRigComponentKey, FRigBodyRecord>& BodyRecordPair : BodyRecords)
-	{
-		const FRigComponentKey& ComponentKey = BodyRecordPair.Key;
-		const FRigBodyRecord& Record = BodyRecordPair.Value;
 
+	// Traverse using the sorted keys
+	for (const FRigComponentKey& ComponentKey : SortedBodyComponentKeys)
+	{
+		const FRigBodyRecord& Record = *BodyRecords.Find(ComponentKey);
 		if (Record.ActorHandle)
 		{
 			if (const FRigPhysicsComponent* PhysicsComponent = Cast<FRigPhysicsComponent>(
@@ -768,8 +786,11 @@ void FRigPhysicsSimulation::UpdatePostPhysics(const FRigPhysicsSolverComponent*
 							TargetComponentSpaceTM.SetRotation(TargetQ);
 							TargetComponentSpaceTM.SetTranslation(TargetT);
 						}
-						// Danny TODO should we be "updating" the children here or not? 
-						Hierarchy->SetGlobalTransform(TargetElementKey, TargetComponentSpaceTM);
+						// Note that we set bAffectChildren = true (i.e. don't counter-animate
+						// children), so that attached animation bones will follow physics, but we
+						// rely on our bodies being sorted so we work out from the leaf nodes so as
+						// not to disturb previously set bodies.
+						Hierarchy->SetGlobalTransform(TargetElementKey, TargetComponentSpaceTM, false, true);
 					}
 				}
 			}
@@ -872,6 +893,12 @@ void FRigPhysicsSimulation::StepSimulation(
 		SimulationSpaceSettings.SpaceMovementAmount * SimulationSpaceData.LinearAcceleration,
 		SimulationSpaceSettings.SpaceMovementAmount * SimulationSpaceData.AngularAcceleration);
 
+	// Only update if there is a delta time:
+	// * We don't want to update our previous TMs and store the dt - because that would end up
+	//   implying infinite velocities
+	// * We don't to update kinematic bodies with the new TMs because, since the simulated ones
+	//   won't move, that would break the pose.
+	// * We can't actually simulate with dt = 0
 	if (PhysicsDeltaTime > 0.0f)
 	{
 		UpdatePrePhysics(ExecuteContext, SolverComponent, PhysicsDeltaTime, SimulationSpaceData);
diff --git a/Engine/Plugins/Experimental/PhysicsControl/Source/ControlRigPhysics/Private/RigPhysicsSimulation.h b/Engine/Plugins/Experimental/PhysicsControl/Source/ControlRigPhysics/Private/RigPhysicsSimulation.h
--- a/Engine/Plugins/Experimental/PhysicsControl/Source/ControlRigPhysics/Private/RigPhysicsSimulation.h
+++ b/Engine/Plugins/Experimental/PhysicsControl/Source/ControlRigPhysics/Private/RigPhysicsSimulation.h
@@ -170,7 +170,12 @@ private:
 	TObjectPtr<UControlRig> OwningControlRig;
 	FRigComponentKey        PhysicsSolverComponentKey;
 
+	// All the bodies, but in no particular order
 	TMap<FRigComponentKey, FRigBodyRecord> BodyRecords;
+
+	// Ordering so that we can traverse from root to leaf bones
+	TArray<FRigComponentKey> SortedBodyComponentKeys;
+
 	TSharedPtr<ImmediatePhysics::FSimulation> Simulation;
 
 	// Used to store things like collision shapes.
