Control Rig Physics - sort out the unwanted storing of "current" data in the physics component

diff --git a/Engine/Plugins/Experimental/PhysicsControl/Source/ControlRigPhysics/Private/RigPhysicsComponent.cpp b/Engine/Plugins/Experimental/PhysicsControl/Source/ControlRigPhysics/Private/RigPhysicsComponent.cpp
--- a/Engine/Plugins/Experimental/PhysicsControl/Source/ControlRigPhysics/Private/RigPhysicsComponent.cpp
+++ b/Engine/Plugins/Experimental/PhysicsControl/Source/ControlRigPhysics/Private/RigPhysicsComponent.cpp
@@ -21,9 +21,8 @@ void FRigPhysicsComponent::Save(FArchive& Ar)
 	Ar << Collision;
 	Ar << Joint;
 	Ar << BodyData;
-	Ar << CurrentKinematicTarget;
-	Ar << CurrentKinematicTargetSpace;
-	Ar << CurrentBodyData;
+	Ar << KinematicTarget;
+	Ar << KinematicTargetSpace;
 }
 
 //======================================================================================================================
@@ -35,9 +34,18 @@ void FRigPhysicsComponent::Load(FArchive& Ar)
 	Ar << Collision;
 	Ar << Joint;
 	Ar << BodyData;
-	Ar << CurrentKinematicTarget;
-	Ar << CurrentKinematicTargetSpace;
-	Ar << CurrentBodyData;
+	Ar << KinematicTarget;
+	Ar << KinematicTargetSpace;
+
+	if (Ar.CustomVer(FPhysicsControlObjectVersion::GUID) <
+		FPhysicsControlObjectVersion::ControlRigRemoveCurrentDataFromPhysicsComponent)
+	{
+		FPhysicsControlModifierData CurrentBodyData;
+		Ar << CurrentBodyData;
+		// Previously people needed to set things in the current data for initial properties, as
+		// this would override. So do the override here.
+		BodyData = CurrentBodyData;
+	}
 }
 
 //======================================================================================================================
diff --git a/Engine/Plugins/Experimental/PhysicsControl/Source/ControlRigPhysics/Private/RigPhysicsData.cpp b/Engine/Plugins/Experimental/PhysicsControl/Source/ControlRigPhysics/Private/RigPhysicsData.cpp
--- a/Engine/Plugins/Experimental/PhysicsControl/Source/ControlRigPhysics/Private/RigPhysicsData.cpp
+++ b/Engine/Plugins/Experimental/PhysicsControl/Source/ControlRigPhysics/Private/RigPhysicsData.cpp
@@ -17,7 +17,7 @@ FArchive& operator <<(FArchive& Ar, FRigPhysicsJoint& Data)
 	Ar << Data.LinearProjectionAmount;
 	Ar << Data.AngularProjectionAmount;
 	Ar << Data.ParentInverseMassScale;
-	if (Ar.CustomVer(FPhysicsControlObjectVersion::GUID) >= FPhysicsControlObjectVersion::IncludeDriveInJoint)
+	if (Ar.CustomVer(FPhysicsControlObjectVersion::GUID) >= FPhysicsControlObjectVersion::ControlRigIncludeDriveInJoint)
 	{
 		Ar << Data.Drive;
 	}
diff --git a/Engine/Plugins/Experimental/PhysicsControl/Source/ControlRigPhysics/Private/RigPhysicsExecution.cpp b/Engine/Plugins/Experimental/PhysicsControl/Source/ControlRigPhysics/Private/RigPhysicsExecution.cpp
--- a/Engine/Plugins/Experimental/PhysicsControl/Source/ControlRigPhysics/Private/RigPhysicsExecution.cpp
+++ b/Engine/Plugins/Experimental/PhysicsControl/Source/ControlRigPhysics/Private/RigPhysicsExecution.cpp
@@ -209,7 +209,6 @@ FRigUnit_AddPhysicsBody_Execute()
 						Component->Collision = Collision;
 						Component->Joint = Joint;
 						Component->BodyData = BodyData;
-						Component->CurrentBodyData = BodyData;
 					}
 				}
 			}
@@ -261,8 +260,8 @@ FRigUnit_HierarchySetPhysicsBodyKinematicTarget_Execute()
 		if (FRigPhysicsComponent* PhysicsComponent = Cast<FRigPhysicsComponent>(
 			ExecuteContext.Hierarchy->FindComponent(PhysicsBodyComponentKey)))
 		{
-			PhysicsComponent->CurrentKinematicTarget = KinematicTarget;
-			PhysicsComponent->CurrentKinematicTargetSpace = KinematicTargetSpace;
+			PhysicsComponent->KinematicTarget = KinematicTarget;
+			PhysicsComponent->KinematicTargetSpace = KinematicTargetSpace;
 		}
 	}
 }
@@ -277,14 +276,7 @@ FRigUnit_HierarchySetPhysicsBodySparseData_Execute()
 		if (FRigPhysicsComponent* PhysicsComponent = Cast<FRigPhysicsComponent>(
 			ExecuteContext.Hierarchy->FindComponent(PhysicsBodyComponentKey)))
 		{
-			if (bInitial)
-			{
-				PhysicsComponent->BodyData.UpdateFromSparseData(Data);
-			}
-			else
-			{
-				PhysicsComponent->CurrentBodyData.UpdateFromSparseData(Data);
-			}
+			PhysicsComponent->BodyData.UpdateFromSparseData(Data);
 		}
 	}
 }
@@ -298,14 +290,7 @@ FRigUnit_HierarchySetPhysicsBodyMovementType_Execute()
 		if (FRigPhysicsComponent* PhysicsComponent = Cast<FRigPhysicsComponent>(
 			ExecuteContext.Hierarchy->FindComponent(PhysicsBodyComponentKey)))
 		{
-			if (bInitial)
-			{
-				PhysicsComponent->BodyData.MovementType = MovementType;
-			}
-			else
-			{
-				PhysicsComponent->CurrentBodyData.MovementType = MovementType;
-			}
+			PhysicsComponent->BodyData.MovementType = MovementType;
 		}
 	}
 }
@@ -319,14 +304,7 @@ FRigUnit_HierarchySetPhysicsBodyCollisionType_Execute()
 		if (FRigPhysicsComponent* PhysicsComponent = Cast<FRigPhysicsComponent>(
 			ExecuteContext.Hierarchy->FindComponent(PhysicsBodyComponentKey)))
 		{
-			if (bInitial)
-			{
-				PhysicsComponent->BodyData.CollisionType = CollisionType;
-			}
-			else
-			{
-				PhysicsComponent->CurrentBodyData.CollisionType = CollisionType;
-			}
+			PhysicsComponent->BodyData.CollisionType = CollisionType;
 		}
 	}
 }
@@ -340,14 +318,7 @@ FRigUnit_HierarchySetPhysicsBodyGravityMultiplier_Execute()
 		if (FRigPhysicsComponent* PhysicsComponent = Cast<FRigPhysicsComponent>(
 			ExecuteContext.Hierarchy->FindComponent(PhysicsBodyComponentKey)))
 		{
-			if (bInitial)
-			{
-				PhysicsComponent->BodyData.GravityMultiplier = GravityMultiplier;
-			}
-			else
-			{
-				PhysicsComponent->CurrentBodyData.GravityMultiplier = GravityMultiplier;
-			}
+			PhysicsComponent->BodyData.GravityMultiplier = GravityMultiplier;
 		}
 	}
 }
@@ -361,14 +332,7 @@ FRigUnit_HierarchySetPhysicsBodyPhysicsBlendWeight_Execute()
 		if (FRigPhysicsComponent* PhysicsComponent = Cast<FRigPhysicsComponent>(
 			ExecuteContext.Hierarchy->FindComponent(PhysicsBodyComponentKey)))
 		{
-			if (bInitial)
-			{
-				PhysicsComponent->BodyData.PhysicsBlendWeight = PhysicsBlendWeight;
-			}
-			else
-			{
-				PhysicsComponent->CurrentBodyData.PhysicsBlendWeight = PhysicsBlendWeight;
-			}
+			PhysicsComponent->BodyData.PhysicsBlendWeight = PhysicsBlendWeight;
 		}
 	}
 }
@@ -382,14 +346,7 @@ FRigUnit_HierarchySetPhysicsBodyUseSkeletalAnimation_Execute()
 		if (FRigPhysicsComponent* PhysicsComponent = Cast<FRigPhysicsComponent>(
 			ExecuteContext.Hierarchy->FindComponent(PhysicsBodyComponentKey)))
 		{
-			if (bInitial)
-			{
-				PhysicsComponent->BodyData.bUseSkeletalAnimation = bUseSkeletalAnimation;
-			}
-			else
-			{
-				PhysicsComponent->CurrentBodyData.bUseSkeletalAnimation = bUseSkeletalAnimation;
-			}
+			PhysicsComponent->BodyData.bUseSkeletalAnimation = bUseSkeletalAnimation;
 		}
 	}
 }
@@ -403,14 +360,7 @@ FRigUnit_HierarchySetPhysicsBodyUpdateKinematicFromSimulation_Execute()
 		if (FRigPhysicsComponent* PhysicsComponent = Cast<FRigPhysicsComponent>(
 			ExecuteContext.Hierarchy->FindComponent(PhysicsBodyComponentKey)))
 		{
-			if (bInitial)
-			{
-				PhysicsComponent->BodyData.bUpdateKinematicFromSimulation = bUpdateKinematicFromSimulation;
-			}
-			else
-			{
-				PhysicsComponent->CurrentBodyData.bUpdateKinematicFromSimulation = bUpdateKinematicFromSimulation;
-			}
+			PhysicsComponent->BodyData.bUpdateKinematicFromSimulation = bUpdateKinematicFromSimulation;
 		}
 	}
 }
diff --git a/Engine/Plugins/Experimental/PhysicsControl/Source/ControlRigPhysics/Private/RigPhysicsSimulation.cpp b/Engine/Plugins/Experimental/PhysicsControl/Source/ControlRigPhysics/Private/RigPhysicsSimulation.cpp
--- a/Engine/Plugins/Experimental/PhysicsControl/Source/ControlRigPhysics/Private/RigPhysicsSimulation.cpp
+++ b/Engine/Plugins/Experimental/PhysicsControl/Source/ControlRigPhysics/Private/RigPhysicsSimulation.cpp
@@ -567,7 +567,7 @@ void FRigPhysicsSimulation::UpdatePrePhysics(
 			if (Record.ActorHandle)
 			{
 				UpdateBodyFromModifierData(
-					Record.ActorHandle, Simulation.Get(), PhysicsComponent->CurrentBodyData, SimulationSpaceData.Gravity);
+					Record.ActorHandle, Simulation.Get(), PhysicsComponent->BodyData, SimulationSpaceData.Gravity);
 
 				if (Record.ActorHandle->GetIsKinematic())
 				{
@@ -575,17 +575,17 @@ void FRigPhysicsSimulation::UpdatePrePhysics(
 
 					// If the target is already in component space, then that's all we need
 					FTransform KinematicTargetCS;
-					switch (PhysicsComponent->CurrentKinematicTargetSpace)
+					switch (PhysicsComponent->KinematicTargetSpace)
 					{
 						case ERigPhysicsKinematicTargetSpace::Component:
 						{
-							KinematicTargetCS = PhysicsComponent->CurrentKinematicTarget;
+							KinematicTargetCS = PhysicsComponent->KinematicTarget;
 							break;
 						}
 						case ERigPhysicsKinematicTargetSpace::World:
 						{
 							// Record.KinematicTarget * SimulationSpaceState.ComponentTM.Inverse();
-							KinematicTargetCS = PhysicsComponent->CurrentKinematicTarget.GetRelativeTransform(
+							KinematicTargetCS = PhysicsComponent->KinematicTarget.GetRelativeTransform(
 								SimulationSpaceState.ComponentTM);
 							break;
 						}
@@ -602,25 +602,25 @@ void FRigPhysicsSimulation::UpdatePrePhysics(
 							{
 								FTransform SourceComponentSpaceTM = Hierarchy->GetGlobalTransform(SourceKey);
 
-								switch (PhysicsComponent->CurrentKinematicTargetSpace)
+								switch (PhysicsComponent->KinematicTargetSpace)
 								{
 									case ERigPhysicsKinematicTargetSpace::OffsetInBoneSpace:
 									{
-										KinematicTargetCS = PhysicsComponent->CurrentKinematicTarget * SourceComponentSpaceTM;
+										KinematicTargetCS = PhysicsComponent->KinematicTarget * SourceComponentSpaceTM;
 										break;
 									}
 									case ERigPhysicsKinematicTargetSpace::OffsetInWorldSpace:
 									{
 										// Convert the bone to WS, apply the target, and convert back
 										FTransform BoneWS = SourceComponentSpaceTM * SimulationSpaceState.ComponentTM;
-										FTransform KinematicTargetWS = BoneWS * PhysicsComponent->CurrentKinematicTarget;
+										FTransform KinematicTargetWS = BoneWS * PhysicsComponent->KinematicTarget;
 										// Danny TODO figure out which of the GetRelativeTransform versions this is
 										KinematicTargetCS = SimulationSpaceState.ComponentTM.Inverse() * KinematicTargetWS;
 										break;
 									}
 									case ERigPhysicsKinematicTargetSpace::OffsetInComponentSpace:
 									{
-										KinematicTargetCS = SourceComponentSpaceTM * PhysicsComponent->CurrentKinematicTarget;
+										KinematicTargetCS = SourceComponentSpaceTM * PhysicsComponent->KinematicTarget;
 										break;
 									}
 									default:
diff --git a/Engine/Plugins/Experimental/PhysicsControl/Source/ControlRigPhysics/Public/RigPhysicsComponent.h b/Engine/Plugins/Experimental/PhysicsControl/Source/ControlRigPhysics/Public/RigPhysicsComponent.h
--- a/Engine/Plugins/Experimental/PhysicsControl/Source/ControlRigPhysics/Public/RigPhysicsComponent.h
+++ b/Engine/Plugins/Experimental/PhysicsControl/Source/ControlRigPhysics/Public/RigPhysicsComponent.h
@@ -41,20 +41,12 @@ public:
 	UPROPERTY(BlueprintReadOnly, EditAnywhere, Category = Physics)
 	FPhysicsControlModifierData BodyData;
 
-	//==================================================================================================================
-	// Current data
-	//==================================================================================================================
-
 	// The target for when this body is kinematic
 	UPROPERTY(BlueprintReadOnly, EditAnywhere, Category = CurrentPhysics)
-	FTransform CurrentKinematicTarget;
-
-	UPROPERTY(BlueprintReadOnly, EditAnywhere, Category = CurrentPhysics)
-	ERigPhysicsKinematicTargetSpace CurrentKinematicTargetSpace = ERigPhysicsKinematicTargetSpace::OffsetInBoneSpace;
+	FTransform KinematicTarget;
 
-	// Desired state of the physics body - this will be written to the simulation during UpdatePrePhysics
 	UPROPERTY(BlueprintReadOnly, EditAnywhere, Category = CurrentPhysics)
-	FPhysicsControlModifierData CurrentBodyData;
+	ERigPhysicsKinematicTargetSpace KinematicTargetSpace = ERigPhysicsKinematicTargetSpace::OffsetInBoneSpace;
 
 	virtual void Save(FArchive& Ar) override;
 	virtual void Load(FArchive& Ar) override;
diff --git a/Engine/Plugins/Experimental/PhysicsControl/Source/ControlRigPhysics/Public/RigPhysicsExecution.h b/Engine/Plugins/Experimental/PhysicsControl/Source/ControlRigPhysics/Public/RigPhysicsExecution.h
--- a/Engine/Plugins/Experimental/PhysicsControl/Source/ControlRigPhysics/Public/RigPhysicsExecution.h
+++ b/Engine/Plugins/Experimental/PhysicsControl/Source/ControlRigPhysics/Public/RigPhysicsExecution.h
@@ -211,9 +211,6 @@ struct CONTROLRIGPHYSICS_API FRigUnit_HierarchySetPhysicsBodySparseData : public
 
 	UPROPERTY(meta = (Input))
 	FPhysicsControlModifierSparseData Data;
-
-	UPROPERTY(meta = (Input))
-	bool bInitial = true;
 };
 
 /**
@@ -235,9 +232,6 @@ struct CONTROLRIGPHYSICS_API FRigUnit_HierarchySetPhysicsBodyKinematicTarget : p
 
 	UPROPERTY(meta = (Input))
 	FTransform KinematicTarget;
-
-	UPROPERTY(meta = (Input))
-	bool bInitial = true;
 };
 
 /**
@@ -257,9 +251,6 @@ struct CONTROLRIGPHYSICS_API FRigUnit_HierarchySetPhysicsBodyMovementType : publ
 
 	UPROPERTY(meta = (Input))
 	EPhysicsMovementType MovementType = EPhysicsMovementType::Simulated;
-
-	UPROPERTY(meta = (Input))
-	bool bInitial = true;
 };
 
 /**
@@ -278,9 +269,6 @@ struct CONTROLRIGPHYSICS_API FRigUnit_HierarchySetPhysicsBodyCollisionType : pub
 
 	UPROPERTY(meta = (Input))
 	TEnumAsByte<ECollisionEnabled::Type> CollisionType = ECollisionEnabled::QueryAndPhysics;
-
-	UPROPERTY(meta = (Input))
-	bool bInitial = true;
 };
 
 /**
@@ -299,9 +287,6 @@ struct CONTROLRIGPHYSICS_API FRigUnit_HierarchySetPhysicsBodyGravityMultiplier :
 
 	UPROPERTY(meta = (Input))
 	float GravityMultiplier = 1.0f;
-
-	UPROPERTY(meta = (Input))
-	bool bInitial = true;
 };
 
 /**
@@ -321,9 +306,6 @@ struct CONTROLRIGPHYSICS_API FRigUnit_HierarchySetPhysicsBodyPhysicsBlendWeight
 
 	UPROPERTY(meta = (Input))
 	float PhysicsBlendWeight = 1.0f;
-
-	UPROPERTY(meta = (Input))
-	bool bInitial = true;
 };
 
 /**
@@ -343,9 +325,6 @@ struct CONTROLRIGPHYSICS_API FRigUnit_HierarchySetPhysicsBodyUseSkeletalAnimatio
 
 	UPROPERTY(meta = (Input))
 	bool bUseSkeletalAnimation = true;
-
-	UPROPERTY(meta = (Input))
-	bool bInitial = true;
 };
 
 /**
@@ -366,9 +345,6 @@ struct CONTROLRIGPHYSICS_API FRigUnit_HierarchySetPhysicsBodyUpdateKinematicFrom
 
 	UPROPERTY(meta = (Input))
 	bool bUpdateKinematicFromSimulation = true;
-
-	UPROPERTY(meta = (Input))
-	bool bInitial = true;
 };
 
 
diff --git a/Engine/Plugins/Experimental/PhysicsControl/Source/PhysicsControl/Public/PhysicsControlObjectVersion.h b/Engine/Plugins/Experimental/PhysicsControl/Source/PhysicsControl/Public/PhysicsControlObjectVersion.h
--- a/Engine/Plugins/Experimental/PhysicsControl/Source/PhysicsControl/Public/PhysicsControlObjectVersion.h
+++ b/Engine/Plugins/Experimental/PhysicsControl/Source/PhysicsControl/Public/PhysicsControlObjectVersion.h
@@ -10,7 +10,9 @@ struct PHYSICSCONTROL_API FPhysicsControlObjectVersion
 		// Before any version changes were made
 		BeforeCustomVersionWasAdded,
 
-		IncludeDriveInJoint,
+		ControlRigIncludeDriveInJoint,
+
+		ControlRigRemoveCurrentDataFromPhysicsComponent,
 
 		// -----<new versions can be added above this line>-------------------------------------------------
 		VersionPlusOne,
