Support for automatically connecting physics components to solvers

#rnx

diff --git a/Engine/Plugins/Experimental/PhysicsControl/Source/ControlRigPhysics/Private/RigPhysicsControlExecution.cpp b/Engine/Plugins/Experimental/PhysicsControl/Source/ControlRigPhysics/Private/RigPhysicsControlExecution.cpp
--- a/Engine/Plugins/Experimental/PhysicsControl/Source/ControlRigPhysics/Private/RigPhysicsControlExecution.cpp
+++ b/Engine/Plugins/Experimental/PhysicsControl/Source/ControlRigPhysics/Private/RigPhysicsControlExecution.cpp
@@ -1,7 +1,6 @@
 // Copyright Epic Games, Inc. All Rights Reserved.
 
 #include "RigPhysicsControlExecution.h"
-#include "RigPhysicsSimulation.h"
 #include "RigPhysicsControlComponent.h"
 
 #include "Rigs/RigHierarchyController.h"
diff --git a/Engine/Plugins/Experimental/PhysicsControl/Source/ControlRigPhysics/Private/RigPhysicsData.cpp b/Engine/Plugins/Experimental/PhysicsControl/Source/ControlRigPhysics/Private/RigPhysicsData.cpp
--- a/Engine/Plugins/Experimental/PhysicsControl/Source/ControlRigPhysics/Private/RigPhysicsData.cpp
+++ b/Engine/Plugins/Experimental/PhysicsControl/Source/ControlRigPhysics/Private/RigPhysicsData.cpp
@@ -3,6 +3,7 @@
 #include "RigPhysicsData.h"
 #include "PhysicsControlObjectVersion.h"
 
+//======================================================================================================================
 FArchive& operator <<(FArchive& Ar, FRigPhysicsJoint& Data)
 {
 	Ar << Data.bEnable;
@@ -24,6 +25,7 @@ FArchive& operator <<(FArchive& Ar, FRigPhysicsJoint& Data)
 	return Ar;
 }
 
+//======================================================================================================================
 FArchive& operator <<(FArchive& Ar, FRigPhysicsSimulationSpaceSettings& Data)
 {
 	Ar << Data.SpaceMovementAmount;
@@ -44,6 +46,7 @@ FArchive& operator <<(FArchive& Ar, FRigPhysicsSimulationSpaceSettings& Data)
 	return Ar;
 }
 
+//======================================================================================================================
 FArchive& operator <<(FArchive& Ar, FRigPhysicsSolverSettings& Data)
 {
 	Ar << Data.SimulationSpace;
@@ -63,9 +66,14 @@ FArchive& operator <<(FArchive& Ar, FRigPhysicsSolverSettings& Data)
 	Ar << Data.bSolveJointPositionsLast;
 	Ar << Data.bUseManifolds;
 	Ar << Data.DistanceThresholdForReset;
+	if (Ar.CustomVer(FPhysicsControlObjectVersion::GUID) >= FPhysicsControlObjectVersion::ControlRigAutomaticallyAddPhysicsComponents)
+	{
+		Ar << Data.bAutomaticallyAddPhysicsComponents;
+	}
 	return Ar;
 }
 
+//======================================================================================================================
 FArchive& operator <<(FArchive& Ar, FRigPhysicsDrive& Data)
 {
 	Ar << Data.bEnable;
@@ -81,14 +89,20 @@ FArchive& operator <<(FArchive& Ar, FRigPhysicsDrive& Data)
 	return Ar;
 }
 
+//======================================================================================================================
 FArchive& operator <<(FArchive& Ar, FRigPhysicsBodySolverSettings& Data)
 {
 	Ar << Data.PhysicsSolverComponentKey;
 	Ar << Data.TargetBone;
 	Ar << Data.SourceBone;
+	if (Ar.CustomVer(FPhysicsControlObjectVersion::GUID) >= FPhysicsControlObjectVersion::ControlRigUseAutomaticSolver)
+	{
+		Ar << Data.bUseAutomaticSolver;
+	}
 	return Ar;
 }
 
+//======================================================================================================================
 FArchive& operator <<(FArchive& Ar, FRigPhysicsCollision& Data)
 {
 	Ar << Data.Boxes;
@@ -97,7 +111,9 @@ FArchive& operator <<(FArchive& Ar, FRigPhysicsCollision& Data)
 	return Ar;
 }
 
-void FRigPhysicsBodySolverSettings::OnRigHierarchyKeyChanged(const FRigHierarchyKey& InOldKey, const FRigHierarchyKey& InNewKey)
+//======================================================================================================================
+void FRigPhysicsBodySolverSettings::OnRigHierarchyKeyChanged(
+	const FRigHierarchyKey& InOldKey, const FRigHierarchyKey& InNewKey)
 {
 	if(InOldKey.IsComponent() && InNewKey.IsComponent())
 	{
diff --git a/Engine/Plugins/Experimental/PhysicsControl/Source/ControlRigPhysics/Private/RigPhysicsSimulation.cpp b/Engine/Plugins/Experimental/PhysicsControl/Source/ControlRigPhysics/Private/RigPhysicsSimulation.cpp
--- a/Engine/Plugins/Experimental/PhysicsControl/Source/ControlRigPhysics/Private/RigPhysicsSimulation.cpp
+++ b/Engine/Plugins/Experimental/PhysicsControl/Source/ControlRigPhysics/Private/RigPhysicsSimulation.cpp
@@ -48,6 +48,61 @@ FRigPhysicsSimulation::FRigPhysicsSimulation(TObjectPtr<UControlRig> InOwningCon
 	Chaos::CVars::bEnableAsyncInitBody = true;
 }
 
+//======================================================================================================================
+bool FRigPhysicsSimulation::ShouldComponentBeInSimulation(
+	FRigComponentKey SolverComponentKey, FRigComponentKey ComponentKey) const
+{
+	const URigHierarchy* Hierarchy = OwningControlRig->GetHierarchy();
+
+	const FRigPhysicsComponent* PhysicsComponent = Cast<FRigPhysicsComponent>(
+		Hierarchy->FindComponent(ComponentKey));
+
+	const FRigPhysicsSolverComponent* PhysicsSolverComponent = Cast<FRigPhysicsSolverComponent>(
+		Hierarchy->FindComponent(SolverComponentKey));
+
+	if (!PhysicsComponent || !PhysicsSolverComponent)
+	{
+		return false;
+	}
+
+	if (PhysicsComponent->BodySolverSettings.PhysicsSolverComponentKey == SolverComponentKey)
+	{
+		return true;
+	}
+
+	if (!PhysicsSolverComponent->SolverSettings.bAutomaticallyAddPhysicsComponents)
+	{
+		return false;
+	}
+
+	if (PhysicsComponent->BodySolverSettings.bUseAutomaticSolver)
+	{
+		FRigElementKey ElementKey = PhysicsComponent->GetElementKey();
+		while (ElementKey.IsValid())
+		{
+			for (FRigComponentKey CK : Hierarchy->GetComponentKeys(ElementKey))
+			{
+				if (CK == SolverComponentKey)
+				{
+					return true;
+				}
+			}
+			// Note that getting the parent of an element at the root doesn't return the top-level element
+			ElementKey = Hierarchy->GetFirstParent(ElementKey);
+		}
+
+		for (FRigComponentKey CK : Hierarchy->GetTopLevelComponentKeys())
+		{
+			if (CK == SolverComponentKey)
+			{
+				return true;
+			}
+		}
+	}
+	return false;
+}
+
+
 //======================================================================================================================
 void FRigPhysicsSimulation::InitialiseSimulation(const FRigPhysicsSolverComponent* SolverComponent)
 {
@@ -77,15 +132,11 @@ void FRigPhysicsSimulation::InitialiseSimulation(const FRigPhysicsSolverComponen
 
 	for (FRigComponentKey RigComponentKey : AllComponentKeys)
 	{
-		if (const FRigPhysicsComponent* PhysicsComponent = Cast<FRigPhysicsComponent>(
-			Hierarchy->FindComponent(RigComponentKey)))
+		if (ShouldComponentBeInSimulation(SolverComponentKey, RigComponentKey))
 		{
-			if (PhysicsComponent->BodySolverSettings.PhysicsSolverComponentKey == SolverComponentKey)
-			{
-				// Just make the record for now - it will be instantiated later
-				BodyRecords.Add(RigComponentKey);
-				UnsortedBodyComponentKeys.Add(RigComponentKey);
-			}
+			// Just make the record for now - it will be instantiated later
+			BodyRecords.Add(RigComponentKey);
+			UnsortedBodyComponentKeys.Add(RigComponentKey);
 		}
 	}
 
diff --git a/Engine/Plugins/Experimental/PhysicsControl/Source/ControlRigPhysics/Private/RigPhysicsSimulation.h b/Engine/Plugins/Experimental/PhysicsControl/Source/ControlRigPhysics/Private/RigPhysicsSimulation.h
--- a/Engine/Plugins/Experimental/PhysicsControl/Source/ControlRigPhysics/Private/RigPhysicsSimulation.h
+++ b/Engine/Plugins/Experimental/PhysicsControl/Source/ControlRigPhysics/Private/RigPhysicsSimulation.h
@@ -123,6 +123,11 @@ private:
 
 	void InstantiatePhysicsBodies(const FRigPhysicsSolverComponent* SolverComponent);
 
+	// Returns true if the component is physics, and its solver matches the solver component
+	// (directly, or automatically)
+	bool ShouldComponentBeInSimulation(
+		FRigComponentKey SolverComponentKey, FRigComponentKey ComponentKey) const;
+
 	// Creates an actor with collision. This will by dynamic if Dynamics is valid, or otherwise kinematic
 	ImmediatePhysics::FActorHandle* CreateBody(
 		const FName                 BodyName,
diff --git a/Engine/Plugins/Experimental/PhysicsControl/Source/ControlRigPhysics/Private/RigPhysicsSolverComponent.cpp b/Engine/Plugins/Experimental/PhysicsControl/Source/ControlRigPhysics/Private/RigPhysicsSolverComponent.cpp
--- a/Engine/Plugins/Experimental/PhysicsControl/Source/ControlRigPhysics/Private/RigPhysicsSolverComponent.cpp
+++ b/Engine/Plugins/Experimental/PhysicsControl/Source/ControlRigPhysics/Private/RigPhysicsSolverComponent.cpp
@@ -1,6 +1,7 @@
 // Copyright Epic Games, Inc. All Rights Reserved.
 
 #include "RigPhysicsSolverComponent.h"
+#include "PhysicsControlObjectVersion.h"
 #include "Rigs/RigHierarchy.h"
 #if WITH_EDITOR
 #include "ControlRigPhysicsEditorStyle.h"
@@ -9,6 +10,8 @@
 //======================================================================================================================
 void FRigPhysicsSolverComponent::Save(FArchive& Ar)
 {
+	Ar.UsingCustomVersion(FPhysicsControlObjectVersion::GUID);
+
 	FRigBaseComponent::Save(Ar);
 	Ar << SolverSettings;
 	Ar << SimulationSpaceSettings;
@@ -26,7 +29,8 @@ void FRigPhysicsSolverComponent::Load(FArchive& Ar)
 #if WITH_EDITOR
 const FSlateIcon& FRigPhysicsSolverComponent::GetIconForUI() const
 {
-	static const FSlateIcon SolverIcon = FSlateIcon(FControlRigPhysicsEditorStyle::Get().GetStyleSetName(), "ControlRigPhysics.Component.Solver");
+	static const FSlateIcon SolverIcon = FSlateIcon(
+		FControlRigPhysicsEditorStyle::Get().GetStyleSetName(), "ControlRigPhysics.Component.Solver");
 	return SolverIcon;
 }
 #endif
\ No newline at end of file
diff --git a/Engine/Plugins/Experimental/PhysicsControl/Source/ControlRigPhysics/Public/RigPhysicsControlExecution.h b/Engine/Plugins/Experimental/PhysicsControl/Source/ControlRigPhysics/Public/RigPhysicsControlExecution.h
--- a/Engine/Plugins/Experimental/PhysicsControl/Source/ControlRigPhysics/Public/RigPhysicsControlExecution.h
+++ b/Engine/Plugins/Experimental/PhysicsControl/Source/ControlRigPhysics/Public/RigPhysicsControlExecution.h
@@ -3,9 +3,7 @@
 #pragma once
 
 #include "RigPhysicsComponent.h"
-#include "RigPhysicsSimulation.h"
 #include "PhysicsControlData.h"
-#include "RigPhysicsExecution.h"
 
 #include "Units/Execution/RigUnit_DynamicHierarchy.h"
 
diff --git a/Engine/Plugins/Experimental/PhysicsControl/Source/ControlRigPhysics/Public/RigPhysicsData.h b/Engine/Plugins/Experimental/PhysicsControl/Source/ControlRigPhysics/Public/RigPhysicsData.h
--- a/Engine/Plugins/Experimental/PhysicsControl/Source/ControlRigPhysics/Public/RigPhysicsData.h
+++ b/Engine/Plugins/Experimental/PhysicsControl/Source/ControlRigPhysics/Public/RigPhysicsData.h
@@ -312,6 +312,11 @@ struct CONTROLRIGPHYSICS_API FRigPhysicsSolverSettings
 
 	friend FArchive& operator <<(FArchive& Ar, FRigPhysicsSolverSettings& Data);
 
+	// If true, then any physics component below (the element owning this one) will be automatically
+	// added to this solver, if its UseAutomaticSolver flag is set.
+	UPROPERTY(BlueprintReadWrite, EditAnywhere, Category = SolverSettings)
+	bool bAutomaticallyAddPhysicsComponents = true;
+
 	UPROPERTY(BlueprintReadOnly, EditAnywhere, Category = SolverSettings)
 	ERigPhysicsSimulationSpace SimulationSpace = ERigPhysicsSimulationSpace::Component;
 
@@ -612,16 +617,24 @@ struct CONTROLRIGPHYSICS_API FRigPhysicsBodySolverSettings
 
 	friend FArchive& operator <<(FArchive& Ar, FRigPhysicsBodySolverSettings& Data);
 
+	// Note that setting the solver component, if known, has the benefit of avoiding the need to
+	// search for an automatic solver.
 	UPROPERTY(BlueprintReadWrite, EditAnywhere, Category = Solver)
 	FRigComponentKey PhysicsSolverComponentKey;
 
-	/// The bone that is used to initialize physics, as well as what to track when the body is set to be kinematic.
-	/// Note that if this is unset, the it will default to the parent of our owner.
+	// If true (and the physics solver is not explicitly set), then this component will be added to
+	// any physics solver that exists above it in the hierarchy, if that solver allows automatically
+	// adding physics components.
+	UPROPERTY(BlueprintReadWrite, EditAnywhere, Category = Solver)
+	bool bUseAutomaticSolver = true;
+
+	// The bone that is used to initialize physics, as well as what to track when the body is set to be kinematic.
+	// Note that if this is unset, the it will default to the parent of our owner.
 	UPROPERTY(BlueprintReadWrite, EditAnywhere, Category = Solver)
 	FRigElementKey SourceBone;
 
-	/// The bone that is written to following simulation. Note that if this is unset, the it will
-	/// default to the parent of our owner.
+	// The bone that is written to following simulation. Note that if this is unset, the it will
+	// default to the parent of our owner.
 	UPROPERTY(BlueprintReadWrite, EditAnywhere, Category = Solver)
 	FRigElementKey TargetBone;
 
diff --git a/Engine/Plugins/Experimental/PhysicsControl/Source/PhysicsControl/Public/PhysicsControlObjectVersion.h b/Engine/Plugins/Experimental/PhysicsControl/Source/PhysicsControl/Public/PhysicsControlObjectVersion.h
--- a/Engine/Plugins/Experimental/PhysicsControl/Source/PhysicsControl/Public/PhysicsControlObjectVersion.h
+++ b/Engine/Plugins/Experimental/PhysicsControl/Source/PhysicsControl/Public/PhysicsControlObjectVersion.h
@@ -14,6 +14,10 @@ struct PHYSICSCONTROL_API FPhysicsControlObjectVersion
 
 		ControlRigRemoveCurrentDataFromPhysicsComponent,
 
+		ControlRigUseAutomaticSolver,
+
+		ControlRigAutomaticallyAddPhysicsComponents,
+
 		// -----<new versions can be added above this line>-------------------------------------------------
 		VersionPlusOne,
 		LatestVersion = VersionPlusOne - 1
