Control Rig Physics - support auto calculation of collision shapes

#rnx

diff --git a/Engine/Plugins/Experimental/PhysicsControl/Source/ControlRigPhysics/Private/RigPhysicsComponent.cpp b/Engine/Plugins/Experimental/PhysicsControl/Source/ControlRigPhysics/Private/RigPhysicsComponent.cpp
--- a/Engine/Plugins/Experimental/PhysicsControl/Source/ControlRigPhysics/Private/RigPhysicsComponent.cpp
+++ b/Engine/Plugins/Experimental/PhysicsControl/Source/ControlRigPhysics/Private/RigPhysicsComponent.cpp
@@ -6,6 +6,8 @@
 #include "Rigs/RigHierarchy.h"
 #include "Rigs/RigHierarchyController.h"
 #include "ControlRigPhysicsModule.h"
+#include "BoxTypes.h"
+
 #if WITH_EDITOR
 #include "ControlRigPhysicsEditorStyle.h"
 #endif
@@ -99,6 +101,11 @@ bool FRigPhysicsComponent::CanBeAddedTo(
 //======================================================================================================================
 void FRigPhysicsComponent::OnAddedToHierarchy(URigHierarchy* InHierarchy, URigHierarchyController* InController)
 {
+	if (!IsProcedural())
+	{
+		AutoCalculateCollision(InHierarchy);
+	}
+
 	if(!CVarControlRigPhysicsAutomaticallySetSolver.GetValueOnAnyThread())
 	{
 		return;
@@ -119,6 +126,83 @@ void FRigPhysicsComponent::OnAddedToHierarchy(URigHierarchy* InHierarchy, URigHi
 	InHierarchy->Notify(ERigHierarchyNotification::ComponentContentChanged, this);
 }
 
+//======================================================================================================================
+void FRigPhysicsComponent::AutoCalculateCollision(URigHierarchy* InHierarchy, float MinAspectRatio)
+{
+	// Start clean
+	Collision = FRigPhysicsCollision();
+
+	TArray<FVector> Points{ FVector::ZeroVector };
+	FVector MidPoint = FVector::ZeroVector;
+	TArray<FRigElementKey> ChildKeys = InHierarchy->GetChildren(GetElementKey());
+	for (FRigElementKey ChildKey : ChildKeys)
+	{
+		FVector ChildPosition = InHierarchy->GetLocalTransform(ChildKey, true).GetTranslation();
+		if (ChildPosition.SquaredLength() > UE_SMALL_NUMBER)
+		{
+			Points.Add(ChildPosition);
+			MidPoint += ChildPosition;
+		}
+	}
+	MidPoint /= Points.Num();
+
+	if (Points.Num() == 1)
+	{
+		// If there's only one point, then there are no significant children. Make a shape that
+		// duplicates the relationship with our parent, if there is one
+		FTransform TM = InHierarchy->GetLocalTransform(GetElementKey(), true);
+		if (TM.GetTranslation().SquaredLength() < UE_SMALL_NUMBER)
+		{
+			// We have no children, and are co-located with our parent. Make a single, arbitrarily
+			// sized, shape.
+			Collision.Boxes.Add(FRigPhysicsCollisionBox(TM, FVector(10.0)));
+		}
+		else
+		{
+			TM.SetTranslation(InHierarchy->GetLocalTransform(GetElementKey()).GetTranslation() * 0.5);
+			// We're going to make a box that goes along the X axis
+			TM.SetRotation(FQuat::FindBetweenVectors(FVector(1, 0, 0), TM.GetTranslation()));
+			FVector Extents;
+			Extents.X = TM.GetTranslation().Length() * 2.0;
+			Extents.Y = Extents.X * MinAspectRatio;
+			Extents.Z = Extents.Y;
+			Collision.Boxes.Add(FRigPhysicsCollisionBox(TM, Extents));
+		}
+	}
+	else if (Points.Num() == 2)
+	{
+		// There's just one child - easier to do this by hand than the more complex eigenvector
+		// based calculation, and better than using an OBB which may not be aligned with the two points.
+		FTransform TM(MidPoint);
+		// We're going to make a box that goes along the X axis
+		TM.SetRotation(FQuat::FindBetweenVectors(FVector(1, 0, 0), Points[1]));
+		FVector Extents;
+		Extents.X = TM.GetTranslation().Length() * 2.0f;
+		Extents.Y = Extents.X * MinAspectRatio;
+		Extents.Z = Extents.Y;
+		Collision.Boxes.Add(FRigPhysicsCollisionBox(TM, Extents));
+	}
+	else
+	{
+		// We could calculate the ideal orientation of a box by calculating the eigenvectors of the
+		// covariance matrix that represents all the joint positions relative to the centroid.
+		// However, for now just use a box orientated with the current join - it is simpler, and
+		// also avoids generating "messy" orientations.
+		UE::Geometry::FAxisAlignedBox3d Box;
+		Box.Contain(Points);
+
+		FTransform TM(Box.Center());
+		FVector Extents = Box.Extents() * 2.0; // The returned extents are half the extents!
+		double MaxExtent = Extents.GetAbsMax();
+		Extents.X = FMath::Max(Extents.X, MaxExtent * MinAspectRatio);
+		Extents.Y = FMath::Max(Extents.Y, MaxExtent * MinAspectRatio);
+		Extents.Z = FMath::Max(Extents.Z, MaxExtent * MinAspectRatio);
+		Collision.Boxes.Add(FRigPhysicsCollisionBox(TM, Extents));
+	}
+
+	InHierarchy->Notify(ERigHierarchyNotification::ComponentContentChanged, this);
+}
+
 //======================================================================================================================
 void FRigPhysicsComponent::OnRigHierarchyKeyChanged(const FRigHierarchyKey& InOldKey, const FRigHierarchyKey& InNewKey)
 {
diff --git a/Engine/Plugins/Experimental/PhysicsControl/Source/ControlRigPhysics/Private/RigPhysicsExecution.cpp b/Engine/Plugins/Experimental/PhysicsControl/Source/ControlRigPhysics/Private/RigPhysicsExecution.cpp
--- a/Engine/Plugins/Experimental/PhysicsControl/Source/ControlRigPhysics/Private/RigPhysicsExecution.cpp
+++ b/Engine/Plugins/Experimental/PhysicsControl/Source/ControlRigPhysics/Private/RigPhysicsExecution.cpp
@@ -204,11 +204,11 @@ FRigUnit_AddPhysicsBody_Execute()
 					if (FRigPhysicsComponent* Component = Cast<FRigPhysicsComponent>(
 						Hierarchy->FindComponent(PhysicsBodyComponentKey)))
 					{
-						Component->BodySolverSettings.PhysicsSolverComponentKey = PhysicsSolverComponentKey;
+						Component->BodySolverSettings = Solver;
 						Component->Dynamics = Dynamics;
-						Component->Collision = Collision;
 						Component->Joint = Joint;
 						Component->BodyData = BodyData;
+						Component->Collision = Collision;
 					}
 				}
 			}
@@ -220,6 +220,66 @@ FRigUnit_AddPhysicsBody_Execute()
 	}
 }
 
+//======================================================================================================================
+FRigUnit_HierarchyAutoCalculateCollision_Execute()
+{
+	DECLARE_SCOPE_HIERARCHICAL_COUNTER_RIGUNIT();
+
+	if (ExecuteContext.Hierarchy)
+	{
+		if (FRigPhysicsComponent* Component = Cast<FRigPhysicsComponent>(
+			ExecuteContext.Hierarchy->FindComponent(PhysicsBodyComponentKey)))
+		{
+			Component->AutoCalculateCollision(ExecuteContext.Hierarchy, MinAspectRatio);
+		}
+	}
+}
+
+//======================================================================================================================
+FRigUnit_HierarchySetDynamics_Execute()
+{
+	DECLARE_SCOPE_HIERARCHICAL_COUNTER_RIGUNIT();
+
+	if (ExecuteContext.Hierarchy)
+	{
+		if (FRigPhysicsComponent* Component = Cast<FRigPhysicsComponent>(
+			ExecuteContext.Hierarchy->FindComponent(PhysicsBodyComponentKey)))
+		{
+			Component->Dynamics = Dynamics;
+		}
+	}
+}
+
+//======================================================================================================================
+FRigUnit_HierarchySetCollision_Execute()
+{
+	DECLARE_SCOPE_HIERARCHICAL_COUNTER_RIGUNIT();
+
+	if (ExecuteContext.Hierarchy)
+	{
+		if (FRigPhysicsComponent* Component = Cast<FRigPhysicsComponent>(
+			ExecuteContext.Hierarchy->FindComponent(PhysicsBodyComponentKey)))
+		{
+			Component->Collision = Collision;
+		}
+	}
+}
+
+//======================================================================================================================
+FRigUnit_HierarchySetJoint_Execute()
+{
+	DECLARE_SCOPE_HIERARCHICAL_COUNTER_RIGUNIT();
+
+	if (ExecuteContext.Hierarchy)
+	{
+		if (FRigPhysicsComponent* Component = Cast<FRigPhysicsComponent>(
+			ExecuteContext.Hierarchy->FindComponent(PhysicsBodyComponentKey)))
+		{
+			Component->Joint= Joint;
+		}
+	}
+}
+
 //======================================================================================================================
 FRigUnit_HierarchySetPhysicsBodySourceBone_Execute()
 {
@@ -227,10 +287,10 @@ FRigUnit_HierarchySetPhysicsBodySourceBone_Execute()
 
 	if (ExecuteContext.Hierarchy)
 	{
-		if (FRigPhysicsComponent* PhysicsComponent = Cast<FRigPhysicsComponent>(
+		if (FRigPhysicsComponent* Component = Cast<FRigPhysicsComponent>(
 			ExecuteContext.Hierarchy->FindComponent(PhysicsBodyComponentKey)))
 		{
-			PhysicsComponent->BodySolverSettings.SourceBone = SourceBone;
+			Component->BodySolverSettings.SourceBone = SourceBone;
 		}
 	}
 }
@@ -242,10 +302,10 @@ FRigUnit_HierarchySetPhysicsBodyTargetBone_Execute()
 
 	if (ExecuteContext.Hierarchy)
 	{
-		if (FRigPhysicsComponent* PhysicsComponent = Cast<FRigPhysicsComponent>(
+		if (FRigPhysicsComponent* Component = Cast<FRigPhysicsComponent>(
 			ExecuteContext.Hierarchy->FindComponent(PhysicsBodyComponentKey)))
 		{
-			PhysicsComponent->BodySolverSettings.TargetBone = TargetBone;
+			Component->BodySolverSettings.TargetBone = TargetBone;
 		}
 	}
 }
@@ -257,11 +317,11 @@ FRigUnit_HierarchySetPhysicsBodyKinematicTarget_Execute()
 
 	if (ExecuteContext.Hierarchy)
 	{
-		if (FRigPhysicsComponent* PhysicsComponent = Cast<FRigPhysicsComponent>(
+		if (FRigPhysicsComponent* Component = Cast<FRigPhysicsComponent>(
 			ExecuteContext.Hierarchy->FindComponent(PhysicsBodyComponentKey)))
 		{
-			PhysicsComponent->KinematicTarget = KinematicTarget;
-			PhysicsComponent->KinematicTargetSpace = KinematicTargetSpace;
+			Component->KinematicTarget = KinematicTarget;
+			Component->KinematicTargetSpace = KinematicTargetSpace;
 		}
 	}
 }
@@ -273,10 +333,10 @@ FRigUnit_HierarchySetPhysicsBodySparseData_Execute()
 
 	if (ExecuteContext.Hierarchy)
 	{
-		if (FRigPhysicsComponent* PhysicsComponent = Cast<FRigPhysicsComponent>(
+		if (FRigPhysicsComponent* Component = Cast<FRigPhysicsComponent>(
 			ExecuteContext.Hierarchy->FindComponent(PhysicsBodyComponentKey)))
 		{
-			PhysicsComponent->BodyData.UpdateFromSparseData(Data);
+			Component->BodyData.UpdateFromSparseData(Data);
 		}
 	}
 }
@@ -287,10 +347,10 @@ FRigUnit_HierarchySetPhysicsBodyMovementType_Execute()
 	DECLARE_SCOPE_HIERARCHICAL_COUNTER_RIGUNIT();
 	if (ExecuteContext.Hierarchy)
 	{
-		if (FRigPhysicsComponent* PhysicsComponent = Cast<FRigPhysicsComponent>(
+		if (FRigPhysicsComponent* Component = Cast<FRigPhysicsComponent>(
 			ExecuteContext.Hierarchy->FindComponent(PhysicsBodyComponentKey)))
 		{
-			PhysicsComponent->BodyData.MovementType = MovementType;
+			Component->BodyData.MovementType = MovementType;
 		}
 	}
 }
@@ -301,10 +361,10 @@ FRigUnit_HierarchySetPhysicsBodyCollisionType_Execute()
 	DECLARE_SCOPE_HIERARCHICAL_COUNTER_RIGUNIT();
 	if (ExecuteContext.Hierarchy)
 	{
-		if (FRigPhysicsComponent* PhysicsComponent = Cast<FRigPhysicsComponent>(
+		if (FRigPhysicsComponent* Component = Cast<FRigPhysicsComponent>(
 			ExecuteContext.Hierarchy->FindComponent(PhysicsBodyComponentKey)))
 		{
-			PhysicsComponent->BodyData.CollisionType = CollisionType;
+			Component->BodyData.CollisionType = CollisionType;
 		}
 	}
 }
@@ -315,10 +375,10 @@ FRigUnit_HierarchySetPhysicsBodyGravityMultiplier_Execute()
 	DECLARE_SCOPE_HIERARCHICAL_COUNTER_RIGUNIT();
 	if (ExecuteContext.Hierarchy)
 	{
-		if (FRigPhysicsComponent* PhysicsComponent = Cast<FRigPhysicsComponent>(
+		if (FRigPhysicsComponent* Component = Cast<FRigPhysicsComponent>(
 			ExecuteContext.Hierarchy->FindComponent(PhysicsBodyComponentKey)))
 		{
-			PhysicsComponent->BodyData.GravityMultiplier = GravityMultiplier;
+			Component->BodyData.GravityMultiplier = GravityMultiplier;
 		}
 	}
 }
@@ -329,10 +389,10 @@ FRigUnit_HierarchySetPhysicsBodyPhysicsBlendWeight_Execute()
 	DECLARE_SCOPE_HIERARCHICAL_COUNTER_RIGUNIT();
 	if (ExecuteContext.Hierarchy)
 	{
-		if (FRigPhysicsComponent* PhysicsComponent = Cast<FRigPhysicsComponent>(
+		if (FRigPhysicsComponent* Component = Cast<FRigPhysicsComponent>(
 			ExecuteContext.Hierarchy->FindComponent(PhysicsBodyComponentKey)))
 		{
-			PhysicsComponent->BodyData.PhysicsBlendWeight = PhysicsBlendWeight;
+			Component->BodyData.PhysicsBlendWeight = PhysicsBlendWeight;
 		}
 	}
 }
@@ -343,10 +403,10 @@ FRigUnit_HierarchySetPhysicsBodyUseSkeletalAnimation_Execute()
 	DECLARE_SCOPE_HIERARCHICAL_COUNTER_RIGUNIT();
 	if (ExecuteContext.Hierarchy)
 	{
-		if (FRigPhysicsComponent* PhysicsComponent = Cast<FRigPhysicsComponent>(
+		if (FRigPhysicsComponent* Component = Cast<FRigPhysicsComponent>(
 			ExecuteContext.Hierarchy->FindComponent(PhysicsBodyComponentKey)))
 		{
-			PhysicsComponent->BodyData.bUseSkeletalAnimation = bUseSkeletalAnimation;
+			Component->BodyData.bUseSkeletalAnimation = bUseSkeletalAnimation;
 		}
 	}
 }
@@ -357,10 +417,10 @@ FRigUnit_HierarchySetPhysicsBodyUpdateKinematicFromSimulation_Execute()
 	DECLARE_SCOPE_HIERARCHICAL_COUNTER_RIGUNIT();
 	if (ExecuteContext.Hierarchy)
 	{
-		if (FRigPhysicsComponent* PhysicsComponent = Cast<FRigPhysicsComponent>(
+		if (FRigPhysicsComponent* Component = Cast<FRigPhysicsComponent>(
 			ExecuteContext.Hierarchy->FindComponent(PhysicsBodyComponentKey)))
 		{
-			PhysicsComponent->BodyData.bUpdateKinematicFromSimulation = bUpdateKinematicFromSimulation;
+			Component->BodyData.bUpdateKinematicFromSimulation = bUpdateKinematicFromSimulation;
 		}
 	}
 }
diff --git a/Engine/Plugins/Experimental/PhysicsControl/Source/ControlRigPhysics/Public/RigPhysicsComponent.h b/Engine/Plugins/Experimental/PhysicsControl/Source/ControlRigPhysics/Public/RigPhysicsComponent.h
--- a/Engine/Plugins/Experimental/PhysicsControl/Source/ControlRigPhysics/Public/RigPhysicsComponent.h
+++ b/Engine/Plugins/Experimental/PhysicsControl/Source/ControlRigPhysics/Public/RigPhysicsComponent.h
@@ -48,6 +48,11 @@ public:
 	UPROPERTY(BlueprintReadOnly, EditAnywhere, Category = CurrentPhysics)
 	ERigPhysicsKinematicTargetSpace KinematicTargetSpace = ERigPhysicsKinematicTargetSpace::OffsetInBoneSpace;
 
+	// Removes any existing collision, and replaces it with a shape calculated from the joint
+	// positions (if possible). The shape will be a single box.
+	// @param MinAspectRatio the minimum box extent, as a proportion of the maximum box extent.
+	void AutoCalculateCollision(URigHierarchy* InHierarchy, float MinAspectRatio = 0.25f);
+
 	virtual void Save(FArchive& Ar) override;
 	virtual void Load(FArchive& Ar) override;
 
diff --git a/Engine/Plugins/Experimental/PhysicsControl/Source/ControlRigPhysics/Public/RigPhysicsData.h b/Engine/Plugins/Experimental/PhysicsControl/Source/ControlRigPhysics/Public/RigPhysicsData.h
--- a/Engine/Plugins/Experimental/PhysicsControl/Source/ControlRigPhysics/Public/RigPhysicsData.h
+++ b/Engine/Plugins/Experimental/PhysicsControl/Source/ControlRigPhysics/Public/RigPhysicsData.h
@@ -111,6 +111,10 @@ struct CONTROLRIGPHYSICS_API FRigPhysicsCollisionBox : public FRigPhysicsCollisi
 {
 	GENERATED_BODY()
 
+	FRigPhysicsCollisionBox(const FTransform& InTM, const FVector& InExtents) : 
+		TM(InTM), Extents(InExtents) {}
+	FRigPhysicsCollisionBox() {}
+
 	friend FArchive& operator <<(FArchive& Ar, FRigPhysicsCollisionBox& Data)
 	{
 		Ar << Data.TM;
@@ -134,6 +138,11 @@ struct CONTROLRIGPHYSICS_API FRigPhysicsCollisionSphere : public FRigPhysicsColl
 {
 	GENERATED_BODY()
 
+	FRigPhysicsCollisionSphere(const FTransform& InTM, const float InRadius) :
+		TM(InTM), Radius(InRadius) {
+	}
+	FRigPhysicsCollisionSphere() {}
+
 	friend FArchive& operator <<(FArchive& Ar, FRigPhysicsCollisionSphere& Data)
 	{
 		Ar << Data.TM;
@@ -156,6 +165,10 @@ struct CONTROLRIGPHYSICS_API FRigPhysicsCollisionCapsule : public FRigPhysicsCol
 {
 	GENERATED_BODY()
 
+	FRigPhysicsCollisionCapsule(const FTransform& InTM, const float InRadius, const float InLength)
+		: TM(InTM), Radius(InRadius), Length(InLength) {}
+	FRigPhysicsCollisionCapsule() {}
+
 	friend FArchive& operator <<(FArchive& Ar, FRigPhysicsCollisionCapsule& Data)
 	{
 		Ar << Data.TM;
@@ -371,12 +384,12 @@ struct CONTROLRIGPHYSICS_API FRigPhysicsSolverSettings
 	float MaxDepenetrationVelocity = 0.0f;
 
 	/**
-	 * The recommended fixed timestep for the RBAN solver. Set to 0 to run with variable timestep (default).
+	 * The recommended fixed timestep for the RBAN solver. Set to 0 to run with variable timestep.
 	 * NOTE: If this value is non-zero and less than the current frame time, the simulation will step multiple times
 	 * which increases the cost.
 	 */
 	UPROPERTY(EditAnywhere, BlueprintReadWrite, Category = SolverSettings, meta = (ClampMin = 0))
-	float FixedTimeStep = 0.0f;
+	float FixedTimeStep = 0.01666667f;
 
 	// The maximum number of solver steps that can be made
 	UPROPERTY(EditAnywhere, BlueprintReadWrite, Category = SolverSettings, meta = (ClampMin = 1))
@@ -576,7 +589,7 @@ struct CONTROLRIGPHYSICS_API FRigPhysicsJoint
 
 	// Angular limit in degrees (twist, swing1, swing2). -ve indicates the movement is free
 	UPROPERTY(BlueprintReadWrite, EditAnywhere, Category = PhysicsJoint)
-	FVector AngularLimit = FVector::ZeroVector;
+	FVector AngularLimit = FVector(-1.0);
 
 	// Disable collisions between the parent and child bodies
 	UPROPERTY(BlueprintReadWrite, EditAnywhere, Category = PhysicsJoint)
diff --git a/Engine/Plugins/Experimental/PhysicsControl/Source/ControlRigPhysics/Public/RigPhysicsExecution.h b/Engine/Plugins/Experimental/PhysicsControl/Source/ControlRigPhysics/Public/RigPhysicsExecution.h
--- a/Engine/Plugins/Experimental/PhysicsControl/Source/ControlRigPhysics/Public/RigPhysicsExecution.h
+++ b/Engine/Plugins/Experimental/PhysicsControl/Source/ControlRigPhysics/Public/RigPhysicsExecution.h
@@ -128,33 +128,91 @@ struct CONTROLRIGPHYSICS_API FRigUnit_AddPhysicsBody : public FRigUnit_PhysicsBa
 	 * The solver to relate this new physics element to
 	 */
 	UPROPERTY(meta = (Input))
-	FRigComponentKey PhysicsSolverComponentKey;
+	FRigPhysicsBodySolverSettings Solver;
 
-	/*
-	 * The dynamics properties of the new physics element
-	 */
+	// The dynamics properties of the new physics element	
 	UPROPERTY(meta = (Input))
 	FRigPhysicsDynamics Dynamics;
 
-	/*
-	 * The collision properties of the new physics element
-	 */
+	// The collision properties of the new physics element
 	UPROPERTY(meta = (Input))
 	FRigPhysicsCollision Collision;
 
-	/**
-	 * The physics joint (typically to its parent) properties of the new physics element
-	 */
+	// The physics joint (typically to its parent) properties of the new physics element
 	UPROPERTY(meta = (Input))
 	FRigPhysicsJoint Joint;
 
-	/**
-	 * The runtime modifiable data
-	 */
+	// The runtime modifiable data
 	UPROPERTY(meta = (Input))
 	FPhysicsControlModifierData BodyData;
 };
 
+// Discards any existing collision data and replaces it with a box based on the joint positions.
+// Note that this must be called before the physics solver is instantiated/stepped.
+USTRUCT(meta = (DisplayName = "Automatically Calculate Collision", Varying))
+struct CONTROLRIGPHYSICS_API FRigUnit_HierarchyAutoCalculateCollision : public FRigUnit_PhysicsBaseMutable
+{
+	GENERATED_BODY()
+
+	RIGVM_METHOD()
+	virtual void Execute() override;
+
+	UPROPERTY(meta = (Input))
+	FRigComponentKey PhysicsBodyComponentKey;
+
+	// The minimum box extent, as a proportion of the maximum box extent.
+	UPROPERTY(meta = (Input))
+	float MinAspectRatio = 0.25f;
+};
+
+// Sets the mass for a physics component body
+USTRUCT(meta = (DisplayName = "Set Physics Dynamics Properties", Varying))
+struct CONTROLRIGPHYSICS_API FRigUnit_HierarchySetDynamics : public FRigUnit_PhysicsBaseMutable
+{
+	GENERATED_BODY()
+
+	RIGVM_METHOD()
+	virtual void Execute() override;
+
+	UPROPERTY(meta = (Input))
+	FRigComponentKey PhysicsBodyComponentKey;
+
+	UPROPERTY(meta = (Input))
+	FRigPhysicsDynamics Dynamics;
+};
+
+// Sets the collision for a physics component body
+USTRUCT(meta = (DisplayName = "Set Physics Collision Properties", Varying))
+struct CONTROLRIGPHYSICS_API FRigUnit_HierarchySetCollision : public FRigUnit_PhysicsBaseMutable
+{
+	GENERATED_BODY()
+
+	RIGVM_METHOD()
+	virtual void Execute() override;
+
+	UPROPERTY(meta = (Input))
+	FRigComponentKey PhysicsBodyComponentKey;
+
+	UPROPERTY(meta = (Input))
+	FRigPhysicsCollision Collision;
+};
+
+// Sets the collision for a physics component body
+USTRUCT(meta = (DisplayName = "Set Physics Joint Properties", Varying))
+struct CONTROLRIGPHYSICS_API FRigUnit_HierarchySetJoint : public FRigUnit_PhysicsBaseMutable
+{
+	GENERATED_BODY()
+
+	RIGVM_METHOD()
+	virtual void Execute() override;
+
+	UPROPERTY(meta = (Input))
+	FRigComponentKey PhysicsBodyComponentKey;
+
+	UPROPERTY(meta = (Input))
+	FRigPhysicsJoint Joint;
+};
+
 /**
  * Sets what bone is used as a source transform for the physics body. This is used as a kinematic target, and when
  * initializing the simulation.
