Allow RigidBodyWithControl and ControlRigPhysics to connect to Chaos Visual Debugger

diff --git a/Engine/Plugins/Experimental/PhysicsControl/Source/ControlRigPhysics/ControlRigPhysics.Build.cs b/Engine/Plugins/Experimental/PhysicsControl/Source/ControlRigPhysics/ControlRigPhysics.Build.cs
--- a/Engine/Plugins/Experimental/PhysicsControl/Source/ControlRigPhysics/ControlRigPhysics.Build.cs
+++ b/Engine/Plugins/Experimental/PhysicsControl/Source/ControlRigPhysics/ControlRigPhysics.Build.cs
@@ -20,16 +20,15 @@ public class ControlRigPhysics : ModuleRules
 			new string[]
 			{
 				"ControlRig",
-				"Chaos",
-				"ChaosCore",
 				"Core",
 				"CoreUObject",
 				"Engine",
-				"PhysicsCore",
 				"RigVM"
 			}
 		);
-		
+
+		SetupModulePhysicsSupport(Target);
+
 		if (Target.bBuildEditor == true)
 		{
 			PrivateDependencyModuleNames.AddRange(
diff --git a/Engine/Plugins/Experimental/PhysicsControl/Source/ControlRigPhysics/Private/RigPhysicsSimulation.cpp b/Engine/Plugins/Experimental/PhysicsControl/Source/ControlRigPhysics/Private/RigPhysicsSimulation.cpp
--- a/Engine/Plugins/Experimental/PhysicsControl/Source/ControlRigPhysics/Private/RigPhysicsSimulation.cpp
+++ b/Engine/Plugins/Experimental/PhysicsControl/Source/ControlRigPhysics/Private/RigPhysicsSimulation.cpp
@@ -36,6 +36,10 @@
 #include "PhysicsEngine/ConstraintInstance.h"
 #include "PhysicsEngine/BodySetup.h"
 
+#if WITH_CHAOS_VISUAL_DEBUGGER
+#include "ChaosVDRuntimeModule.h"
+#endif
+
 #include "Logging/LogMacros.h"
 
 DEFINE_LOG_CATEGORY(LogRigPhysics);
@@ -109,6 +113,12 @@ void FRigPhysicsSimulation::InitialiseSimulation(const FRigPhysicsSolverComponen
 	DestroyPhysicsSimulation();
 
 	Simulation = MakeShared<ImmediatePhysics::FSimulation>();
+
+#if WITH_CHAOS_VISUAL_DEBUGGER
+	Simulation->GetChaosVDContextData().Id = FChaosVDRuntimeModule::Get().GenerateUniqueID();
+	Simulation->GetChaosVDContextData().Type = static_cast<int32>(EChaosVDContextType::Solver);
+#endif
+
 	// This is needed so that when using a fixed timestep, velocities are rewound as well as
 	// positions. This is not only more accurate, but it's needed in order to get soft constraint
 	// behavior (in particular, for controls) that behave fairly independently of the control-rig
diff --git a/Engine/Plugins/Experimental/PhysicsControl/Source/PhysicsControl/PhysicsControl.Build.cs b/Engine/Plugins/Experimental/PhysicsControl/Source/PhysicsControl/PhysicsControl.Build.cs
--- a/Engine/Plugins/Experimental/PhysicsControl/Source/PhysicsControl/PhysicsControl.Build.cs
+++ b/Engine/Plugins/Experimental/PhysicsControl/Source/PhysicsControl/PhysicsControl.Build.cs
@@ -21,13 +21,12 @@ public class PhysicsControl : ModuleRules
 			{
 				"AnimationCore",
 				"AnimGraphRuntime",
-				"Chaos",
-				"ChaosCore",
 				"Core",
 				"CoreUObject",
-				"Engine",
-				"PhysicsCore"
+				"Engine"
 			}
 		);
+
+		SetupModulePhysicsSupport(Target);
 	}
 }
diff --git a/Engine/Plugins/Experimental/PhysicsControl/Source/PhysicsControl/Private/AnimNode_RigidBodyWithControl.cpp b/Engine/Plugins/Experimental/PhysicsControl/Source/PhysicsControl/Private/AnimNode_RigidBodyWithControl.cpp
--- a/Engine/Plugins/Experimental/PhysicsControl/Source/PhysicsControl/Private/AnimNode_RigidBodyWithControl.cpp
+++ b/Engine/Plugins/Experimental/PhysicsControl/Source/PhysicsControl/Private/AnimNode_RigidBodyWithControl.cpp
@@ -30,6 +30,10 @@
 
 #include "Chaos/PBDJointConstraintTypes.h"
 
+#if WITH_CHAOS_VISUAL_DEBUGGER
+#include "ChaosVDRuntimeModule.h"
+#endif
+
 #include UE_INLINE_GENERATED_CPP_BY_NAME(AnimNode_RigidBodyWithControl)
 
 LLM_DEFINE_TAG(Animation_RigidBodyWithControl);
@@ -990,6 +994,12 @@ void FAnimNode_RigidBodyWithControl::InitPhysics(const UAnimInstance* InAnimInst
 	if(bEnabled)
 	{
 		PhysicsSimulation = new ImmediatePhysics::FSimulation();
+
+#if WITH_CHAOS_VISUAL_DEBUGGER
+		PhysicsSimulation->GetChaosVDContextData().Id = FChaosVDRuntimeModule::Get().GenerateUniqueID();
+		PhysicsSimulation->GetChaosVDContextData().Type = static_cast<int32>(EChaosVDContextType::Solver);
+#endif
+
 		PhysicsSimulation->SetRewindVelocities(true);
 
 		const int32 NumBodies = PhysicsAssetToUse->SkeletalBodySetups.Num();
