Control Rig Physics - support setting drive parameters at runtime. Also sets/uses default names in the physics components

#rnx

diff --git a/Engine/Plugins/Experimental/PhysicsControl/Source/ControlRigPhysics/Private/RigPhysicsExecution.cpp b/Engine/Plugins/Experimental/PhysicsControl/Source/ControlRigPhysics/Private/RigPhysicsExecution.cpp
--- a/Engine/Plugins/Experimental/PhysicsControl/Source/ControlRigPhysics/Private/RigPhysicsExecution.cpp
+++ b/Engine/Plugins/Experimental/PhysicsControl/Source/ControlRigPhysics/Private/RigPhysicsExecution.cpp
@@ -280,6 +280,21 @@ FRigUnit_HierarchySetJoint_Execute()
 	}
 }
 
+//======================================================================================================================
+FRigUnit_HierarchySetJointDrive_Execute()
+{
+	DECLARE_SCOPE_HIERARCHICAL_COUNTER_RIGUNIT();
+
+	if (ExecuteContext.Hierarchy)
+	{
+		if (FRigPhysicsComponent* Component = Cast<FRigPhysicsComponent>(
+			ExecuteContext.Hierarchy->FindComponent(PhysicsBodyComponentKey)))
+		{
+			Component->Joint.Drive = Drive;
+		}
+	}
+}
+
 //======================================================================================================================
 FRigUnit_HierarchySetPhysicsBodySourceBone_Execute()
 {
diff --git a/Engine/Plugins/Experimental/PhysicsControl/Source/ControlRigPhysics/Private/RigPhysicsSimulation.cpp b/Engine/Plugins/Experimental/PhysicsControl/Source/ControlRigPhysics/Private/RigPhysicsSimulation.cpp
--- a/Engine/Plugins/Experimental/PhysicsControl/Source/ControlRigPhysics/Private/RigPhysicsSimulation.cpp
+++ b/Engine/Plugins/Experimental/PhysicsControl/Source/ControlRigPhysics/Private/RigPhysicsSimulation.cpp
@@ -496,57 +496,9 @@ void FRigPhysicsSimulation::InstantiatePhysicsJoints(const FRigPhysicsSolverComp
 			const FRigPhysicsDrive& Drive = PhysicsJoint.Drive;
 			if (ChildBodyRecord.JointHandle && Drive.bEnable)
 			{
-				float AngularSpring;
-				float AngularDamping;
-				const float MaxTorque = Drive.MaxTorque;
-
-				float LinearSpring;
-				float LinearDamping;
-				const float MaxForce = Drive.MaxForce;
-
-				UE::PhysicsControl::ConvertStrengthToSpringParams(
-					AngularSpring, AngularDamping,
-					Drive.AngularStrength, Drive.AngularDampingRatio, Drive.AngularExtraDamping);
-				UE::PhysicsControl::ConvertStrengthToSpringParams(
-					LinearSpring, LinearDamping,
-					Drive.LinearStrength, Drive.LinearDampingRatio, Drive.LinearExtraDamping);
-
 				if (Chaos::FPBDJointConstraintHandle* Constraint = ChildBodyRecord.JointHandle->GetConstraint())
 				{
 					Chaos::FPBDJointSettings Settings = Constraint->GetSettings();
-
-					// Linear
-					Settings.LinearDriveForceMode = Chaos::EJointForceMode::Acceleration;
-					const bool bLinearPositionDriveEnabled = LinearSpring > 0.0f;
-					const bool bLinearVelocityDriveEnabled = LinearDamping > 0.0f;
-					Settings.bLinearPositionDriveEnabled = { 
-						bLinearPositionDriveEnabled, bLinearPositionDriveEnabled, bLinearPositionDriveEnabled};
-					Settings.bLinearVelocityDriveEnabled = {
-						bLinearVelocityDriveEnabled, bLinearVelocityDriveEnabled, bLinearVelocityDriveEnabled};
-
-					// Angular
-					Settings.AngularDriveForceMode = Chaos::EJointForceMode::Acceleration;
-					// SLERP is generally better, but doesn't work if anything is locked
-					if (ConstraintSettings.AngularMotionTypes[0] == Chaos::EJointMotionType::Locked ||
-						ConstraintSettings.AngularMotionTypes[1] == Chaos::EJointMotionType::Locked ||
-						ConstraintSettings.AngularMotionTypes[2] == Chaos::EJointMotionType::Locked)
-					{
-						Settings.bAngularSwingPositionDriveEnabled = AngularSpring > 0.0f;
-						Settings.bAngularTwistPositionDriveEnabled = AngularSpring > 0.0f;
-						Settings.bAngularSwingVelocityDriveEnabled = AngularDamping > 0.0f;
-						Settings.bAngularTwistVelocityDriveEnabled = AngularDamping > 0.0f;
-					}
-					else
-					{
-						Settings.bAngularSLerpPositionDriveEnabled = AngularSpring > 0.0f;
-						Settings.bAngularSLerpVelocityDriveEnabled = AngularDamping > 0.0f;
-					}
-
-					Constraint->SetSettings(Settings);
-					Constraint->SetDriveParams(
-						Chaos::FVec3(LinearSpring), Chaos::FVec3(LinearDamping), Chaos::FVec3(MaxForce),
-						Chaos::FVec3(AngularSpring), Chaos::FVec3(AngularDamping), Chaos::FVec3(MaxTorque));
-
 					if (!Settings.bCollisionEnabled)
 					{
 						// Immediate physics doesn't actually honor the flag in settings, so do it manually
@@ -556,7 +508,7 @@ void FRigPhysicsSimulation::InstantiatePhysicsJoints(const FRigPhysicsSolverComp
 						IgnorePairs.Add(Pair);
 					}
 
-					// The target will be set in the pre-physics update
+					// The target and properties will be set in the pre-physics update
 				}
 			}
 		}
@@ -702,6 +654,67 @@ void FRigPhysicsSimulation::UpdatePrePhysics(
 				{
 					if (Chaos::FPBDJointConstraintHandle* Constraint = Record.JointHandle->GetConstraint())
 					{
+						// Set the drive strength etc
+						const FRigPhysicsDrive& Drive = PhysicsComponent->Joint.Drive;
+						if (Drive.bEnable)
+						{
+							float AngularSpring;
+							float AngularDamping;
+							const float MaxTorque = Drive.MaxTorque;
+
+							float LinearSpring;
+							float LinearDamping;
+							const float MaxForce = Drive.MaxForce;
+
+							UE::PhysicsControl::ConvertStrengthToSpringParams(
+								AngularSpring, AngularDamping,
+								Drive.AngularStrength, Drive.AngularDampingRatio, Drive.AngularExtraDamping);
+							UE::PhysicsControl::ConvertStrengthToSpringParams(
+								LinearSpring, LinearDamping,
+								Drive.LinearStrength, Drive.LinearDampingRatio, Drive.LinearExtraDamping);
+
+							Chaos::FPBDJointSettings Settings = Constraint->GetSettings();
+
+							// Linear
+							Settings.LinearDriveForceMode = Chaos::EJointForceMode::Acceleration;
+							const bool bLinearPositionDriveEnabled = LinearSpring > 0.0f;
+							const bool bLinearVelocityDriveEnabled = LinearDamping > 0.0f;
+							Settings.bLinearPositionDriveEnabled = {
+								bLinearPositionDriveEnabled, bLinearPositionDriveEnabled, bLinearPositionDriveEnabled };
+							Settings.bLinearVelocityDriveEnabled = {
+								bLinearVelocityDriveEnabled, bLinearVelocityDriveEnabled, bLinearVelocityDriveEnabled };
+
+							// Angular
+							Settings.AngularDriveForceMode = Chaos::EJointForceMode::Acceleration;
+							// SLERP is generally better, but doesn't work if anything is locked
+							if (Settings.AngularMotionTypes[0] == Chaos::EJointMotionType::Locked ||
+								Settings.AngularMotionTypes[1] == Chaos::EJointMotionType::Locked ||
+								Settings.AngularMotionTypes[2] == Chaos::EJointMotionType::Locked)
+							{
+								Settings.bAngularSwingPositionDriveEnabled = AngularSpring > 0.0f;
+								Settings.bAngularTwistPositionDriveEnabled = AngularSpring > 0.0f;
+								Settings.bAngularSwingVelocityDriveEnabled = AngularDamping > 0.0f;
+								Settings.bAngularTwistVelocityDriveEnabled = AngularDamping > 0.0f;
+							}
+							else
+							{
+								Settings.bAngularSLerpPositionDriveEnabled = AngularSpring > 0.0f;
+								Settings.bAngularSLerpVelocityDriveEnabled = AngularDamping > 0.0f;
+							}
+
+							Constraint->SetSettings(Settings);
+							Constraint->SetDriveParams(
+								Chaos::FVec3(LinearSpring), Chaos::FVec3(LinearDamping), Chaos::FVec3(MaxForce),
+								Chaos::FVec3(AngularSpring), Chaos::FVec3(AngularDamping), Chaos::FVec3(MaxTorque));
+						}
+						else
+						{
+							Constraint->SetDriveParams(
+								Chaos::FVec3(0), Chaos::FVec3(0), Chaos::FVec3(0),
+								Chaos::FVec3(0), Chaos::FVec3(0), Chaos::FVec3(0));
+						}
+
+						// Now set the actual target
 						const Chaos::FPBDJointSettings& Settings = Constraint->GetSettings();
 
 						if (Settings.AngularDriveStiffness.SquaredLength() > 0.0f ||
diff --git a/Engine/Plugins/Experimental/PhysicsControl/Source/ControlRigPhysics/Public/RigPhysicsComponent.h b/Engine/Plugins/Experimental/PhysicsControl/Source/ControlRigPhysics/Public/RigPhysicsComponent.h
--- a/Engine/Plugins/Experimental/PhysicsControl/Source/ControlRigPhysics/Public/RigPhysicsComponent.h
+++ b/Engine/Plugins/Experimental/PhysicsControl/Source/ControlRigPhysics/Public/RigPhysicsComponent.h
@@ -66,5 +66,9 @@ public:
 	virtual void OnAddedToHierarchy(URigHierarchy* InHierarchy, URigHierarchyController* InController) override;
 
 	virtual void OnRigHierarchyKeyChanged(const FRigHierarchyKey& InOldKey, const FRigHierarchyKey& InNewKey) override;
+
+	virtual FName GetDefaultComponentName() const override { return GetDefaultName(); }
+
+	static FName GetDefaultName() { return TEXT("PhysicsBody"); }
 };
 
diff --git a/Engine/Plugins/Experimental/PhysicsControl/Source/ControlRigPhysics/Public/RigPhysicsExecution.h b/Engine/Plugins/Experimental/PhysicsControl/Source/ControlRigPhysics/Public/RigPhysicsExecution.h
--- a/Engine/Plugins/Experimental/PhysicsControl/Source/ControlRigPhysics/Public/RigPhysicsExecution.h
+++ b/Engine/Plugins/Experimental/PhysicsControl/Source/ControlRigPhysics/Public/RigPhysicsExecution.h
@@ -3,6 +3,7 @@
 #pragma once
 
 #include "RigPhysicsComponent.h"
+#include "RigPhysicsSolverComponent.h"
 #include "RigPhysicsSimulation.h"
 #include "PhysicsControlData.h"
 
@@ -32,6 +33,11 @@ struct CONTROLRIGPHYSICS_API FRigUnit_AddPhysicsSolver : public FRigUnit_Physics
 {
 	GENERATED_BODY()
 
+	FRigUnit_AddPhysicsSolver()
+	{
+		Owner.Type = ERigElementType::Bone;
+	}
+
 	RIGVM_METHOD()
 	virtual void Execute() override;
 
@@ -62,6 +68,12 @@ struct CONTROLRIGPHYSICS_API FRigUnit_InstantiatePhysics : public FRigUnit_Physi
 {
 	GENERATED_BODY()
 
+	FRigUnit_InstantiatePhysics()
+	{
+		PhysicsSolverComponentKey.ElementKey.Type = ERigElementType::Bone;
+		PhysicsSolverComponentKey.Name = FRigPhysicsSolverComponent::GetDefaultName();
+	}
+
 	RIGVM_METHOD()
 	virtual void Execute() override;
 
@@ -81,6 +93,12 @@ struct CONTROLRIGPHYSICS_API FRigUnit_StepPhysicsSolver : public FRigUnit_Physic
 {
 	GENERATED_BODY()
 
+	FRigUnit_StepPhysicsSolver()
+	{
+		PhysicsSolverComponentKey.ElementKey.Type = ERigElementType::Bone;
+		PhysicsSolverComponentKey.Name = FRigPhysicsSolverComponent::GetDefaultName();
+	}
+
 	RIGVM_METHOD()
 	virtual void Execute() override;
 
@@ -112,6 +130,13 @@ struct CONTROLRIGPHYSICS_API FRigUnit_AddPhysicsBody : public FRigUnit_PhysicsBa
 {
 	GENERATED_BODY()
 
+	FRigUnit_AddPhysicsBody()
+	{
+		Owner.Type = ERigElementType::Bone;
+		Solver.PhysicsSolverComponentKey.ElementKey.Type = ERigElementType::Bone;
+		Solver.PhysicsSolverComponentKey.Name = FRigPhysicsComponent::GetDefaultName();
+	}
+
 	RIGVM_METHOD()
 	virtual void Execute() override;
 
@@ -154,6 +179,12 @@ struct CONTROLRIGPHYSICS_API FRigUnit_HierarchyAutoCalculateCollision : public F
 {
 	GENERATED_BODY()
 
+	FRigUnit_HierarchyAutoCalculateCollision()
+	{
+		PhysicsBodyComponentKey.ElementKey.Type = ERigElementType::Bone;
+		PhysicsBodyComponentKey.Name = FRigPhysicsComponent::GetDefaultName();
+	}
+
 	RIGVM_METHOD()
 	virtual void Execute() override;
 
@@ -171,6 +202,12 @@ struct CONTROLRIGPHYSICS_API FRigUnit_HierarchySetDynamics : public FRigUnit_Phy
 {
 	GENERATED_BODY()
 
+	FRigUnit_HierarchySetDynamics()
+	{
+		PhysicsBodyComponentKey.ElementKey.Type = ERigElementType::Bone;
+		PhysicsBodyComponentKey.Name = FRigPhysicsComponent::GetDefaultName();
+	}
+
 	RIGVM_METHOD()
 	virtual void Execute() override;
 
@@ -187,6 +224,12 @@ struct CONTROLRIGPHYSICS_API FRigUnit_HierarchySetCollision : public FRigUnit_Ph
 {
 	GENERATED_BODY()
 
+	FRigUnit_HierarchySetCollision()
+	{
+		PhysicsBodyComponentKey.ElementKey.Type = ERigElementType::Bone;
+		PhysicsBodyComponentKey.Name = FRigPhysicsComponent::GetDefaultName();
+	}
+
 	RIGVM_METHOD()
 	virtual void Execute() override;
 
@@ -197,12 +240,18 @@ struct CONTROLRIGPHYSICS_API FRigUnit_HierarchySetCollision : public FRigUnit_Ph
 	FRigPhysicsCollision Collision;
 };
 
-// Sets the collision for a physics component body
+// Sets the joint for a physics component body
 USTRUCT(meta = (DisplayName = "Set Physics Joint Properties", Varying))
 struct CONTROLRIGPHYSICS_API FRigUnit_HierarchySetJoint : public FRigUnit_PhysicsBaseMutable
 {
 	GENERATED_BODY()
 
+	FRigUnit_HierarchySetJoint()
+	{
+		PhysicsBodyComponentKey.ElementKey.Type = ERigElementType::Bone;
+		PhysicsBodyComponentKey.Name = FRigPhysicsComponent::GetDefaultName();
+	}
+
 	RIGVM_METHOD()
 	virtual void Execute() override;
 
@@ -213,6 +262,28 @@ struct CONTROLRIGPHYSICS_API FRigUnit_HierarchySetJoint : public FRigUnit_Physic
 	FRigPhysicsJoint Joint;
 };
 
+// Sets the joint drive for a physics component body
+USTRUCT(meta = (DisplayName = "Set Physics Joint Drive Properties", Varying))
+struct CONTROLRIGPHYSICS_API FRigUnit_HierarchySetJointDrive : public FRigUnit_PhysicsBaseMutable
+{
+	GENERATED_BODY()
+
+	FRigUnit_HierarchySetJointDrive()
+	{
+		PhysicsBodyComponentKey.ElementKey.Type = ERigElementType::Bone;
+		PhysicsBodyComponentKey.Name = FRigPhysicsComponent::GetDefaultName();
+	}
+
+	RIGVM_METHOD()
+	virtual void Execute() override;
+
+	UPROPERTY(meta = (Input))
+	FRigComponentKey PhysicsBodyComponentKey;
+
+	UPROPERTY(meta = (Input))
+	FRigPhysicsDrive Drive;
+};
+
 /**
  * Sets what bone is used as a source transform for the physics body. This is used as a kinematic target, and when
  * initializing the simulation.
@@ -222,6 +293,13 @@ struct CONTROLRIGPHYSICS_API FRigUnit_HierarchySetPhysicsBodySourceBone : public
 {
 	GENERATED_BODY()
 
+	FRigUnit_HierarchySetPhysicsBodySourceBone()
+	{
+		PhysicsBodyComponentKey.ElementKey.Type = ERigElementType::Bone;
+		PhysicsBodyComponentKey.Name = FRigPhysicsComponent::GetDefaultName();
+		SourceBone.Type = ERigElementType::Bone;
+	}
+
 	RIGVM_METHOD()
 	virtual void Execute() override;
 
@@ -240,6 +318,13 @@ struct CONTROLRIGPHYSICS_API FRigUnit_HierarchySetPhysicsBodyTargetBone : public
 {
 	GENERATED_BODY()
 
+	FRigUnit_HierarchySetPhysicsBodyTargetBone()
+	{
+		PhysicsBodyComponentKey.ElementKey.Type = ERigElementType::Bone;
+		PhysicsBodyComponentKey.Name = FRigPhysicsComponent::GetDefaultName();
+		TargetBone.Type = ERigElementType::Bone;
+	}
+
 	RIGVM_METHOD()
 	virtual void Execute() override;
 
@@ -261,6 +346,12 @@ struct CONTROLRIGPHYSICS_API FRigUnit_HierarchySetPhysicsBodySparseData : public
 {
 	GENERATED_BODY()
 
+	FRigUnit_HierarchySetPhysicsBodySparseData()
+	{
+		PhysicsBodyComponentKey.ElementKey.Type = ERigElementType::Bone;
+		PhysicsBodyComponentKey.Name = FRigPhysicsComponent::GetDefaultName();
+	}
+
 	RIGVM_METHOD()
 	virtual void Execute() override;
 
@@ -279,6 +370,12 @@ struct CONTROLRIGPHYSICS_API FRigUnit_HierarchySetPhysicsBodyKinematicTarget : p
 {
 	GENERATED_BODY()
 
+	FRigUnit_HierarchySetPhysicsBodyKinematicTarget()
+	{
+		PhysicsBodyComponentKey.ElementKey.Type = ERigElementType::Bone;
+		PhysicsBodyComponentKey.Name = FRigPhysicsComponent::GetDefaultName();
+	}
+
 	RIGVM_METHOD()
 	virtual void Execute() override;
 
@@ -301,6 +398,12 @@ struct CONTROLRIGPHYSICS_API FRigUnit_HierarchySetPhysicsBodyMovementType : publ
 {
 	GENERATED_BODY()
 
+	FRigUnit_HierarchySetPhysicsBodyMovementType()
+	{
+		PhysicsBodyComponentKey.ElementKey.Type = ERigElementType::Bone;
+		PhysicsBodyComponentKey.Name = FRigPhysicsComponent::GetDefaultName();
+	}
+
 	RIGVM_METHOD()
 	virtual void Execute() override;
 
@@ -319,6 +422,12 @@ struct CONTROLRIGPHYSICS_API FRigUnit_HierarchySetPhysicsBodyCollisionType : pub
 {
 	GENERATED_BODY()
 
+	FRigUnit_HierarchySetPhysicsBodyCollisionType()
+	{
+		PhysicsBodyComponentKey.ElementKey.Type = ERigElementType::Bone;
+		PhysicsBodyComponentKey.Name = FRigPhysicsComponent::GetDefaultName();
+	}
+
 	RIGVM_METHOD()
 	virtual void Execute() override;
 
@@ -337,6 +446,12 @@ struct CONTROLRIGPHYSICS_API FRigUnit_HierarchySetPhysicsBodyGravityMultiplier :
 {
 	GENERATED_BODY()
 
+	FRigUnit_HierarchySetPhysicsBodyGravityMultiplier()
+	{
+		PhysicsBodyComponentKey.ElementKey.Type = ERigElementType::Bone;
+		PhysicsBodyComponentKey.Name = FRigPhysicsComponent::GetDefaultName();
+	}
+
 	RIGVM_METHOD()
 	virtual void Execute() override;
 
@@ -356,6 +471,12 @@ struct CONTROLRIGPHYSICS_API FRigUnit_HierarchySetPhysicsBodyPhysicsBlendWeight
 {
 	GENERATED_BODY()
 
+	FRigUnit_HierarchySetPhysicsBodyPhysicsBlendWeight()
+	{
+		PhysicsBodyComponentKey.ElementKey.Type = ERigElementType::Bone;
+		PhysicsBodyComponentKey.Name = FRigPhysicsComponent::GetDefaultName();
+	}
+
 	RIGVM_METHOD()
 	virtual void Execute() override;
 
@@ -375,6 +496,12 @@ struct CONTROLRIGPHYSICS_API FRigUnit_HierarchySetPhysicsBodyUseSkeletalAnimatio
 {
 	GENERATED_BODY()
 
+	FRigUnit_HierarchySetPhysicsBodyUseSkeletalAnimation()
+	{
+		PhysicsBodyComponentKey.ElementKey.Type = ERigElementType::Bone;
+		PhysicsBodyComponentKey.Name = FRigPhysicsComponent::GetDefaultName();
+	}
+
 	RIGVM_METHOD()
 	virtual void Execute() override;
 
@@ -395,6 +522,12 @@ struct CONTROLRIGPHYSICS_API FRigUnit_HierarchySetPhysicsBodyUpdateKinematicFrom
 {
 	GENERATED_BODY()
 
+	FRigUnit_HierarchySetPhysicsBodyUpdateKinematicFromSimulation()
+	{
+		PhysicsBodyComponentKey.ElementKey.Type = ERigElementType::Bone;
+		PhysicsBodyComponentKey.Name = FRigPhysicsComponent::GetDefaultName();
+	}
+
 	RIGVM_METHOD()
 	virtual void Execute() override;
 
diff --git a/Engine/Plugins/Experimental/PhysicsControl/Source/ControlRigPhysics/Public/RigPhysicsSolverComponent.h b/Engine/Plugins/Experimental/PhysicsControl/Source/ControlRigPhysics/Public/RigPhysicsSolverComponent.h
--- a/Engine/Plugins/Experimental/PhysicsControl/Source/ControlRigPhysics/Public/RigPhysicsSolverComponent.h
+++ b/Engine/Plugins/Experimental/PhysicsControl/Source/ControlRigPhysics/Public/RigPhysicsSolverComponent.h
@@ -31,5 +31,10 @@ public:
 #if WITH_EDITOR
 	virtual const FSlateIcon& GetIconForUI() const override;
 #endif
+
+	virtual FName GetDefaultComponentName() const override { return GetDefaultName(); }
+
+	static FName GetDefaultName() { return TEXT("PhysicsSolver"); }
+
 };
 
