Support for physics controls in Control Rig Physics

diff --git a/Engine/Plugins/Experimental/PhysicsControl/Source/ControlRigPhysics/Private/RigPhysicsControlComponent.cpp b/Engine/Plugins/Experimental/PhysicsControl/Source/ControlRigPhysics/Private/RigPhysicsControlComponent.cpp
--- a/Engine/Plugins/Experimental/PhysicsControl/Source/ControlRigPhysics/Private/RigPhysicsControlComponent.cpp
+++ b/Engine/Plugins/Experimental/PhysicsControl/Source/ControlRigPhysics/Private/RigPhysicsControlComponent.cpp
@@ -1,14 +1,27 @@
 // Copyright Epic Games, Inc. All Rights Reserved.
 
 #include "RigPhysicsControlComponent.h"
+#include "PhysicsControlObjectVersion.h"
 
 //======================================================================================================================
 void FRigPhysicsControlComponent::Save(FArchive& Ar)
 {
+	Ar.UsingCustomVersion(FPhysicsControlObjectVersion::GUID);
+
 	FRigBaseComponent::Save(Ar);
 	Ar << ParentBodyComponentKey;
+	if (Ar.CustomVer(FPhysicsControlObjectVersion::GUID) >=
+		FPhysicsControlObjectVersion::ControlRigControlAddChildBodyComponentKey)
+	{
+		Ar << ChildBodyComponentKey;
+	}
 	Ar << ControlData;
 	Ar << ControlTarget;
+	if (Ar.CustomVer(FPhysicsControlObjectVersion::GUID) >=
+		FPhysicsControlObjectVersion::ControlRigControlAddChildBodyComponentKey)
+	{
+		Ar << ControlMultiplier;
+	}
 }
 
 //======================================================================================================================
@@ -16,17 +29,42 @@ void FRigPhysicsControlComponent::Load(FArchive& Ar)
 {
 	FRigBaseComponent::Load(Ar);
 	Ar << ParentBodyComponentKey;
+	if (Ar.CustomVer(FPhysicsControlObjectVersion::GUID) >=
+		FPhysicsControlObjectVersion::ControlRigControlAddChildBodyComponentKey)
+	{
+		Ar << ChildBodyComponentKey;
+	}
 	Ar << ControlData;
 	Ar << ControlTarget;
+	if (Ar.CustomVer(FPhysicsControlObjectVersion::GUID) >=
+		FPhysicsControlObjectVersion::ControlRigControlAddChildBodyComponentKey)
+	{
+		Ar << ControlMultiplier;
+	}
 }
 
 //======================================================================================================================
 bool FRigPhysicsControlComponent::CanBeAddedTo(
 	const FRigElementKey& InElementKey, const URigHierarchy* InHierarchy, FString* OutFailureReason) const
 {
-	if(OutFailureReason)
+	return true;
+}
+
+//======================================================================================================================
+void FRigPhysicsControlComponent::OnRigHierarchyKeyChanged(
+	const FRigHierarchyKey& InOldKey, const FRigHierarchyKey& InNewKey)
+{
+	FRigBaseComponent::OnRigHierarchyKeyChanged(InOldKey, InNewKey);
+
+	if (InOldKey.IsComponent() && InNewKey.IsComponent())
 	{
-		*OutFailureReason = TEXT("Physics control component is currently not supported.");
+		if (ParentBodyComponentKey == InOldKey.GetComponent())
+		{
+			ParentBodyComponentKey = InNewKey.GetComponent();
+		}
+		if (ChildBodyComponentKey == InOldKey.GetComponent())
+		{
+			ChildBodyComponentKey = InNewKey.GetComponent();
+		}
 	}
-	return false;
 }
diff --git a/Engine/Plugins/Experimental/PhysicsControl/Source/ControlRigPhysics/Private/RigPhysicsControlExecution.cpp b/Engine/Plugins/Experimental/PhysicsControl/Source/ControlRigPhysics/Private/RigPhysicsControlExecution.cpp
--- a/Engine/Plugins/Experimental/PhysicsControl/Source/ControlRigPhysics/Private/RigPhysicsControlExecution.cpp
+++ b/Engine/Plugins/Experimental/PhysicsControl/Source/ControlRigPhysics/Private/RigPhysicsControlExecution.cpp
@@ -31,7 +31,9 @@ FRigUnit_AddPhysicsControl_Execute()
 						Hierarchy->FindComponent(ControlComponentKey)))
 					{
 						Component->ParentBodyComponentKey = ParentBodyComponentKey;
+						Component->ChildBodyComponentKey = ChildBodyComponentKey;
 						Component->ControlData = ControlData;
+						Component->ControlMultiplier = ControlMultiplier;
 						Component->ControlTarget = ControlTarget;
 					}
 				}
@@ -43,3 +45,46 @@ FRigUnit_AddPhysicsControl_Execute()
 		UE_CONTROLRIG_RIGUNIT_REPORT_ERROR(TEXT("AddPhysicsControl can only be used during Setup"));
 	}
 }
+
+//======================================================================================================================
+FRigUnit_HierarchySetControlData_Execute()
+{
+	DECLARE_SCOPE_HIERARCHICAL_COUNTER_RIGUNIT();
+	if (ExecuteContext.Hierarchy)
+	{
+		if (FRigPhysicsControlComponent* Component = Cast<FRigPhysicsControlComponent>(
+			ExecuteContext.Hierarchy->FindComponent(PhysicsControlComponentKey)))
+		{
+			Component->ControlData = ControlData;
+		}
+	}
+}
+
+//======================================================================================================================
+FRigUnit_HierarchySetControlTarget_Execute()
+{
+	DECLARE_SCOPE_HIERARCHICAL_COUNTER_RIGUNIT();
+	if (ExecuteContext.Hierarchy)
+	{
+		if (FRigPhysicsControlComponent* Component = Cast<FRigPhysicsControlComponent>(
+			ExecuteContext.Hierarchy->FindComponent(PhysicsControlComponentKey)))
+		{
+			Component->ControlTarget = ControlTarget;
+		}
+	}
+}
+
+//======================================================================================================================
+FRigUnit_HierarchySetControlMultiplier_Execute()
+{
+	DECLARE_SCOPE_HIERARCHICAL_COUNTER_RIGUNIT();
+	if (ExecuteContext.Hierarchy)
+	{
+		if (FRigPhysicsControlComponent* Component = Cast<FRigPhysicsControlComponent>(
+			ExecuteContext.Hierarchy->FindComponent(PhysicsControlComponentKey)))
+		{
+			Component->ControlMultiplier = ControlMultiplier;
+		}
+	}
+}
+
diff --git a/Engine/Plugins/Experimental/PhysicsControl/Source/ControlRigPhysics/Private/RigPhysicsSimulation.cpp b/Engine/Plugins/Experimental/PhysicsControl/Source/ControlRigPhysics/Private/RigPhysicsSimulation.cpp
--- a/Engine/Plugins/Experimental/PhysicsControl/Source/ControlRigPhysics/Private/RigPhysicsSimulation.cpp
+++ b/Engine/Plugins/Experimental/PhysicsControl/Source/ControlRigPhysics/Private/RigPhysicsSimulation.cpp
@@ -2,6 +2,7 @@
 
 #include "RigPhysicsSimulation.h"
 #include "RigPhysicsComponent.h"
+#include "RigPhysicsControlComponent.h"
 #include "RigPhysicsSolverComponent.h"
 
 #include "AnimNode_RigidBodyWithControl.h"
@@ -44,6 +45,32 @@
 
 DEFINE_LOG_CATEGORY(LogRigPhysics);
 
+constexpr int32 ConstraintChildIndex = 0;
+constexpr int32 ConstraintParentIndex = 1;
+
+//======================================================================================================================
+// Danny TODO remove duplication with RigidBodyWithControl
+static void SetConstraintEnabled(ImmediatePhysics::FJointHandle* const JointHandle, const bool bIsEnabled)
+{
+	if (JointHandle)
+	{
+		if (ImmediatePhysics::FJointHandle::FChaosConstraintHandle* ConstraintHandle = JointHandle->GetConstraint())
+		{
+			ConstraintHandle->SetConstraintEnabled(bIsEnabled);
+
+			if (!bIsEnabled)
+			{
+				// The call above disables the constraint, but if any actor in the simulation is
+				// flagged as dirty, it gets re-enabled! See UE-204006 TODO remove this workaround
+				// when the bug has been fixed.
+				ConstraintHandle->SetDriveParams(
+					Chaos::FVec3::ZeroVector, Chaos::FVec3::ZeroVector, Chaos::FVec3::ZeroVector,
+					Chaos::FVec3::ZeroVector, Chaos::FVec3::ZeroVector, Chaos::FVec3::ZeroVector);
+			}
+		}
+	}
+}
+
 //======================================================================================================================
 FRigPhysicsSimulation::FRigPhysicsSimulation(TObjectPtr<UControlRig> InOwningControlRig)
 	: FRigPhysicsSimulationBase(FRigPhysicsSimulation::StaticStruct()), OwningControlRig(InOwningControlRig)
@@ -124,7 +151,47 @@ void FRigPhysicsSimulation::InitialiseSimulation(const FRigPhysicsSolverComponen
 	// behavior (in particular, for controls) that behave fairly independently of the control-rig
 	// tick rate.
 	Simulation->SetRewindVelocities(true);
+}
 
+//======================================================================================================================
+void FRigPhysicsSimulation::InitialiseControlRecords(const FRigPhysicsSolverComponent* SolverComponent)
+{
+	ensure(ControlRecords.IsEmpty());
+
+	const URigHierarchy* Hierarchy = OwningControlRig->GetHierarchy();
+	if (!Hierarchy || !SolverComponent)
+	{
+		return;
+	}
+
+	FRigComponentKey SolverComponentKey = SolverComponent->GetKey();
+
+	TArray<FRigComponentKey> AllComponentKeys = Hierarchy->GetAllComponentKeys();
+
+	// All the components in this simulation
+	TArray<FRigComponentKey> UnsortedBodyComponentKeys;
+
+	for (FRigComponentKey ComponentKey : AllComponentKeys)
+	{
+		if (const FRigPhysicsControlComponent* ControlComponent = Cast<FRigPhysicsControlComponent>(
+			Hierarchy->FindComponent(ComponentKey)))
+		{
+			if (ShouldComponentBeInSimulation(SolverComponentKey, ControlComponent->ChildBodyComponentKey))
+			{
+				if (!ControlComponent->ParentBodyComponentKey.IsValid() || 
+					ShouldComponentBeInSimulation(SolverComponentKey, ControlComponent->ParentBodyComponentKey))
+				{
+					// Just make the record for now - it will be instantiated later
+					ControlRecords.Add(ComponentKey);
+				}
+			}
+		}
+	}
+}
+
+//======================================================================================================================
+void FRigPhysicsSimulation::InitialiseBodyRecords(const FRigPhysicsSolverComponent* SolverComponent)
+{
 	ensure(BodyRecords.IsEmpty());
 
 	const URigHierarchy* Hierarchy = OwningControlRig->GetHierarchy();
@@ -140,19 +207,19 @@ void FRigPhysicsSimulation::InitialiseSimulation(const FRigPhysicsSolverComponen
 	// All the components in this simulation
 	TArray<FRigComponentKey> UnsortedBodyComponentKeys;
 
-	for (FRigComponentKey RigComponentKey : AllComponentKeys)
+	for (FRigComponentKey ComponentKey : AllComponentKeys)
 	{
-		if (ShouldComponentBeInSimulation(SolverComponentKey, RigComponentKey))
+		if (ShouldComponentBeInSimulation(SolverComponentKey, ComponentKey))
 		{
 			// Just make the record for now - it will be instantiated later
-			BodyRecords.Add(RigComponentKey);
-			UnsortedBodyComponentKeys.Add(RigComponentKey);
+			BodyRecords.Add(ComponentKey);
+			UnsortedBodyComponentKeys.Add(ComponentKey);
 		}
 	}
 
 	// Sort the component keys according to the traversal of their element (i.e. from root to leaf)
 	SortedBodyComponentKeys.Empty(UnsortedBodyComponentKeys.Num());
-	Hierarchy->Traverse([this, &UnsortedBodyComponentKeys](FRigBaseElement* Element, bool& bContinue) 
+	Hierarchy->Traverse([this, &UnsortedBodyComponentKeys](FRigBaseElement* Element, bool& bContinue)
 		{
 			const FRigElementKey& Key = Element->GetKey();
 			for (const FRigComponentKey& ComponentKey : UnsortedBodyComponentKeys)
@@ -182,11 +249,27 @@ void FRigPhysicsSimulation::DestroyPhysicsSimulation()
 	}
 	BodyRecords.Empty();
 
-	for (ImmediatePhysics::FActorHandle* ActorHandle : AdditionalActorHandles)
+	for (TPair<FRigComponentKey, FRigControlRecord>& ControlRecordPair : ControlRecords)
+	{
+		FRigControlRecord& Record = ControlRecordPair.Value;
+		if (Record.JointHandle)
+		{
+			Simulation->DestroyJoint(Record.JointHandle);
+		}
+	}
+	ControlRecords.Empty();
+
+	if (CollisionActorHandle)
 	{
-		Simulation->DestroyActor(ActorHandle);
+		Simulation->DestroyActor(CollisionActorHandle);
 	}
-	AdditionalActorHandles.Empty();
+	CollisionActorHandle = nullptr;
+
+	if (SimulationActorHandle)
+	{
+		Simulation->DestroyActor(SimulationActorHandle);
+	}
+	SimulationActorHandle = nullptr;
 
 	Simulation.Reset();
 }
@@ -326,18 +409,29 @@ void FRigPhysicsSimulation::Instantiate(const FRigPhysicsSolverComponent* Solver
 
 	InitialiseSimulation(SolverComponent);
 
+	InitialiseBodyRecords(SolverComponent);
+
+	InitialiseControlRecords(SolverComponent);
+
 	InstantiateSolverCollision(SolverComponent);
 
 	InstantiatePhysicsBodies(SolverComponent);
 
 	InstantiatePhysicsJoints(SolverComponent);
 
+	InstantiateControls(SolverComponent);
+
 	bNeedToInstantiate = false;
 }
 
 //======================================================================================================================
 void FRigPhysicsSimulation::InstantiateSolverCollision(const FRigPhysicsSolverComponent* SolverComponent)
 {
+	// Always create a world actor at the origin, for attaching controls to. 
+	SimulationActorHandle = CreateBody(
+		FName(TEXT("Simulation")), FRigPhysicsCollision(), nullptr, FTransform());
+
+	// Optionally create an object to contain environment collision
 	if (!SolverComponent->SolverSettings.Collision.IsEmpty())
 	{
 		// When we make these additional collision shapes, their actors are all considered to be at
@@ -345,9 +439,46 @@ void FRigPhysicsSimulation::InstantiateSolverCollision(const FRigPhysicsSolverCo
 		FTransform BodyRelSimSpaceTM = ConvertCollisionSpaceTransformToSimSpace(
 			SolverComponent->SolverSettings, FTransform());
 
-		ImmediatePhysics::FActorHandle* EnvironmentActorHandle = CreateBody(
+		CollisionActorHandle = CreateBody(
 			FName(TEXT("Environment")), SolverComponent->SolverSettings.Collision, nullptr, BodyRelSimSpaceTM);
-		AdditionalActorHandles.Add(EnvironmentActorHandle);
+	}
+}
+
+//======================================================================================================================
+void FRigPhysicsSimulation::InstantiatePhysicsBodies(const FRigPhysicsSolverComponent* SolverComponent)
+{
+	const URigHierarchy* Hierarchy = OwningControlRig->GetHierarchy();
+
+	const FRigPhysicsSolverSettings& SolverSettings = SolverComponent->SolverSettings;
+	const FRigPhysicsSimulationSpaceSettings& SimulationSpaceSettings = SolverComponent->SimulationSpaceSettings;
+
+	for (TPair<FRigComponentKey, FRigBodyRecord>& BodyRecordPair : BodyRecords)
+	{
+		const FRigComponentKey& ComponentKey = BodyRecordPair.Key;
+		FRigBodyRecord& Record = BodyRecordPair.Value;
+
+		if (const FRigPhysicsComponent* PhysicsComponent = Cast<FRigPhysicsComponent>(
+			Hierarchy->FindComponent(ComponentKey)))
+		{
+			const FRigPhysicsCollision& Collision = PhysicsComponent->Collision;
+			const FRigPhysicsDynamics& Dynamics = PhysicsComponent->Dynamics;
+			const FPhysicsControlModifierData& BodyData = PhysicsComponent->BodyData;
+
+			FRigElementKey SourceKey = PhysicsComponent->BodySolverSettings.SourceBone;
+			if (!SourceKey.IsValid())
+			{
+				SourceKey = ComponentKey.ElementKey;
+			}
+
+			// What should we do if the key is not valid? 
+			if (SourceKey.IsValid())
+			{
+				FTransform SourceComponentSpaceTM = Hierarchy->GetGlobalTransform(SourceKey);
+				const FTransform SourceSimulationSpaceTM =
+					ConvertComponentSpaceTransformToSimSpace(SolverComponent->SolverSettings, SourceComponentSpaceTM);
+				Record.ActorHandle = CreateBody(ComponentKey.ElementKey.Name, Collision, &Dynamics, SourceSimulationSpaceTM);
+			}
+		}
 	}
 }
 
@@ -517,43 +648,307 @@ void FRigPhysicsSimulation::InstantiatePhysicsJoints(const FRigPhysicsSolverComp
 }
 
 //======================================================================================================================
-void FRigPhysicsSimulation::InstantiatePhysicsBodies(const FRigPhysicsSolverComponent* SolverComponent)
+void FRigPhysicsSimulation::InstantiateControls(const FRigPhysicsSolverComponent* SolverComponent)
 {
 	const URigHierarchy* Hierarchy = OwningControlRig->GetHierarchy();
 
-	const FRigPhysicsSolverSettings& SolverSettings = SolverComponent->SolverSettings;
-	const FRigPhysicsSimulationSpaceSettings& SimulationSpaceSettings = SolverComponent->SimulationSpaceSettings;
-
-	for (TPair<FRigComponentKey, FRigBodyRecord>& BodyRecordPair : BodyRecords)
+	for (TPair<FRigComponentKey, FRigControlRecord>& ControlRecordPair : ControlRecords)
 	{
-		const FRigComponentKey& ComponentKey = BodyRecordPair.Key;
-		FRigBodyRecord& Record = BodyRecordPair.Value;
+		const FRigComponentKey& ComponentKey = ControlRecordPair.Key;
+		FRigControlRecord& ControlRecord = ControlRecordPair.Value;
 
-		if (const FRigPhysicsComponent* PhysicsComponent = Cast<FRigPhysicsComponent>(
+		if (const FRigPhysicsControlComponent* ControlComponent = Cast<FRigPhysicsControlComponent>(
 			Hierarchy->FindComponent(ComponentKey)))
 		{
-			const FRigPhysicsCollision& Collision = PhysicsComponent->Collision;
-			const FRigPhysicsDynamics& Dynamics = PhysicsComponent->Dynamics;
-			const FPhysicsControlModifierData& BodyData = PhysicsComponent->BodyData;
+			if (FRigBodyRecord* ChildBodyRecord = BodyRecords.Find(ControlComponent->ChildBodyComponentKey))
+			{
+				// this can be nullptr - just means a global control
+				FRigBodyRecord* ParentBodyRecord = BodyRecords.Find(ControlComponent->ParentBodyComponentKey);
+
+				ImmediatePhysics::FActorHandle* const ChildBodyHandle = ChildBodyRecord->ActorHandle;
+				ImmediatePhysics::FActorHandle* const ParentBodyHandle = 
+					ParentBodyRecord ? ParentBodyRecord->ActorHandle : SimulationActorHandle;
+
+				// This handles nullptrs. The constraint is created disabled - it will be updated in pre-physics
+				ControlRecord.JointHandle = CreatePhysicsJoint(Simulation.Get(), ChildBodyHandle, ParentBodyHandle);
+
+				if (!ControlRecord.JointHandle)
+				{
+					UE_LOG(LogRigPhysics, Warning,
+						TEXT("Unable to create control constraint for %s"), *ControlComponent->GetKey().ToString());
+				}
+			}
+		}
+	}
+}
+
+//======================================================================================================================
+void FRigPhysicsSimulation::UpdateBodyPrePhysics(
+	const FRigVMExecuteContext& ExecuteContext,
+	const FRigPhysicsSolverComponent* SolverComponent,
+	const FRigBodyRecord& Record, 
+	const FRigPhysicsComponent* PhysicsComponent,
+	const FSimulationSpaceData& SimulationSpaceData)
+{
+	URigHierarchy* Hierarchy = OwningControlRig->GetHierarchy();
+	const FRigComponentKey ComponentKey = PhysicsComponent->GetKey();
+
+	UpdateBodyFromModifierData(
+		Record.ActorHandle, Simulation.Get(), PhysicsComponent->BodyData, SimulationSpaceData.Gravity);
+
+	if (Record.ActorHandle->GetIsKinematic())
+	{
+		// Get the target in component space, and then convert it into sim space if necessary.
 
+		// If the target is already in component space, then that's all we need
+		FTransform KinematicTargetCS;
+		switch (PhysicsComponent->KinematicTargetSpace)
+		{
+		case ERigPhysicsKinematicTargetSpace::Component:
+		{
+			KinematicTargetCS = PhysicsComponent->KinematicTarget;
+			break;
+		}
+		case ERigPhysicsKinematicTargetSpace::World:
+		{
+			// Record.KinematicTarget * SimulationSpaceState.ComponentTM.Inverse();
+			KinematicTargetCS = PhysicsComponent->KinematicTarget.GetRelativeTransform(
+				SimulationSpaceState.ComponentTM);
+			break;
+		}
+		default:
+		{
+			// All the other options are relative to a bone, so the first task is to get
+			// that, which will be in component space
 			FRigElementKey SourceKey = PhysicsComponent->BodySolverSettings.SourceBone;
 			if (!SourceKey.IsValid())
 			{
 				SourceKey = ComponentKey.ElementKey;
 			}
-
-			// What should we do if the key is not valid? 
 			if (SourceKey.IsValid())
 			{
 				FTransform SourceComponentSpaceTM = Hierarchy->GetGlobalTransform(SourceKey);
-				const FTransform SourceSimulationSpaceTM =
-					ConvertComponentSpaceTransformToSimSpace(SolverComponent->SolverSettings, SourceComponentSpaceTM);
-				Record.ActorHandle = CreateBody(ComponentKey.ElementKey.Name, Collision, &Dynamics, SourceSimulationSpaceTM);
+
+				switch (PhysicsComponent->KinematicTargetSpace)
+				{
+				case ERigPhysicsKinematicTargetSpace::OffsetInBoneSpace:
+				{
+					KinematicTargetCS = PhysicsComponent->KinematicTarget * SourceComponentSpaceTM;
+					break;
+				}
+				case ERigPhysicsKinematicTargetSpace::OffsetInWorldSpace:
+				{
+					// Convert the bone to WS, apply the target, and convert back
+					FTransform BoneWS = SourceComponentSpaceTM * SimulationSpaceState.ComponentTM;
+					FTransform KinematicTargetWS = BoneWS * PhysicsComponent->KinematicTarget;
+					// Danny TODO figure out which of the GetRelativeTransform versions this is
+					KinematicTargetCS = SimulationSpaceState.ComponentTM.Inverse() * KinematicTargetWS;
+					break;
+				}
+				case ERigPhysicsKinematicTargetSpace::OffsetInComponentSpace:
+				{
+					KinematicTargetCS = SourceComponentSpaceTM * PhysicsComponent->KinematicTarget;
+					break;
+				}
+				default:
+				{
+					UE_CONTROLRIG_RIGUNIT_REPORT_ERROR(TEXT("Kinematic target space is not valid"));
+				}
+				}
 			}
 		}
+		}
+		FTransform KinematicTargetTM = ConvertComponentSpaceTransformToSimSpace(
+			SolverComponent->SolverSettings, KinematicTargetCS);
+		Record.ActorHandle->SetKinematicTarget(KinematicTargetTM);
 	}
 }
 
+//======================================================================================================================
+void FRigPhysicsSimulation::UpdateJointPrePhysics(
+	FRigBodyRecord&             Record, 
+	const FRigPhysicsComponent* PhysicsComponent, 
+	const URigHierarchy*        Hierarchy, 
+	const float                 DeltaTime)
+{
+	const FRigComponentKey ComponentKey = PhysicsComponent->GetKey();
+
+	// Now update the joint targets
+	if (Record.JointHandle)
+	{
+		if (Chaos::FPBDJointConstraintHandle* Constraint = Record.JointHandle->GetConstraint())
+		{
+			// Set the drive strength etc
+			const FRigPhysicsDrive& Drive = PhysicsComponent->Joint.Drive;
+			if (Drive.bEnable)
+			{
+				float AngularSpring;
+				float AngularDamping;
+				const float MaxTorque = Drive.MaxTorque;
+
+				float LinearSpring;
+				float LinearDamping;
+				const float MaxForce = Drive.MaxForce;
+
+				UE::PhysicsControl::ConvertStrengthToSpringParams(
+					AngularSpring, AngularDamping,
+					Drive.AngularStrength, Drive.AngularDampingRatio, Drive.AngularExtraDamping);
+				UE::PhysicsControl::ConvertStrengthToSpringParams(
+					LinearSpring, LinearDamping,
+					Drive.LinearStrength, Drive.LinearDampingRatio, Drive.LinearExtraDamping);
+
+				Chaos::FPBDJointSettings Settings = Constraint->GetSettings();
+
+				// Linear
+				Settings.LinearDriveForceMode = Chaos::EJointForceMode::Acceleration;
+				const bool bLinearPositionDriveEnabled = LinearSpring > 0.0f;
+				const bool bLinearVelocityDriveEnabled = LinearDamping > 0.0f;
+				Settings.bLinearPositionDriveEnabled = {
+					bLinearPositionDriveEnabled, bLinearPositionDriveEnabled, bLinearPositionDriveEnabled };
+				Settings.bLinearVelocityDriveEnabled = {
+					bLinearVelocityDriveEnabled, bLinearVelocityDriveEnabled, bLinearVelocityDriveEnabled };
+
+				// Angular
+				Settings.AngularDriveForceMode = Chaos::EJointForceMode::Acceleration;
+				// SLERP is generally better, but doesn't work if anything is locked
+				if (Settings.AngularMotionTypes[0] == Chaos::EJointMotionType::Locked ||
+					Settings.AngularMotionTypes[1] == Chaos::EJointMotionType::Locked ||
+					Settings.AngularMotionTypes[2] == Chaos::EJointMotionType::Locked)
+				{
+					Settings.bAngularSwingPositionDriveEnabled = AngularSpring > 0.0f;
+					Settings.bAngularTwistPositionDriveEnabled = AngularSpring > 0.0f;
+					Settings.bAngularSwingVelocityDriveEnabled = AngularDamping > 0.0f;
+					Settings.bAngularTwistVelocityDriveEnabled = AngularDamping > 0.0f;
+				}
+				else
+				{
+					Settings.bAngularSLerpPositionDriveEnabled = AngularSpring > 0.0f;
+					Settings.bAngularSLerpVelocityDriveEnabled = AngularDamping > 0.0f;
+				}
+
+				Constraint->SetSettings(Settings);
+				Constraint->SetDriveParams(
+					Chaos::FVec3(LinearSpring), Chaos::FVec3(LinearDamping), Chaos::FVec3(MaxForce),
+					Chaos::FVec3(AngularSpring), Chaos::FVec3(AngularDamping), Chaos::FVec3(MaxTorque));
+			}
+			else
+			{
+				Constraint->SetDriveParams(
+					Chaos::FVec3(0), Chaos::FVec3(0), Chaos::FVec3(0),
+					Chaos::FVec3(0), Chaos::FVec3(0), Chaos::FVec3(0));
+			}
+
+			// Now set the actual target
+			const Chaos::FPBDJointSettings& Settings = Constraint->GetSettings();
+
+			if (Settings.AngularDriveStiffness.SquaredLength() > 0.0f ||
+				Settings.AngularDriveDamping.SquaredLength() > 0.0f ||
+				Settings.LinearDriveStiffness.SquaredLength() > 0.0f ||
+				Settings.LinearDriveDamping.SquaredLength() > 0.0f)
+			{
+				// Multiplier on the velocity calculated from the current and previous target
+				float TargetVelocityMultiplier = 1.0f;
+
+				FRigElementKey ChildSourceKey;
+				if (const FRigPhysicsComponent* ChildPhysicsComponent = Cast<FRigPhysicsComponent>(
+					Hierarchy->FindComponent(ComponentKey)))
+				{
+					ChildSourceKey = ChildPhysicsComponent->BodySolverSettings.SourceBone;
+					if (!ChildSourceKey.IsValid())
+					{
+						ChildSourceKey = ComponentKey.ElementKey;
+					}
+					TargetVelocityMultiplier =
+						ChildPhysicsComponent->Joint.Drive.SkeletalAnimationVelocityMultiplier;
+				}
+
+				FRigElementKey ParentSourceKey;
+				if (const FRigPhysicsComponent* ParentPhysicsComponent = Cast<FRigPhysicsComponent>(
+					Hierarchy->FindComponent(Record.ParentBodyComponentKey)))
+				{
+					ParentSourceKey = ParentPhysicsComponent->BodySolverSettings.SourceBone;
+					if (!ParentSourceKey.IsValid())
+					{
+						ParentSourceKey = Record.ParentBodyComponentKey.ElementKey;
+					}
+				}
+
+				if (!ChildSourceKey.IsValid() || !ParentSourceKey.IsValid())
+				{
+					return;
+				}
+
+				// Note that the drive operates between a parent and child part, so we
+				// don't need to worry about global/component (etc) space.
+				const FTransform ChildTM = Hierarchy->GetGlobalTransform(ChildSourceKey);
+				const FTransform ParentTM = Hierarchy->GetGlobalTransform(ParentSourceKey);
+
+				const UE::PhysicsControl::FPosQuat ComponentSpaceParentFrameTM = 
+					Settings.ConnectorTransforms[1] * ParentTM;
+				const UE::PhysicsControl::FPosQuat ComponentSpaceChildFrameTM = 
+					Settings.ConnectorTransforms[0] * ChildTM;
+
+				const UE::PhysicsControl::FPosQuat DriveTargetTM =
+					ComponentSpaceParentFrameTM.Inverse() * ComponentSpaceChildFrameTM;
+
+				Constraint->SetLinearDrivePositionTarget(DriveTargetTM.GetTranslation());
+				Constraint->SetAngularDrivePositionTarget(DriveTargetTM.GetRotation());
+
+				if (Record.PreviousDriveTargetUpdateCounter + 1 == UpdateCounter &&
+					TargetVelocityMultiplier > 0.0f)
+				{
+					if (DeltaTime > SMALL_NUMBER)
+					{
+						const UE::PhysicsControl::FPosQuat DriveTargetTMDelta =
+							DriveTargetTM * Record.PreviousDriveTargetTM.Inverse();
+						const FVector Velocity = DriveTargetTMDelta.GetTranslation() / DeltaTime;
+						const FVector AngularVelocity =
+							DriveTargetTMDelta.GetRotation().GetShortestArcWith(
+								FQuat::Identity).ToRotationVector() / DeltaTime;
+						if (!Velocity.ContainsNaN() && !AngularVelocity.ContainsNaN())
+						{
+							Constraint->SetLinearDriveVelocityTarget(Velocity * TargetVelocityMultiplier);
+							Constraint->SetAngularDriveVelocityTarget(AngularVelocity * TargetVelocityMultiplier);
+						}
+					}
+				}
+				else
+				{
+					Constraint->SetLinearDriveVelocityTarget(Chaos::FVec3(0));
+					Constraint->SetAngularDriveVelocityTarget(Chaos::FVec3(0));
+				}
+				Record.PreviousDriveTargetUpdateCounter = UpdateCounter;
+				Record.PreviousDriveTargetTM = DriveTargetTM;
+			}
+		}
+	}
+}
+
+//======================================================================================================================
+static UE::PhysicsControl::FPosQuat CalculateTargetTM(
+	const URigHierarchy* Hierarchy,
+	const Chaos::FPBDJointSettings& JointSettings,
+	const FRigPhysicsControlComponent* ControlComponent)
+{
+	const FTransform ChildTM = ControlComponent->ChildBodyComponentKey.IsValid() ?
+		Hierarchy->GetGlobalTransform(ControlComponent->ChildBodyComponentKey.ElementKey) :
+		FTransform();
+
+	const UE::PhysicsControl::FPosQuat ChildTargetTM =
+		UE::PhysicsControl::FPosQuat(ChildTM) * 
+		UE::PhysicsControl::FPosQuat(JointSettings.ConnectorTransforms[ConstraintChildIndex]);
+
+	if (ControlComponent->ParentBodyComponentKey.IsValid())
+	{
+		const FTransform ParentTM = Hierarchy->GetGlobalTransform(ControlComponent->ParentBodyComponentKey.ElementKey);
+		const UE::PhysicsControl::FPosQuat ParentTargetTM =
+			UE::PhysicsControl::FPosQuat(ParentTM) *
+			UE::PhysicsControl::FPosQuat(JointSettings.ConnectorTransforms[ConstraintParentIndex]);
+		return ParentTargetTM.Inverse() * ChildTargetTM;
+	}
+	return ChildTargetTM;
+}
+
 //======================================================================================================================
 void FRigPhysicsSimulation::UpdatePrePhysics(
 	const FRigVMExecuteContext&       ExecuteContext, 
@@ -561,246 +956,139 @@ void FRigPhysicsSimulation::UpdatePrePhysics(
 	const float                       DeltaTime, 
 	const FSimulationSpaceData&       SimulationSpaceData)
 {
+	DECLARE_SCOPE_HIERARCHICAL_COUNTER_FUNC()
+	QUICK_SCOPE_CYCLE_COUNTER(STAT_RigPhysics_UpdatePrePhysics);
+
 	URigHierarchy* Hierarchy = OwningControlRig->GetHierarchy();
 
-	for (ImmediatePhysics::FActorHandle* ActorHandle : AdditionalActorHandles)
+	if (CollisionActorHandle)
 	{
 		FTransform BodyRelSimSpaceTM = ConvertCollisionSpaceTransformToSimSpace(
 			SolverComponent->SolverSettings, FTransform());
-		ActorHandle->SetKinematicTarget(BodyRelSimSpaceTM);
+		CollisionActorHandle->SetKinematicTarget(BodyRelSimSpaceTM);
 	}
+
 	for (TPair<FRigComponentKey, FRigBodyRecord>& BodyRecordPair : BodyRecords)
 	{
 		const FRigComponentKey& ComponentKey = BodyRecordPair.Key;
-		const FRigBodyRecord& Record = BodyRecordPair.Value;
+		FRigBodyRecord& Record = BodyRecordPair.Value;
 
 		if (const FRigPhysicsComponent* PhysicsComponent = Cast<FRigPhysicsComponent>(
 			Hierarchy->FindComponent(ComponentKey)))
 		{
 			if (Record.ActorHandle)
 			{
-				UpdateBodyFromModifierData(
-					Record.ActorHandle, Simulation.Get(), PhysicsComponent->BodyData, SimulationSpaceData.Gravity);
+				UpdateBodyPrePhysics(ExecuteContext, SolverComponent, Record, PhysicsComponent, SimulationSpaceData);
+
+				UpdateJointPrePhysics(Record, PhysicsComponent, Hierarchy, DeltaTime);
+			}
+		}
+	}
 
-				if (Record.ActorHandle->GetIsKinematic())
+	for (TPair<FRigComponentKey, FRigControlRecord>& ControlRecordPair : ControlRecords)
+	{
+		const FRigComponentKey& ComponentKey = ControlRecordPair.Key;
+		FRigControlRecord& Record = ControlRecordPair.Value;
+
+		if (const FRigPhysicsControlComponent* ControlComponent = Cast<FRigPhysicsControlComponent>(
+			Hierarchy->FindComponent(ComponentKey)))
+		{
+			if (ImmediatePhysics::FJointHandle * JointHandle = Record.JointHandle)
+			{
+				Chaos::FPBDJointConstraintHandle* Constraint = JointHandle->GetConstraint();
+				if (Constraint)
 				{
-					// Get the target in component space, and then convert it into sim space if necessary.
+					float ThisDeltaTime = DeltaTime;
 
-					// If the target is already in component space, then that's all we need
-					FTransform KinematicTargetCS;
-					switch (PhysicsComponent->KinematicTargetSpace)
+					if (Record.PreviousTargetUpdateCounter + 1 != UpdateCounter)
 					{
-						case ERigPhysicsKinematicTargetSpace::Component:
-						{
-							KinematicTargetCS = PhysicsComponent->KinematicTarget;
-							break;
-						}
-						case ERigPhysicsKinematicTargetSpace::World:
-						{
-							// Record.KinematicTarget * SimulationSpaceState.ComponentTM.Inverse();
-							KinematicTargetCS = PhysicsComponent->KinematicTarget.GetRelativeTransform(
-								SimulationSpaceState.ComponentTM);
-							break;
-						}
-						default:
-						{
-							// All the other options are relative to a bone, so the first task is to get
-							// that, which will be in component space
-							FRigElementKey SourceKey = PhysicsComponent->BodySolverSettings.SourceBone;
-							if (!SourceKey.IsValid())
-							{
-								SourceKey = ComponentKey.ElementKey;
-							}
-							if (SourceKey.IsValid())
-							{
-								FTransform SourceComponentSpaceTM = Hierarchy->GetGlobalTransform(SourceKey);
-
-								switch (PhysicsComponent->KinematicTargetSpace)
-								{
-									case ERigPhysicsKinematicTargetSpace::OffsetInBoneSpace:
-									{
-										KinematicTargetCS = PhysicsComponent->KinematicTarget * SourceComponentSpaceTM;
-										break;
-									}
-									case ERigPhysicsKinematicTargetSpace::OffsetInWorldSpace:
-									{
-										// Convert the bone to WS, apply the target, and convert back
-										FTransform BoneWS = SourceComponentSpaceTM * SimulationSpaceState.ComponentTM;
-										FTransform KinematicTargetWS = BoneWS * PhysicsComponent->KinematicTarget;
-										// Danny TODO figure out which of the GetRelativeTransform versions this is
-										KinematicTargetCS = SimulationSpaceState.ComponentTM.Inverse() * KinematicTargetWS;
-										break;
-									}
-									case ERigPhysicsKinematicTargetSpace::OffsetInComponentSpace:
-									{
-										KinematicTargetCS = SourceComponentSpaceTM * PhysicsComponent->KinematicTarget;
-										break;
-									}
-									default:
-									{
-										UE_CONTROLRIG_RIGUNIT_REPORT_ERROR(TEXT("Kinematic target space is not valid"));
-									}
-								}
-							}
-						}
+						// If we missed some intermediate updates, then we don't want to use the previous
+						// positions etc to calculate velocities. This will mean velocity/damping will be
+						// incorrect for one frame, but that's probably OK.
+						ThisDeltaTime = 0.0f;
 					}
-					FTransform KinematicTargetTM = ConvertComponentSpaceTransformToSimSpace(
-						SolverComponent->SolverSettings, KinematicTargetCS);
-					Record.ActorHandle->SetKinematicTarget(KinematicTargetTM);
-				}
 
-				// Now update the joint targets
-				if (Record.JointHandle)
-				{
-					if (Chaos::FPBDJointConstraintHandle* Constraint = Record.JointHandle->GetConstraint())
+					Constraint->SetCollisionEnabled(!ControlComponent->ControlData.bDisableCollision);
+					Constraint->SetParentInvMassScale(ControlComponent->ControlData.bOnlyControlChildObject ? 0 : 1);
+
+					const Chaos::FPBDJointSettings& JointSettings = Constraint->GetSettings();
+					if (UpdateDriveSpringDamperSettings(
+						JointHandle, JointSettings, ControlComponent->ControlData, ControlComponent->ControlMultiplier))
 					{
-						// Set the drive strength etc
-						const FRigPhysicsDrive& Drive = PhysicsComponent->Joint.Drive;
-						if (Drive.bEnable)
-						{
-							float AngularSpring;
-							float AngularDamping;
-							const float MaxTorque = Drive.MaxTorque;
-
-							float LinearSpring;
-							float LinearDamping;
-							const float MaxForce = Drive.MaxForce;
-
-							UE::PhysicsControl::ConvertStrengthToSpringParams(
-								AngularSpring, AngularDamping,
-								Drive.AngularStrength, Drive.AngularDampingRatio, Drive.AngularExtraDamping);
-							UE::PhysicsControl::ConvertStrengthToSpringParams(
-								LinearSpring, LinearDamping,
-								Drive.LinearStrength, Drive.LinearDampingRatio, Drive.LinearExtraDamping);
-
-							Chaos::FPBDJointSettings Settings = Constraint->GetSettings();
-
-							// Linear
-							Settings.LinearDriveForceMode = Chaos::EJointForceMode::Acceleration;
-							const bool bLinearPositionDriveEnabled = LinearSpring > 0.0f;
-							const bool bLinearVelocityDriveEnabled = LinearDamping > 0.0f;
-							Settings.bLinearPositionDriveEnabled = {
-								bLinearPositionDriveEnabled, bLinearPositionDriveEnabled, bLinearPositionDriveEnabled };
-							Settings.bLinearVelocityDriveEnabled = {
-								bLinearVelocityDriveEnabled, bLinearVelocityDriveEnabled, bLinearVelocityDriveEnabled };
-
-							// Angular
-							Settings.AngularDriveForceMode = Chaos::EJointForceMode::Acceleration;
-							// SLERP is generally better, but doesn't work if anything is locked
-							if (Settings.AngularMotionTypes[0] == Chaos::EJointMotionType::Locked ||
-								Settings.AngularMotionTypes[1] == Chaos::EJointMotionType::Locked ||
-								Settings.AngularMotionTypes[2] == Chaos::EJointMotionType::Locked)
-							{
-								Settings.bAngularSwingPositionDriveEnabled = AngularSpring > 0.0f;
-								Settings.bAngularTwistPositionDriveEnabled = AngularSpring > 0.0f;
-								Settings.bAngularSwingVelocityDriveEnabled = AngularDamping > 0.0f;
-								Settings.bAngularTwistVelocityDriveEnabled = AngularDamping > 0.0f;
-							}
-							else
-							{
-								Settings.bAngularSLerpPositionDriveEnabled = AngularSpring > 0.0f;
-								Settings.bAngularSLerpVelocityDriveEnabled = AngularDamping > 0.0f;
-							}
+						const ImmediatePhysics::FActorHandle* ChildActorHandle = 
+							JointHandle->GetActorHandles()[ConstraintChildIndex];
+						const ImmediatePhysics::FActorHandle* ParentActorHandle = 
+							JointHandle->GetActorHandles()[ConstraintParentIndex];
 
-							Constraint->SetSettings(Settings);
-							Constraint->SetDriveParams(
-								Chaos::FVec3(LinearSpring), Chaos::FVec3(LinearDamping), Chaos::FVec3(MaxForce),
-								Chaos::FVec3(AngularSpring), Chaos::FVec3(AngularDamping), Chaos::FVec3(MaxTorque));
-						}
-						else
+						if (ChildActorHandle && ParentActorHandle)
 						{
-							Constraint->SetDriveParams(
-								Chaos::FVec3(0), Chaos::FVec3(0), Chaos::FVec3(0),
-								Chaos::FVec3(0), Chaos::FVec3(0), Chaos::FVec3(0));
-						}
+							// TODO
+							// - cache settings / previous input parameters to avoid unnecessary repeating
+							//   calculations and making physics API calls every update.
 
-						// Now set the actual target
-						const Chaos::FPBDJointSettings& Settings = Constraint->GetSettings();
+							// Update the target point on the child
+							Constraint->SetChildConnectorLocation(
+								ControlComponent->ControlData.GetControlPoint(ChildActorHandle));
 
-						if (Settings.AngularDriveStiffness.SquaredLength() > 0.0f ||
-							Settings.AngularDriveDamping.SquaredLength() > 0.0f ||
-							Settings.LinearDriveStiffness.SquaredLength() > 0.0f ||
-							Settings.LinearDriveDamping.SquaredLength() > 0.0f)
-						{
-							// Multiplier on the velocity calculated from the current and previous target
-							float TargetVelocityMultiplier = 1.0f;
+							UE::PhysicsControl::FPosQuat TargetTM(
+								ControlComponent->ControlTarget.TargetOrientation,
+								ControlComponent->ControlTarget.TargetPosition);
 
-							FRigElementKey ChildSourceKey;
-							if (const FRigPhysicsComponent* ChildPhysicsComponent = Cast<FRigPhysicsComponent>(
-								Hierarchy->FindComponent(ComponentKey)))
+							if (ControlComponent->ControlData.bUseSkeletalAnimation)
 							{
-								ChildSourceKey = ChildPhysicsComponent->BodySolverSettings.SourceBone;
-								if (!ChildSourceKey.IsValid())
-								{
-									ChildSourceKey = ComponentKey.ElementKey;
-								}
-								TargetVelocityMultiplier = 
-									ChildPhysicsComponent->Joint.Drive.SkeletalAnimationVelocityMultiplier;
+								UE::PhysicsControl::FPosQuat AnimTargetTM = CalculateTargetTM(
+									Hierarchy, JointSettings, ControlComponent);
+								TargetTM = AnimTargetTM * TargetTM;
 							}
 
-							FRigElementKey ParentSourceKey;
-							if (const FRigPhysicsComponent* ParentPhysicsComponent = Cast<FRigPhysicsComponent>(
-								Hierarchy->FindComponent(Record.ParentBodyComponentKey)))
+							Constraint->SetLinearDrivePositionTarget(TargetTM.GetTranslation());
+							Constraint->SetAngularDrivePositionTarget(TargetTM.GetRotation());
+
+							if ((ThisDeltaTime * ControlComponent->ControlData.LinearTargetVelocityMultiplier) != 0)
 							{
-								ParentSourceKey = ParentPhysicsComponent->BodySolverSettings.SourceBone;
-								if (!ParentSourceKey.IsValid())
-								{
-									ParentSourceKey = Record.ParentBodyComponentKey.ElementKey;
-								}
+								FVector Velocity =
+									(TargetTM.GetTranslation() - Record.PreviousTargetTM.GetTranslation()) / ThisDeltaTime;
+								Constraint->SetLinearDriveVelocityTarget(
+									Velocity * ControlComponent->ControlData.LinearTargetVelocityMultiplier);
 							}
-
-							if (!ChildSourceKey.IsValid() || !ParentSourceKey.IsValid())
+							else
 							{
-								continue;
+								Constraint->SetLinearDriveVelocityTarget(Chaos::FVec3(0));
 							}
 
-							// Note that the drive operates between a parent and child part, so we
-							// don't need to worry about global/component (etc) space.
-							const FTransform ChildTM = Hierarchy->GetGlobalTransform(ChildSourceKey);
-							const FTransform ParentTM = Hierarchy->GetGlobalTransform(ParentSourceKey);
-
-							const FTransform ComponentSpaceParentFrameTM = Settings.ConnectorTransforms[1] * ParentTM;
-							const FTransform ComponentSpaceChildFrameTM = Settings.ConnectorTransforms[0] * ChildTM;
-
-							const FTransform DriveTargetTM = 
-								ComponentSpaceChildFrameTM * ComponentSpaceParentFrameTM.Inverse();
-
-							Constraint->SetLinearDrivePositionTarget(DriveTargetTM.GetTranslation());
-							Constraint->SetAngularDrivePositionTarget(DriveTargetTM.GetRotation());
-
-							if (Record.PreviousDriveTargetUpdateCounter + 1 == UpdateCounter && 
-								TargetVelocityMultiplier > 0.0f)
+							if ((ThisDeltaTime * ControlComponent->ControlData.AngularTargetVelocityMultiplier) != 0)
 							{
-								if (DeltaTime > SMALL_NUMBER)
-								{
-									const FTransform DriveTargetTMDelta = 
-										Record.PreviousDriveTargetTM.Inverse() * DriveTargetTM;
-									const FVector Velocity = DriveTargetTMDelta.GetTranslation() / DeltaTime;
-									const FVector AngularVelocity = 
-										DriveTargetTMDelta.GetRotation().GetShortestArcWith(
-											FQuat::Identity).ToRotationVector() / DeltaTime;
-									if (!Velocity.ContainsNaN() && !AngularVelocity.ContainsNaN())
-									{
-										Constraint->SetLinearDriveVelocityTarget(
-											Velocity * TargetVelocityMultiplier);
-										Constraint->SetAngularDriveVelocityTarget(
-											AngularVelocity * TargetVelocityMultiplier);
-									}
-								}
+								// Note that quats multiply in the opposite order to TMs, and must be in the same hemisphere.
+								const FQuat Q = TargetTM.GetRotation();
+								FQuat PrevQ = Record.PreviousTargetTM.GetRotation();
+								PrevQ.EnforceShortestArcWith(Q);
+								const FQuat DeltaQ = Q * PrevQ.Inverse();
+								const FVector AngularVelocity = DeltaQ.ToRotationVector() / ThisDeltaTime;
+
+								Constraint->SetAngularDriveVelocityTarget(
+									AngularVelocity * ControlComponent->ControlData.AngularTargetVelocityMultiplier);
 							}
 							else
 							{
-								Constraint->SetLinearDriveVelocityTarget(Chaos::FVec3(0));
 								Constraint->SetAngularDriveVelocityTarget(Chaos::FVec3(0));
 							}
-							BodyRecordPair.Value.PreviousDriveTargetUpdateCounter = UpdateCounter;
-							BodyRecordPair.Value.PreviousDriveTargetTM = DriveTargetTM;
+
+
+							Record.PreviousTargetTM = TargetTM;
+							Record.PreviousTargetUpdateCounter = UpdateCounter;
+						}
+						else
+						{
+							// Note that if we don't have any strength, then we don't calculate the targets.
+							// However, make sure that we don't apply velocities using the wrong calculation
+							// when the strength/damping is increased in the future
 						}
 					}
 				}
 			}
 		}
 	}
+
 }
 
 //======================================================================================================================
@@ -808,6 +1096,9 @@ void FRigPhysicsSimulation::UpdatePrePhysics(
 // otherwise default to the physics element parent.
 void FRigPhysicsSimulation::UpdatePostPhysics(const FRigPhysicsSolverComponent* SolverComponent, float Alpha)
 {
+	DECLARE_SCOPE_HIERARCHICAL_COUNTER_FUNC()
+	QUICK_SCOPE_CYCLE_COUNTER(STAT_RigPhysics_UpdatePostPhysics);
+
 	if (Alpha == 0.0f)
 	{
 		return;
@@ -906,11 +1197,14 @@ void FRigPhysicsSimulation::UpdatePostPhysics(const FRigPhysicsSolverComponent*
 
 //======================================================================================================================
 void FRigPhysicsSimulation::StepSimulation(
-	const FRigVMExecuteContext&     ExecuteContext, 
+	const FRigVMExecuteContext&       ExecuteContext, 
 	const FRigPhysicsSolverComponent* SolverComponent, 
-	const float                     DeltaTimeOverride, 
-	const float                     Alpha)
+	const float                       DeltaTimeOverride, 
+	const float                       Alpha)
 {
+	DECLARE_SCOPE_HIERARCHICAL_COUNTER_FUNC()
+	QUICK_SCOPE_CYCLE_COUNTER(STAT_RigPhysics_StepSimulation);
+
 	float PhysicsDeltaTime = ExecuteContext.GetDeltaTime();
 	if (DeltaTimeOverride > 0.0f)
 	{
diff --git a/Engine/Plugins/Experimental/PhysicsControl/Source/ControlRigPhysics/Private/RigPhysicsSimulation.h b/Engine/Plugins/Experimental/PhysicsControl/Source/ControlRigPhysics/Private/RigPhysicsSimulation.h
--- a/Engine/Plugins/Experimental/PhysicsControl/Source/ControlRigPhysics/Private/RigPhysicsSimulation.h
+++ b/Engine/Plugins/Experimental/PhysicsControl/Source/ControlRigPhysics/Private/RigPhysicsSimulation.h
@@ -5,6 +5,7 @@
 #include "RigPhysicsData.h"
 
 #include "PhysicsControlData.h"
+#include "PhysicsControlPoseData.h"
 
 #include "Rigs/RigHierarchyComponents.h"
 
@@ -21,6 +22,7 @@ struct FRigPhysicsSolverComponent;
 
 struct FRigPhysicsCollision;
 struct FRigPhysicsDynamics;
+struct FRigPhysicsComponent;
 
 #if UE_BUILD_SHIPPING || UE_BUILD_TEST
 DECLARE_LOG_CATEGORY_EXTERN(LogRigPhysics, Log, Warning);
@@ -28,9 +30,9 @@ DECLARE_LOG_CATEGORY_EXTERN(LogRigPhysics, Log, Warning);
 DECLARE_LOG_CATEGORY_EXTERN(LogRigPhysics, Log, All);
 #endif
 
-/**
- * For internal use only, to keep track of objects we have created
- */
+//======================================================================================================================
+// For internal use only, to keep track of objects we have created
+//======================================================================================================================
 USTRUCT()
 struct FRigBodyRecord
 {
@@ -46,12 +48,38 @@ struct FRigBodyRecord
 	// Things that are updated as the simulation progresses
 
 	// The drive works with velocities so we store the previous target transform, and when it was stored.
-	FTransform                      PreviousDriveTargetTM;
+	UE::PhysicsControl::FPosQuat    PreviousDriveTargetTM;
 	// This is stored from the main solver update counter, marking when the previous drive TM was valid.
 	uint64                          PreviousDriveTargetUpdateCounter = 0;
 };
 
+//======================================================================================================================
+// For internal use only, to keep track of controls we have created
+//======================================================================================================================
+USTRUCT()
+struct FRigControlRecord
+{
+	GENERATED_BODY()
+
+	// Things that are set during instantiation
+	ImmediatePhysics::FJointHandle* JointHandle = nullptr;
+
+	// When we have a joint, this is the key for the parent body. Note that the child body component
+	// key is stored as the lookup key for the records, so we don't need to duplicate it here.
+	FRigComponentKey                ParentBodyComponentKey;
+
+	// Things that are updated as the simulation progresses
+
+	// The control works with velocities so we store the previous target transform, and when it was stored.
+	UE::PhysicsControl::FPosQuat    PreviousTargetTM;
+	// This is stored from the main solver update counter, marking when the previous target TM was valid.
+	uint64                          PreviousTargetUpdateCounter = 0;
+};
 
+
+//======================================================================================================================
+// This represents the low level simulation, plus all the objects and controls we make to go in it
+//======================================================================================================================
 USTRUCT()
 struct FRigPhysicsSimulation : public FRigPhysicsSimulationBase
 {
@@ -114,15 +142,27 @@ private:
 		FVector Gravity;
 	};
 
-	// This collects all the bodies associated with the solver and makes records for them
+	// Creates the low level simulation
 	void InitialiseSimulation(const FRigPhysicsSolverComponent* SolverComponent);
 
+	// This collects all the bodies associated with the solver and makes records for them
+	void InitialiseBodyRecords(const FRigPhysicsSolverComponent* SolverComponent);
+
+	// This collects all the bodies associated with the solver and makes records for them
+	void InitialiseControlRecords(const FRigPhysicsSolverComponent* SolverComponent);
+
+	// Creates collision shapes associated with the solver
 	void InstantiateSolverCollision(const FRigPhysicsSolverComponent* SolverComponent);
 
+	// Creates low level physics joints 
 	void InstantiatePhysicsJoints(const FRigPhysicsSolverComponent* SolverComponent);
 
+	// Creates low level bodies
 	void InstantiatePhysicsBodies(const FRigPhysicsSolverComponent* SolverComponent);
 
+	// Creates low level controls
+	void InstantiateControls(const FRigPhysicsSolverComponent* SolverComponent);
+
 	// Returns true if the component is physics, and its solver matches the solver component
 	// (directly, or automatically)
 	bool ShouldComponentBeInSimulation(
@@ -144,8 +184,22 @@ private:
 		const FRigPhysicsSolverComponent* SolverComponent, 
 		const float                       DeltaTime, 
 		const FSimulationSpaceData&       SimulationSpaceData);
+
+	void UpdateJointPrePhysics(
+		FRigBodyRecord&             Record,
+		const FRigPhysicsComponent* PhysicsComponent, 
+		const URigHierarchy*        Hierarchy, 
+		const float                 DeltaTime);
+
 	void UpdatePostPhysics(const FRigPhysicsSolverComponent* SolverComponent, const float Alpha);
 
+	void UpdateBodyPrePhysics(
+		const FRigVMExecuteContext& ExecuteContext,
+		const FRigPhysicsSolverComponent* SolverComponent,
+		const FRigBodyRecord& Record, 
+		const FRigPhysicsComponent* PhysicsComponent,
+		const FSimulationSpaceData& SimulationSpaceData);
+
 	void InitSimulationSpace(const FTransform& ComponentToWorld, const FTransform& BoneToComponent);
 
 	FSimulationSpaceData UpdateSimulationSpaceStateAndCalculateData(
@@ -171,6 +225,9 @@ private:
 	FTransform ConvertCollisionSpaceTransformToSimSpace(
 		const FRigPhysicsSolverSettings& SolverSettings, const FTransform& TM) const;
 
+	ImmediatePhysics::FJointHandle* CreateConstraint(
+		ImmediatePhysics::FActorHandle* ChildActorHandle, ImmediatePhysics::FActorHandle* ParentActorHandle);
+
 private:
 	TObjectPtr<UControlRig> OwningControlRig;
 	FRigComponentKey        PhysicsSolverComponentKey;
@@ -181,10 +238,18 @@ private:
 	// Ordering so that we can traverse from root to leaf bones
 	TArray<FRigComponentKey> SortedBodyComponentKeys;
 
+	// All the controls
+	TMap<FRigComponentKey, FRigControlRecord> ControlRecords;
+
 	TSharedPtr<ImmediatePhysics::FSimulation> Simulation;
 
-	// Used to store things like collision shapes.
-	TArray<ImmediatePhysics::FActorHandle*> AdditionalActorHandles;
+	// Used to store things like collision shapes. May be offset from the origin if collision is in
+	// a different space to the simulation.
+	ImmediatePhysics::FActorHandle* CollisionActorHandle = nullptr;
+
+	// Used to make controls when they're not attached to another simulated body. Will always be at
+	// the origin.
+	ImmediatePhysics::FActorHandle* SimulationActorHandle = nullptr;
 
 	Chaos::FPBDJointSolverSettings ChaosJointSolverSettings;
 
diff --git a/Engine/Plugins/Experimental/PhysicsControl/Source/ControlRigPhysics/Private/RigPhysicsSimulationVisualization.cpp b/Engine/Plugins/Experimental/PhysicsControl/Source/ControlRigPhysics/Private/RigPhysicsSimulationVisualization.cpp
--- a/Engine/Plugins/Experimental/PhysicsControl/Source/ControlRigPhysics/Private/RigPhysicsSimulationVisualization.cpp
+++ b/Engine/Plugins/Experimental/PhysicsControl/Source/ControlRigPhysics/Private/RigPhysicsSimulationVisualization.cpp
@@ -305,6 +305,9 @@ void FRigPhysicsSimulation::Draw(
 	const FRigPhysicsVisualizationSettings& VisualizationSettings,
 	const UWorld*                           DebugWorld) const
 {
+	DECLARE_SCOPE_HIERARCHICAL_COUNTER_FUNC()
+	QUICK_SCOPE_CYCLE_COUNTER(STAT_RigPhysics_Draw);
+
 #if CHAOS_DEBUG_DRAW
 	// Danny TODO enable/disable chaos debug draw with a cvar or flag
 	if (DebugWorld && DebugWorld->GetPhysicsScene() && DebugWorld->GetPhysicsScene()->GetDebugDrawScene())
@@ -324,9 +327,9 @@ void FRigPhysicsSimulation::Draw(
 	// component space.
 	const FTransform SpaceTransform = ConvertSimSpaceTransformToComponentSpace(SolverSettings, FTransform());
 
-	for (const ImmediatePhysics::FActorHandle* ActorHandle : AdditionalActorHandles)
+	if (CollisionActorHandle)
 	{
-		DrawActor(DI, VisualizationSettings, SpaceTransform, *ActorHandle);
+		DrawActor(DI, VisualizationSettings, SpaceTransform, *CollisionActorHandle);
 	}
 
 	for (const TPair<FRigComponentKey, FRigBodyRecord>& BodyRecordPair : BodyRecords)
diff --git a/Engine/Plugins/Experimental/PhysicsControl/Source/ControlRigPhysics/Public/RigPhysicsControlComponent.h b/Engine/Plugins/Experimental/PhysicsControl/Source/ControlRigPhysics/Public/RigPhysicsControlComponent.h
--- a/Engine/Plugins/Experimental/PhysicsControl/Source/ControlRigPhysics/Public/RigPhysicsControlComponent.h
+++ b/Engine/Plugins/Experimental/PhysicsControl/Source/ControlRigPhysics/Public/RigPhysicsControlComponent.h
@@ -5,6 +5,7 @@
 #include "PhysicsControlData.h"
 
 #include "Rigs/RigPhysics.h"
+
 #include "RigPhysicsControlComponent.generated.h"
 
 /**
@@ -19,13 +20,23 @@ public:
 	GENERATED_BODY()
 	DECLARE_RIG_COMPONENT_METHODS(FRigPhysicsControlComponent)
 
+	// The body that controls the body being controlled. If this is dynamic, it will be affected
+	// too. If unset, then it implies a global control.
 	UPROPERTY(BlueprintReadOnly, EditAnywhere, Category = Control)
 	FRigComponentKey ParentBodyComponentKey;
 
+	// The body being controlled
+	UPROPERTY(BlueprintReadOnly, EditAnywhere, Category = Control)
+	FRigComponentKey ChildBodyComponentKey;
+
 	/** Describes the initial strength etc of the new control */
 	UPROPERTY(BlueprintReadWrite, EditAnywhere, Category = Control)
 	FPhysicsControlData ControlData;
 
+	// This is the currently active control multiplier. 
+	UPROPERTY(BlueprintReadWrite, EditAnywhere, Category = Control)
+	FPhysicsControlMultiplier ControlMultiplier;
+
 	/** Describes the initial target for the new control */
 	UPROPERTY(BlueprintReadWrite, EditAnywhere, Category = Control)
 	FPhysicsControlTarget ControlTarget;
@@ -35,5 +46,11 @@ public:
 
 	virtual bool CanBeAddedTo(
 		const FRigElementKey& InElementKey, const URigHierarchy* InHierarchy, FString* OutFailureReason) const override;
+
+	virtual void OnRigHierarchyKeyChanged(const FRigHierarchyKey& InOldKey, const FRigHierarchyKey& InNewKey) override;
+
+	virtual FName GetDefaultComponentName() const override { return GetDefaultName(); }
+
+	static FName GetDefaultName() { return TEXT("PhysicsControl"); }
 };
 
diff --git a/Engine/Plugins/Experimental/PhysicsControl/Source/ControlRigPhysics/Public/RigPhysicsControlExecution.h b/Engine/Plugins/Experimental/PhysicsControl/Source/ControlRigPhysics/Public/RigPhysicsControlExecution.h
--- a/Engine/Plugins/Experimental/PhysicsControl/Source/ControlRigPhysics/Public/RigPhysicsControlExecution.h
+++ b/Engine/Plugins/Experimental/PhysicsControl/Source/ControlRigPhysics/Public/RigPhysicsControlExecution.h
@@ -2,19 +2,28 @@
 
 #pragma once
 
-#include "RigPhysicsComponent.h"
+#include "RigPhysicsControlComponent.h"
 #include "PhysicsControlData.h"
 
 #include "Units/Execution/RigUnit_DynamicHierarchy.h"
 
 #include "RigPhysicsControlExecution.generated.h"
 
+/**
+ * Base struct for all other mutable physics control nodes
+ */
+USTRUCT(meta = (Category = "RigPhysics", NodeColor = "0.8 0.8 0.3", Keywords = "Physics, Control"))
+struct CONTROLRIGPHYSICS_API FRigUnit_PhysicsControlBaseMutable : public FRigUnitMutable
+{
+	GENERATED_BODY()
+};
+
 /**
  * Adds a new physics control as a component on the owner element.
  * Note: This node only runs as part of the construction event.
  */
 USTRUCT(meta = (Category = "RigPhysics", NodeColor = "1.0 0.6 0.3", DisplayName = "Add Physics Control Component", Keywords = "Construction,Create,New,Joint,Control,Physics", Varying))
-struct CONTROLRIGPHYSICS_API FRigUnit_AddPhysicsControl : public FRigUnitMutable
+struct CONTROLRIGPHYSICS_API FRigUnit_AddPhysicsControl : public FRigUnit_PhysicsControlBaseMutable
 {
 	GENERATED_BODY()
 
@@ -30,29 +39,97 @@ struct CONTROLRIGPHYSICS_API FRigUnit_AddPhysicsControl : public FRigUnitMutable
 	UPROPERTY(meta = (Output))
 	FRigComponentKey ControlComponentKey;
 
-	/*
-	 * The solver to relate this new physics element to
-	 */
+	// The optional body that "does" the controlling - though if it is dynamic then it can move too
 	UPROPERTY(meta = (Input))
-	FRigComponentKey PhysicsSolverComponentKey;
+	FRigComponentKey ParentBodyComponentKey;
 
+	// The body that is controlled
 	UPROPERTY(meta = (Input))
-	FRigComponentKey ParentBodyComponentKey;
+	FRigComponentKey ChildBodyComponentKey;
 
 	/** Describes the initial strength etc of the new control */
 	UPROPERTY(meta = (Input))
 	FPhysicsControlData ControlData;
 
+	/** Fine control over the control strengths etc */
+	UPROPERTY(meta = (Input))
+	FPhysicsControlMultiplier ControlMultiplier;
+
 	/** Describes the initial target for the new control */
 	UPROPERTY(meta = (Input))
 	FPhysicsControlTarget ControlTarget;
 
 	/** Which set to include the control in (optional). Note that it automatically gets added to the set "All" */
-	UPROPERTY(meta = (Input))
-	FName Set;
+	//UPROPERTY(meta = (Input))
+	//FName Set;
 
 	/** NamePrefix Name of the control component. */
 	UPROPERTY(meta = (Input))
 	FName Name = TEXT("Physics Control");
 };
 
+// Sets the control data for a physics control
+USTRUCT(meta = (DisplayName = "Set Physics Control Data", Varying))
+struct CONTROLRIGPHYSICS_API FRigUnit_HierarchySetControlData : public FRigUnit_PhysicsControlBaseMutable
+{
+	GENERATED_BODY()
+
+	FRigUnit_HierarchySetControlData()
+	{
+		PhysicsControlComponentKey.ElementKey.Type = ERigElementType::Bone;
+		PhysicsControlComponentKey.Name = FRigPhysicsControlComponent::GetDefaultName();
+	}
+
+	RIGVM_METHOD()
+	virtual void Execute() override;
+
+	UPROPERTY(meta = (Input))
+	FRigComponentKey PhysicsControlComponentKey;
+
+	UPROPERTY(meta = (Input))
+	FPhysicsControlData ControlData;
+};
+
+// Sets the target for a physics control
+USTRUCT(meta = (DisplayName = "Set Physics Control Target", Varying))
+struct CONTROLRIGPHYSICS_API FRigUnit_HierarchySetControlTarget : public FRigUnit_PhysicsControlBaseMutable
+{
+	GENERATED_BODY()
+
+	FRigUnit_HierarchySetControlTarget()
+	{
+		PhysicsControlComponentKey.ElementKey.Type = ERigElementType::Bone;
+		PhysicsControlComponentKey.Name = FRigPhysicsControlComponent::GetDefaultName();
+	}
+
+	RIGVM_METHOD()
+	virtual void Execute() override;
+
+	UPROPERTY(meta = (Input))
+	FRigComponentKey PhysicsControlComponentKey;
+
+	UPROPERTY(meta = (Input))
+	FPhysicsControlTarget ControlTarget;
+};
+
+// Sets the multipliers for a physics control
+USTRUCT(meta = (DisplayName = "Set Physics Control Multiplier", Varying))
+struct CONTROLRIGPHYSICS_API FRigUnit_HierarchySetControlMultiplier : public FRigUnit_PhysicsControlBaseMutable
+{
+	GENERATED_BODY()
+
+	FRigUnit_HierarchySetControlMultiplier()
+	{
+		PhysicsControlComponentKey.ElementKey.Type = ERigElementType::Bone;
+		PhysicsControlComponentKey.Name = FRigPhysicsControlComponent::GetDefaultName();
+	}
+
+	RIGVM_METHOD()
+	virtual void Execute() override;
+
+	UPROPERTY(meta = (Input))
+	FRigComponentKey PhysicsControlComponentKey;
+
+	UPROPERTY(meta = (Input))
+	FPhysicsControlMultiplier ControlMultiplier;
+};
diff --git a/Engine/Plugins/Experimental/PhysicsControl/Source/PhysicsControl/Private/AnimNode_RigidBodyWithControl.cpp b/Engine/Plugins/Experimental/PhysicsControl/Source/PhysicsControl/Private/AnimNode_RigidBodyWithControl.cpp
--- a/Engine/Plugins/Experimental/PhysicsControl/Source/PhysicsControl/Private/AnimNode_RigidBodyWithControl.cpp
+++ b/Engine/Plugins/Experimental/PhysicsControl/Source/PhysicsControl/Private/AnimNode_RigidBodyWithControl.cpp
@@ -165,7 +165,7 @@ FAnimNode_RigidBodyWithControl::FAnimNode_RigidBodyWithControl()
 	, AnimPhysicsMinDeltaTime(0.0f)
 	, bSimulateAnimPhysicsAfterReset(false)
 	, SkelMeshCompWeakPtr()
-	, PhysicsSimulation(nullptr)
+	, Simulation(nullptr)
 	, SolverSettings()
 	, SolverIterations()
 	, SimulationTask()
@@ -421,7 +421,7 @@ void FAnimNode_RigidBodyWithControl::RunPhysicsSimulation(float DeltaSeconds, co
 	const int32 MaxSteps = RBANWithControl_MaxSubSteps;
 	const float MaxDeltaSeconds = 1.f / 30.f;
 
-	PhysicsSimulation->Simulate_AssumesLocked(DeltaSeconds, MaxDeltaSeconds, MaxSteps, SimSpaceGravity);
+	Simulation->Simulate_AssumesLocked(DeltaSeconds, MaxDeltaSeconds, MaxSteps, SimSpaceGravity);
 	bSimulationStarted = true;
 }
 
@@ -439,8 +439,8 @@ void FAnimNode_RigidBodyWithControl::DestroyPhysicsSimulation()
 {
 	ClothColliders.Reset();
 	FlushDeferredSimulationTask();
-	delete PhysicsSimulation;
-	PhysicsSimulation = nullptr;
+	delete Simulation;
+	Simulation = nullptr;
 	CurrentConstraintProfile = FName();
 }
 
@@ -487,7 +487,7 @@ void FAnimNode_RigidBodyWithControl::EvaluateSkeletalControl_AnyThread(FComponen
 	const float DeltaSeconds = AccumulatedDeltaTime;
 	AccumulatedDeltaTime = 0.f;
 
-	if (bEnabled && PhysicsSimulation)	
+	if (bEnabled && Simulation)	
 	{
 		FlushDeferredSimulationTask();
 
@@ -804,14 +804,14 @@ void FAnimNode_RigidBodyWithControl::EvaluateSkeletalControl_AnyThread(FComponen
 			UpdateWorldObjects(SpaceTransform, SimulationSpaceLinearVelocity, SimulationSpaceAngularVelocity, DeltaSeconds);
 			UpdateClothColliderObjects(SpaceTransform);
 
-			PhysicsSimulation->UpdateSimulationSpace(
+			Simulation->UpdateSimulationSpace(
 				SpaceTransform, SimulationSpaceLinearVelocity, SimulationSpaceAngularVelocity,
 				SimulationSpaceLinearAcceleration, SimulationSpaceAngularAcceleration);
 
-			PhysicsSimulation->SetSimulationSpaceSettings(
+			Simulation->SetSimulationSpaceSettings(
 				(SimSpaceSettingsToUse->WorldAlpha > 0.0f), SimSpaceSettingsToUse->DampingAlpha, SimSpaceSettingsToUse->ExternalLinearDragV);
 
-			PhysicsSimulation->SetSolverSettings(
+			Simulation->SetSolverSettings(
 				SolverSettings.FixedTimeStep,
 				SolverSettings.CullDistance,
 				SolverSettings.MaxDepenetrationVelocity,
@@ -829,7 +829,7 @@ void FAnimNode_RigidBodyWithControl::EvaluateSkeletalControl_AnyThread(FComponen
 			// Draw here even if the simulation is deferred since we want the shapes drawn relative to the current transform
 			if (bRBANWithControl_DebugDraw)
 			{
-				PhysicsSimulation->DebugDraw();
+				Simulation->DebugDraw();
 			}
 		}
 		
@@ -993,14 +993,14 @@ void FAnimNode_RigidBodyWithControl::InitPhysics(const UAnimInstance* InAnimInst
 	bEnabled = PhysicsAssetToUse && SkeletalMeshComp->GetAllowRigidBodyAnimNode() && CVarEnableRigidBodyNodeWithControl.GetValueOnAnyThread() != 0;
 	if(bEnabled)
 	{
-		PhysicsSimulation = new ImmediatePhysics::FSimulation();
+		Simulation = new ImmediatePhysics::FSimulation();
 
 #if WITH_CHAOS_VISUAL_DEBUGGER
-		PhysicsSimulation->GetChaosVDContextData().Id = FChaosVDRuntimeModule::Get().GenerateUniqueID();
-		PhysicsSimulation->GetChaosVDContextData().Type = static_cast<int32>(EChaosVDContextType::Solver);
+		Simulation->GetChaosVDContextData().Id = FChaosVDRuntimeModule::Get().GenerateUniqueID();
+		Simulation->GetChaosVDContextData().Type = static_cast<int32>(EChaosVDContextType::Solver);
 #endif
 
-		PhysicsSimulation->SetRewindVelocities(true);
+		Simulation->SetRewindVelocities(true);
 
 		const int32 NumBodies = PhysicsAssetToUse->SkeletalBodySetups.Num();
 		BodyNameToIndexMap.Reset();
@@ -1080,7 +1080,7 @@ void FAnimNode_RigidBodyWithControl::InitPhysics(const UAnimInstance* InAnimInst
 				// as kinematic prevents them being subsequently made dynamic.
 				bool bSimulated = (BodySetup->PhysicsType == EPhysicsType::PhysType_Simulated);
 				ImmediatePhysics::EActorType ActorType = ImmediatePhysics::EActorType::DynamicActor;
-				ImmediatePhysics::FActorHandle* ActorHandle = PhysicsSimulation->CreateActor(
+				ImmediatePhysics::FActorHandle* ActorHandle = Simulation->CreateActor(
 					ActorType, BodyInstance, BodyInstance->GetUnrealWorldTransform());
 				if (ActorHandle)
 				{
@@ -1095,7 +1095,7 @@ void FAnimNode_RigidBodyWithControl::InitPhysics(const UAnimInstance* InAnimInst
 						// before making it kinematic, otherwise we end up with particles that are
 						// disabled, but still simulate.
 						ActorHandle->SetEnabled(true);
-						PhysicsSimulation->SetIsKinematic(ActorHandle, true);
+						Simulation->SetIsKinematic(ActorHandle, true);
 					}
 
 					ActorHandle->SetName(BodySetup->BoneName);
@@ -1163,7 +1163,7 @@ void FAnimNode_RigidBodyWithControl::InitPhysics(const UAnimInstance* InAnimInst
 					if (bMakeKinematicConstraints || (Body1Handle->IsSimulated() || Body2Handle->IsSimulated()))
 					{
 						ImmediatePhysics::FJointHandle* JointHandle =
-							PhysicsSimulation->CreateJoint(CI, Body1Handle, Body2Handle);
+							Simulation->CreateJoint(CI, Body1Handle, Body2Handle);
 
 						// Record the joint handle under the child bone (i.e. more leaf-ward bone),
 						// since each bone may have multiple children, but only one parent.
@@ -1243,13 +1243,13 @@ void FAnimNode_RigidBodyWithControl::InitPhysics(const UAnimInstance* InAnimInst
 			}
 		}
 
-		PhysicsSimulation->SetIgnoreCollisionPairTable(IgnorePairs);
-		PhysicsSimulation->SetIgnoreCollisionActors(IgnoreCollisionActors);
+		Simulation->SetIgnoreCollisionPairTable(IgnorePairs);
+		Simulation->SetIgnoreCollisionActors(IgnoreCollisionActors);
 
 		CollectClothColliderObjects(SkeletalMeshComp);
 
 		SolverSettings = PhysicsAssetToUse->SolverSettings;
-		PhysicsSimulation->SetSolverSettings(
+		Simulation->SetSolverSettings(
 			SolverSettings.FixedTimeStep,
 			SolverSettings.CullDistance,
 			SolverSettings.MaxDepenetrationVelocity,
@@ -1475,7 +1475,7 @@ void FAnimNode_RigidBodyWithControl::PreUpdate(const UAnimInstance* InAnimInstan
 			// Store game time for use in parallel evaluation. This may be the totol time (inc pauses) or the time the game has been unpaused.
 			WorldTimeSeconds = (float) (SKC->PrimaryComponentTick.bTickEvenWhenPaused ? World->UnpausedTimeSeconds : World->TimeSeconds);
 
-			if (PhysicsSimulation && bEnableWorldGeometry)
+			if (Simulation && bEnableWorldGeometry)
 			{ 
 				// @todo: this logic can be simplified now. We used to run PurgeExpiredWorldObjects and CollectWorldObjects
 				// in UpdateAnimation, but we can't access the world actor's geometry there
@@ -1554,7 +1554,7 @@ void FAnimNode_RigidBodyWithControl::UpdateInternal(const FAnimationUpdateContex
 
 void FAnimNode_RigidBodyWithControl::CollectClothColliderObjects(const USkeletalMeshComponent* SkeletalMeshComp)
 {
-	if (bUseExternalClothCollision && bRBANWithControl_IncludeClothColliders && SkeletalMeshComp && PhysicsSimulation)
+	if (bUseExternalClothCollision && bRBANWithControl_IncludeClothColliders && SkeletalMeshComp && Simulation)
 	{
 		const TArray<FClothCollisionSource>& SkeletalMeshClothCollisionSources = SkeletalMeshComp->GetClothCollisionSources();
 		
@@ -1572,8 +1572,8 @@ void FAnimNode_RigidBodyWithControl::CollectClothColliderObjects(const USkeletal
 				{
 					FBodyInstance* const BodyInstance = BodyInstances[BodyInstanceIndex];
 
-					ImmediatePhysics::FActorHandle* const ActorHandle = PhysicsSimulation->CreateActor(ImmediatePhysics::EActorType::KinematicActor, BodyInstance, BodyInstance->GetUnrealWorldTransform());
-					PhysicsSimulation->AddToCollidingPairs(ActorHandle); // <-allow collision between this actor and all dynamic actors.
+					ImmediatePhysics::FActorHandle* const ActorHandle = Simulation->CreateActor(ImmediatePhysics::EActorType::KinematicActor, BodyInstance, BodyInstance->GetUnrealWorldTransform());
+					Simulation->AddToCollidingPairs(ActorHandle); // <-allow collision between this actor and all dynamic actors.
 					ClothColliders.Add(FClothCollider(ActorHandle, SourceComponent, BodyInstance->InstanceBoneIndex));
 
 					// Terminate the instance.
@@ -1596,7 +1596,7 @@ void FAnimNode_RigidBodyWithControl::RemoveClothColliderObjects()
 {
 	for (const FClothCollider& ClothCollider : ClothColliders)
 	{
-		PhysicsSimulation->DestroyActor(ClothCollider.ActorHandle);
+		Simulation->DestroyActor(ClothCollider.ActorHandle);
 	}
 	
 	ClothColliders.Reset();
@@ -1677,10 +1677,10 @@ void FAnimNode_RigidBodyWithControl::CollectWorldObjects()
 						FTransform TM = OverlapComp->GetComponentTransform();
 						// Create a kinematic actor. Not using Static as world-static objects may
 						// move in the simulation's frame of reference. Note that TM may be in the wrong space here
-						ImmediatePhysics::FActorHandle* ActorHandle = PhysicsSimulation->CreateActor(
+						ImmediatePhysics::FActorHandle* ActorHandle = Simulation->CreateActor(
 							ImmediatePhysics::EActorType::KinematicActor, &OverlapComp->BodyInstance, TM);
 						ActorHandle->SetName(FName(OverlapComp->GetName()));
-						PhysicsSimulation->AddToCollidingPairs(ActorHandle);
+						Simulation->AddToCollidingPairs(ActorHandle);
 						ComponentsInSim.Add(OverlapComp, FWorldObject(ActorHandle, ComponentsInSimTick));
 						// We need this: 
 						// OverlapComp->BodyInstance.bUpdateKinematicFromSimulation = true;
@@ -1724,7 +1724,7 @@ void FAnimNode_RigidBodyWithControl::PurgeExpiredWorldObjects()
 
 		if (WorldObject.bExpired)
 		{
-			PhysicsSimulation->DestroyActor(WorldObject.ActorHandle);
+			Simulation->DestroyActor(WorldObject.ActorHandle);
 			WorldObject.ActorHandle = nullptr;
 
 			PurgedEntries.Add(WorldEntry.Key);
@@ -1947,9 +1947,9 @@ void FAnimNode_RigidBodyWithControl::InitializeBoneReferences(const FBoneContain
 		// We'll have to initialize their transform.
 		bCheckForBodyTransformInit = true;
 
-		if (PhysicsSimulation)
+		if (Simulation)
 		{
-			PhysicsSimulation->SetNumActiveBodies(NumSimulatedBodies, SimulatedBodyIndices);
+			Simulation->SetNumActiveBodies(NumSimulatedBodies, SimulatedBodyIndices);
 		}
 
 		// We're switching to a new LOD, this invalidates our captured poses.
diff --git a/Engine/Plugins/Experimental/PhysicsControl/Source/PhysicsControl/Private/PhysicsControlData.cpp b/Engine/Plugins/Experimental/PhysicsControl/Source/PhysicsControl/Private/PhysicsControlData.cpp
--- a/Engine/Plugins/Experimental/PhysicsControl/Source/PhysicsControl/Private/PhysicsControlData.cpp
+++ b/Engine/Plugins/Experimental/PhysicsControl/Source/PhysicsControl/Private/PhysicsControlData.cpp
@@ -1,6 +1,7 @@
 // Copyright Epic Games, Inc. All Rights Reserved.
 
 #include "PhysicsControlData.h"
+#include "Physics/ImmediatePhysics/ImmediatePhysicsActorHandle.h"
 
 #define SET_SPARSE_DATA(NAME) NAME = SparseData.bEnable##NAME ? SparseData.NAME : NAME
 
@@ -36,6 +37,15 @@ bool FPhysicsControlAndBodyModifierCreationDatas::operator==(const FPhysicsContr
 	return false;
 }
 
+//======================================================================================================================
+FVector FPhysicsControlData::GetControlPoint(const ImmediatePhysics::FActorHandle* ChildActorHandle) const
+{
+	if (bUseCustomControlPoint)
+	{
+		return CustomControlPoint;
+	}
+	return ChildActorHandle->GetLocalCoMLocation();
+}
 
 //======================================================================================================================
 FPhysicsControlData Interpolate(
diff --git a/Engine/Plugins/Experimental/PhysicsControl/Source/PhysicsControl/Private/RigidBodyWithControl.cpp b/Engine/Plugins/Experimental/PhysicsControl/Source/PhysicsControl/Private/RigidBodyWithControl.cpp
--- a/Engine/Plugins/Experimental/PhysicsControl/Source/PhysicsControl/Private/RigidBodyWithControl.cpp
+++ b/Engine/Plugins/Experimental/PhysicsControl/Source/PhysicsControl/Private/RigidBodyWithControl.cpp
@@ -29,7 +29,7 @@
 
 DECLARE_CYCLE_STAT(TEXT("RigidBodyNodeWithControl_ApplyControls"), STAT_RigidBodyNodeWithControl_ApplyControls, STATGROUP_Anim);
 DECLARE_CYCLE_STAT(TEXT("RigidBodyNodeWithControl_ApplyModifiers"), STAT_RigidBodyNodeWithControl_ApplyModifiers, STATGROUP_Anim);
-DECLARE_CYCLE_STAT(TEXT("RigidBodyNodeWithControl_SetConstraintEnabled"), STAT_RigidBodyNodeWithControl_SetConstraintEnabled, STATGROUP_Anim);
+DECLARE_CYCLE_STAT(TEXT("RigidBodyNodeWithControl_SetPhysicsJointEnabled"), STAT_RigidBodyNodeWithControl_SetPhysicsJointEnabled, STATGROUP_Anim);
 
 DECLARE_CYCLE_STAT(TEXT("RigidBodyNodeWithControl_ApplyKinematicTargets"), STAT_RigidBodyNodeWithControl_ApplyKinematicTargets, STATGROUP_Anim);
 DECLARE_CYCLE_STAT(TEXT("RigidBodyNodeWithControl_ApplyCurrentConstraintProfile"), STAT_RigidBodyNodeWithControl_ApplyCurrentConstraintProfile, STATGROUP_Anim);
@@ -38,9 +38,9 @@
 constexpr int32 ConstraintParentIndex = 1;
 
 //======================================================================================================================
-static void SetConstraintEnabled(ImmediatePhysics::FJointHandle* const JointHandle, const bool bIsEnabled)
+void SetPhysicsJointEnabled(ImmediatePhysics::FJointHandle* const JointHandle, const bool bIsEnabled)
 {
-	SCOPE_CYCLE_COUNTER(STAT_RigidBodyNodeWithControl_SetConstraintEnabled);
+	SCOPE_CYCLE_COUNTER(STAT_RigidBodyNodeWithControl_SetPhysicsJointEnabled);
 
 	if (JointHandle)
 	{
@@ -129,6 +129,7 @@
 		}
 	}
 }
+
 //======================================================================================================================
 FName MapConstraintsBehaviorTypeToString(const MapConstraintsBehaviorType InType)
 {
@@ -139,50 +140,56 @@
 }
 
 //======================================================================================================================
-const FTransform FAnimNode_RigidBodyWithControl::GetBodyTransform(const int32 BodyIndex) const
+void UpdateBodyFromModifierData(
+	ImmediatePhysics::FActorHandle*    ActorHandle,
+	ImmediatePhysics::FSimulation*     PhysicsSimulation,
+	const FPhysicsControlModifierData& ModifierData,
+	const FVector& SimSpaceGravity)
 {
-	if (Bodies.IsValidIndex(BodyIndex))
+	if (ModifierData.MovementType != EPhysicsMovementType::Default)
 	{
-		return Bodies[BodyIndex]->GetWorldTransform();
+		// Note that there's an early out if there's no change needed, so this should be OK.
+		PhysicsSimulation->SetIsKinematic(
+			ActorHandle,
+			ModifierData.MovementType != EPhysicsMovementType::Simulated);
 	}
-	return FTransform::Identity;
+
+	// Note that the actual kinematic targets will be set separately, since they need to be set
+	// for all kinematics whether or not they were under a modifier.
+
+	// Scale gravity
+	float GravityMultiplier = ModifierData.GravityMultiplier;
+	if (GravityMultiplier != 1 && ActorHandle->IsGravityEnabled())
+	{
+		float Mass = (float)ActorHandle->GetMass();
+		FVector AntiGravityForce = SimSpaceGravity * (-Mass * (1.0f - GravityMultiplier));
+		ActorHandle->AddForce(AntiGravityForce);
+	}
+
+	// Set collision
+	PhysicsSimulation->SetHasCollision(
+		ActorHandle, CollisionEnabledHasPhysics(ModifierData.CollisionType));
 }
 
 //======================================================================================================================
-void FAnimNode_RigidBodyWithControl::CreateWorldSpaceControlRootBody(UPhysicsAsset* const PhysicsAsset)
-{
-	// Note that this never moves - it's just defined as being an actor at the root of the world
-	WorldSpaceControlActorHandle = PhysicsSimulation->CreateActor(
-		ImmediatePhysics::EActorType::StaticActor, nullptr, FTransform());
-	if (WorldSpaceControlActorHandle)
-	{
-		WorldSpaceControlActorHandle->SetName(FName("WorldSpaceControlHandle"));
-		PhysicsSimulation->SetHasCollision(WorldSpaceControlActorHandle, false);
-	}
-	else
-	{
-		UE_LOG(LogPhysicsControl, Error, TEXT("Failed to create world space control root actor"));
-	}
-}
-
-//======================================================================================================================
-ImmediatePhysics::FJointHandle* FAnimNode_RigidBodyWithControl::CreateConstraint(
-	ImmediatePhysics::FActorHandle* ChildActorHandle, 
+ImmediatePhysics::FJointHandle* CreatePhysicsJoint(
+	ImmediatePhysics::FSimulation*  Simulation,
+	ImmediatePhysics::FActorHandle* ChildActorHandle,
 	ImmediatePhysics::FActorHandle* ParentActorHandle)
 {
 	ImmediatePhysics::FJointHandle* JointHandle = nullptr;
 
-	if (PhysicsSimulation && ChildActorHandle && ParentActorHandle)
+	if (Simulation && ChildActorHandle && ParentActorHandle)
 	{
 		Chaos::FPBDJointSettings Settings;
 
-		Settings.LinearMotionTypes = { 
+		Settings.LinearMotionTypes = {
 			Chaos::EJointMotionType::Free, Chaos::EJointMotionType::Free, Chaos::EJointMotionType::Free };
-		Settings.AngularMotionTypes = { 
+		Settings.AngularMotionTypes = {
 			Chaos::EJointMotionType::Free, Chaos::EJointMotionType::Free, Chaos::EJointMotionType::Free };
 
 		Settings.bLinearPositionDriveEnabled = { true, true, true };
-		Settings.bLinearVelocityDriveEnabled = { true, true, true }; 
+		Settings.bLinearVelocityDriveEnabled = { true, true, true };
 		Settings.LinearDriveForceMode = Chaos::EJointForceMode::Acceleration;
 
 		Settings.bAngularSLerpPositionDriveEnabled = true;
@@ -207,17 +214,44 @@
 		Settings.ConnectorTransforms[ConstraintChildIndex].SetLocation(ChildCoMPositionOffset);
 
 		Settings.Sanitize();
-		JointHandle = PhysicsSimulation->CreateJoint(Settings, ChildActorHandle, ParentActorHandle);
+		JointHandle = Simulation->CreateJoint(Settings, ChildActorHandle, ParentActorHandle);
 	}
 
 	if (JointHandle)
 	{
-		SetConstraintEnabled(JointHandle, false); // Disable constraints by default.
+		SetPhysicsJointEnabled(JointHandle, false); // Disable constraints by default.
 	}
 
 	return JointHandle;
 }
 
 //======================================================================================================================
+const FTransform FAnimNode_RigidBodyWithControl::GetBodyTransform(const int32 BodyIndex) const
+{
+	if (Bodies.IsValidIndex(BodyIndex))
+	{
+		return Bodies[BodyIndex]->GetWorldTransform();
+	}
+	return FTransform::Identity;
+}
+
+//======================================================================================================================
+void FAnimNode_RigidBodyWithControl::CreateWorldSpaceControlRootBody(UPhysicsAsset* const PhysicsAsset)
+{
+	// Note that this never moves - it's just defined as being an actor at the root of the world
+	WorldSpaceControlActorHandle = Simulation->CreateActor(
+		ImmediatePhysics::EActorType::StaticActor, nullptr, FTransform());
+	if (WorldSpaceControlActorHandle)
+	{
+		WorldSpaceControlActorHandle->SetName(FName("WorldSpaceControlHandle"));
+		Simulation->SetHasCollision(WorldSpaceControlActorHandle, false);
+	}
+	else
+	{
+		UE_LOG(LogPhysicsControl, Error, TEXT("Failed to create world space control root actor"));
+	}
+}
+
+//======================================================================================================================
 int32 FAnimNode_RigidBodyWithControl::AddBody(ImmediatePhysics::FActorHandle* const BodyHandle)
 {
@@ -284,7 +318,7 @@
 		ParentBodyHandle = FindBodyFromBoneName(ParentBoneName);
 	}
 	ImmediatePhysics::FActorHandle* const ChildBodyHandle = FindBodyFromBoneName(ChildBoneName);
-	JointHandle = CreateConstraint(ChildBodyHandle, ParentBodyHandle);
+	JointHandle = CreatePhysicsJoint(Simulation, ChildBodyHandle, ParentBodyHandle);
 
 	if (!JointHandle)
 	{
@@ -441,7 +475,7 @@
 		ImmediatePhysics::FJointHandle* JointHandle = NameRecordPair.Value.JointHandle;
 		if (JointHandle)
 		{
-			PhysicsSimulation->DestroyJoint(JointHandle);
+			Simulation->DestroyJoint(JointHandle);
 		}
 	}
 
@@ -530,12 +564,13 @@
 //======================================================================================================================
 // Adjusts the constraint spring drive settings to reflect the control data
 // Returns true if there is some control, false if the drive has no effect
-static bool UpdateDriveSpringDamperSettings(
-	Chaos::FPBDJointConstraintHandle* Constraint,
+bool UpdateDriveSpringDamperSettings(
+	ImmediatePhysics::FJointHandle*   JointHandle,
 	const Chaos::FPBDJointSettings&   Settings, 
 	const FPhysicsControlData&        Data,
 	const FPhysicsControlMultiplier&  Multiplier)
 {
+	Chaos::FPBDJointConstraintHandle* Constraint = JointHandle->GetConstraint();
 	float AngularSpring;
 	float AngularDamping;
 	const float MaxTorque = Data.MaxTorque * Multiplier.MaxTorqueMultiplier;
@@ -583,7 +618,6 @@
 	bool bHaveAngular = (AngularSpring + AngularDamping) > 0;
 	bool bHaveLinear = (LinearSpring + LinearDamping).GetMax() > 0;
 	return bHaveLinear || bHaveAngular;
-
 }
 
 //======================================================================================================================
@@ -621,7 +655,7 @@
 void FAnimNode_RigidBodyWithControl::ApplyControl(FRigidBodyControlRecord& ControlRecord, float DeltaTime)
 {
 	using namespace ImmediatePhysics;
-	FJointHandle* const JointHandle = ControlRecord.JointHandle;
+	FJointHandle* JointHandle = ControlRecord.JointHandle;
 
 	if (JointHandle)
 	{
@@ -642,7 +676,7 @@
 
 			const Chaos::FPBDJointSettings& JointSettings = Constraint->GetSettings();
 			if (UpdateDriveSpringDamperSettings(
-				Constraint, JointSettings, ControlRecord.ControlData, ControlRecord.ControlMultiplier))
+				JointHandle, JointSettings, ControlRecord.ControlData, ControlRecord.ControlMultiplier))
 			{
 				const FActorHandle* ChildActorHandle = JointHandle->GetActorHandles()[ConstraintChildIndex];
 				const FActorHandle* ParentActorHandle = JointHandle->GetActorHandles()[ConstraintParentIndex];
@@ -721,46 +755,14 @@
 }
 
 //======================================================================================================================
-void UpdateBodyFromModifierData(
-	ImmediatePhysics::FActorHandle*    ActorHandle, 
-	ImmediatePhysics::FSimulation*     PhysicsSimulation,
-	const FPhysicsControlModifierData& ModifierData,
-	const FVector&                     SimSpaceGravity)
-{
-	if (ModifierData.MovementType != EPhysicsMovementType::Default)
-	{
-		// Note that there's an early out if there's no change needed, so this should be OK.
-		PhysicsSimulation->SetIsKinematic(
-			ActorHandle,
-			ModifierData.MovementType != EPhysicsMovementType::Simulated);
-	}
-
-	// Note that the actual kinematic targets will be set separately, since they need to be set
-	// for all kinematics whether or not they were under a modifier.
-
-	// Scale gravity
-	float GravityMultiplier = ModifierData.GravityMultiplier;
-	if (GravityMultiplier != 1 && ActorHandle->IsGravityEnabled())
-	{
-		float Mass = (float) ActorHandle->GetMass();
-		FVector AntiGravityForce = SimSpaceGravity * (-Mass * (1.0f - GravityMultiplier));
-		ActorHandle->AddForce(AntiGravityForce);
-	}
-
-	// Set collision
-	PhysicsSimulation->SetHasCollision(
-		ActorHandle, CollisionEnabledHasPhysics(ModifierData.CollisionType));
-}
-
-//======================================================================================================================
 void FAnimNode_RigidBodyWithControl::ApplyModifier(
 	const FRigidBodyModifierRecord& BodyModifierRecord, 
 	const FVector&                  SimSpaceGravity)
 {
-	if (BodyModifierRecord.ActorHandle && PhysicsSimulation)
+	if (BodyModifierRecord.ActorHandle && Simulation)
 	{
 		UpdateBodyFromModifierData(
-			BodyModifierRecord.ActorHandle, PhysicsSimulation, BodyModifierRecord.ModifierData, SimSpaceGravity);
+			BodyModifierRecord.ActorHandle, Simulation, BodyModifierRecord.ModifierData, SimSpaceGravity);
 	}
 }
 
@@ -896,7 +898,7 @@
 				{
 					ApplyControl(ControlRecord, DeltaTime);
 				}
-				SetConstraintEnabled(ControlRecord.JointHandle, ControlRecord.IsEnabled());
+				SetPhysicsJointEnabled(ControlRecord.JointHandle, ControlRecord.IsEnabled());
 			}
 		}
 
diff --git a/Engine/Plugins/Experimental/PhysicsControl/Source/PhysicsControl/Public/AnimNode_RigidBodyWithControl.h b/Engine/Plugins/Experimental/PhysicsControl/Source/PhysicsControl/Public/AnimNode_RigidBodyWithControl.h
--- a/Engine/Plugins/Experimental/PhysicsControl/Source/PhysicsControl/Public/AnimNode_RigidBodyWithControl.h
+++ b/Engine/Plugins/Experimental/PhysicsControl/Source/PhysicsControl/Public/AnimNode_RigidBodyWithControl.h
@@ -60,6 +60,21 @@ PHYSICSCONTROL_API void UpdateBodyFromModifierData(
 	const FPhysicsControlModifierData& ModifierData,
 	const FVector& SimSpaceGravity);
 
+PHYSICSCONTROL_API bool UpdateDriveSpringDamperSettings(
+	ImmediatePhysics::FJointHandle* JointHandle,
+	const Chaos::FPBDJointSettings& Settings,
+	const FPhysicsControlData& Data,
+	const FPhysicsControlMultiplier& Multiplier);
+
+// Creates a physics joint in the form that's useful for controls.
+PHYSICSCONTROL_API ImmediatePhysics::FJointHandle* CreatePhysicsJoint(
+	ImmediatePhysics::FSimulation*  Simulation,
+	ImmediatePhysics::FActorHandle* ChildActorHandle,
+	ImmediatePhysics::FActorHandle* ParentActorHandle);
+
+PHYSICSCONTROL_API void SetPhysicsJointEnabled(
+	ImmediatePhysics::FJointHandle* const JointHandle, const bool bIsEnabled);
+
 /**
  * Controller that simulates physics based on the physics asset of the skeletal mesh component
  */
@@ -94,7 +109,7 @@ struct PHYSICSCONTROL_API FAnimNode_RigidBodyWithControl : public FAnimNode_Skel
 	virtual void AddImpulseAtLocation(FVector Impulse, FVector Location, FName BoneName = NAME_None);
 
 	// TEMP: Exposed for use in PhAT as a quick way to get drag handles working with Chaos
-	virtual ImmediatePhysics::FSimulation* GetSimulation() { return PhysicsSimulation; }
+	virtual ImmediatePhysics::FSimulation* GetSimulation() { return Simulation; }
 
 	/**
 	 * Set the override physics asset. This will automatically trigger a physics re-init if the 
@@ -571,7 +586,7 @@ private:
 	/** This should only be used for removing the delegate during termination. Do NOT use this for any per frame work */
 	TWeakObjectPtr<USkeletalMeshComponent> SkelMeshCompWeakPtr;
 
-	ImmediatePhysics::FSimulation* PhysicsSimulation;
+	ImmediatePhysics::FSimulation* Simulation;
 	FPhysicsAssetSolverSettings SolverSettings;
 	FSolverIterations SolverIterations;	// to be deprecated
 
diff --git a/Engine/Plugins/Experimental/PhysicsControl/Source/PhysicsControl/Public/PhysicsControlData.h b/Engine/Plugins/Experimental/PhysicsControl/Source/PhysicsControl/Public/PhysicsControlData.h
--- a/Engine/Plugins/Experimental/PhysicsControl/Source/PhysicsControl/Public/PhysicsControlData.h
+++ b/Engine/Plugins/Experimental/PhysicsControl/Source/PhysicsControl/Public/PhysicsControlData.h
@@ -6,6 +6,7 @@
 #include "Components/ActorComponent.h"
 #include "Components/MeshComponent.h"
 #include "EngineDefines.h"
+#include "Physics/ImmediatePhysics/ImmediatePhysicsDeclares.h"
 #include "PhysicsControlData.generated.h"
 
 /**
@@ -692,6 +693,9 @@ struct PHYSICSCONTROL_API FPhysicsControlData
 	/** If true then the control will only affect the child object, not the parent */
 	UPROPERTY(EditAnywhere, BlueprintReadWrite, Category = PhysicsControl)
 	bool bOnlyControlChildObject;
+
+	/** Returns the control point, which may be custom or automatic (centre of mass) */
+	FVector GetControlPoint(const ImmediatePhysics::FActorHandle* ChildActorHandle) const;
 };
 
 /**
@@ -816,6 +820,19 @@ struct PHYSICSCONTROL_API FPhysicsControlMultiplier
 	{
 	}
 
+	friend FArchive& operator<<(FArchive& Ar, FPhysicsControlMultiplier& Data)
+	{
+		Ar << Data.LinearStrengthMultiplier;
+		Ar << Data.LinearDampingRatioMultiplier;
+		Ar << Data.LinearExtraDampingMultiplier;
+		Ar << Data.MaxForceMultiplier;
+		Ar << Data.AngularStrengthMultiplier;
+		Ar << Data.AngularDampingRatioMultiplier;
+		Ar << Data.AngularExtraDampingMultiplier;
+		Ar << Data.MaxTorqueMultiplier;
+		return Ar;
+	}
+
 	bool operator==(const FPhysicsControlMultiplier& other) const = default;
 
 	/** Applies the values that have been flagged as enabled from the sparse data */
diff --git a/Engine/Plugins/Experimental/PhysicsControl/Source/PhysicsControl/Public/PhysicsControlObjectVersion.h b/Engine/Plugins/Experimental/PhysicsControl/Source/PhysicsControl/Public/PhysicsControlObjectVersion.h
--- a/Engine/Plugins/Experimental/PhysicsControl/Source/PhysicsControl/Public/PhysicsControlObjectVersion.h
+++ b/Engine/Plugins/Experimental/PhysicsControl/Source/PhysicsControl/Public/PhysicsControlObjectVersion.h
@@ -18,6 +18,10 @@ struct PHYSICSCONTROL_API FPhysicsControlObjectVersion
 
 		ControlRigAutomaticallyAddPhysicsComponents,
 
+		ControlRigControlAddChildBodyComponentKey,
+
+		ControlRigControlAddControlMultiplier,
+
 		// -----<new versions can be added above this line>-------------------------------------------------
 		VersionPlusOne,
 		LatestVersion = VersionPlusOne - 1
