Support for instantiating from a physics asset. Splits out the PhysicsJoint into its own component

diff --git a/Engine/Plugins/Experimental/PhysicsControl/Source/ControlRigPhysics/Private/RigPhysicsBodyComponent.cpp b/Engine/Plugins/Experimental/PhysicsControl/Source/ControlRigPhysics/Private/RigPhysicsBodyComponent.cpp
new file mode 100644
--- /dev/null
+++ b/Engine/Plugins/Experimental/PhysicsControl/Source/ControlRigPhysics/Private/RigPhysicsBodyComponent.cpp
@@ -0,0 +1,223 @@
+// Copyright Epic Games, Inc. All Rights Reserved.
+
+#include "RigPhysicsBodyComponent.h"
+#include "PhysicsControlObjectVersion.h"
+#include "RigPhysicsSolverComponent.h"
+#include "Rigs/RigHierarchy.h"
+#include "Rigs/RigHierarchyController.h"
+#include "ControlRigPhysicsModule.h"
+#include "BoxTypes.h"
+
+#if WITH_EDITOR
+#include "ControlRigPhysicsEditorStyle.h"
+#endif
+
+//======================================================================================================================
+void FRigPhysicsBodyComponent::Save(FArchive& Ar)
+{
+	Ar.UsingCustomVersion(FPhysicsControlObjectVersion::GUID);
+
+	FRigBaseComponent::Save(Ar);
+	Ar << BodySolverSettings;
+	Ar << Dynamics;
+	Ar << Collision;
+	Ar << BodyData;
+	Ar << KinematicTarget;
+	Ar << KinematicTargetSpace;
+	Ar << NoCollisionBodies;
+}
+
+//======================================================================================================================
+void FRigPhysicsBodyComponent::Load(FArchive& Ar)
+{
+	FRigBaseComponent::Load(Ar);
+	Ar << BodySolverSettings;
+	Ar << Dynamics;
+	Ar << Collision;
+	if (Ar.CustomVer(FPhysicsControlObjectVersion::GUID) < FPhysicsControlObjectVersion::ControlRigSeparateOutJointFromBody)
+	{
+		FRigPhysicsJointData Joint;
+		Ar << Joint;
+	}
+	Ar << BodyData;
+	Ar << KinematicTarget;
+	Ar << KinematicTargetSpace;
+
+	if (Ar.CustomVer(FPhysicsControlObjectVersion::GUID) <
+		FPhysicsControlObjectVersion::ControlRigRemoveCurrentDataFromPhysicsComponent)
+	{
+		FPhysicsControlModifierData CurrentBodyData;
+		Ar << CurrentBodyData;
+		// Previously people needed to set things in the current data for initial properties, as
+		// this would override. So do the override here.
+		BodyData = CurrentBodyData;
+	}
+
+	if (Ar.CustomVer(FPhysicsControlObjectVersion::GUID) >=
+		FPhysicsControlObjectVersion::ControlRigSupportNoCollisionBodies)
+	{
+		Ar << NoCollisionBodies;
+	}
+}
+
+//======================================================================================================================
+#if WITH_EDITOR
+const FSlateIcon& FRigPhysicsBodyComponent::GetIconForUI() const
+{
+	static const FSlateIcon BodyMultipleDefaultIcon = FSlateIcon(
+		FControlRigPhysicsEditorStyle::Get().GetStyleSetName(), "ControlRigPhysics.Component.BodyMultipleDefault");
+	static const FSlateIcon BodyMultipleKinematicIcon = FSlateIcon(
+FControlRigPhysicsEditorStyle::Get().GetStyleSetName(), "ControlRigPhysics.Component.BodyMultipleKinematic");
+	static const FSlateIcon BodyMultipleSimulatedIcon = FSlateIcon(
+		FControlRigPhysicsEditorStyle::Get().GetStyleSetName(), "ControlRigPhysics.Component.BodyMultipleSimulated");
+	static const FSlateIcon BodySingleDefaultIcon = FSlateIcon(
+		FControlRigPhysicsEditorStyle::Get().GetStyleSetName(), "ControlRigPhysics.Component.BodySingleDefault");
+	static const FSlateIcon BodySingleKinematicIcon = FSlateIcon(
+FControlRigPhysicsEditorStyle::Get().GetStyleSetName(), "ControlRigPhysics.Component.BodySingleKinematic");
+	static const FSlateIcon BodySingleSimulatedIcon = FSlateIcon(
+		FControlRigPhysicsEditorStyle::Get().GetStyleSetName(), "ControlRigPhysics.Component.BodySingleSimulated");
+
+	const int32 NumShapes = Collision.NumShapes();;
+	if(NumShapes == 0)
+	{
+		return BodySingleDefaultIcon;
+	}
+	
+	if(BodyData.MovementType == EPhysicsMovementType::Simulated)
+	{
+		return NumShapes <= 1 ? BodySingleSimulatedIcon : BodyMultipleSimulatedIcon;
+	}
+
+	// todo: do we want another set of icons for the static movement type?
+	return NumShapes <= 1 ? BodySingleKinematicIcon : BodyMultipleKinematicIcon;
+}
+#endif
+
+//======================================================================================================================
+bool FRigPhysicsBodyComponent::CanBeAddedTo(
+	const FRigElementKey& InElementKey, const URigHierarchy* InHierarchy, FString* OutFailureReason) const
+{
+	if(InElementKey.Type != ERigElementType::Bone)
+	{
+		if(OutFailureReason)
+		{
+			*OutFailureReason = TEXT("Physics components can only be added to bones.");
+		}
+		return false;
+	}
+	return true;
+}
+
+//======================================================================================================================
+void FRigPhysicsBodyComponent::OnAddedToHierarchy(URigHierarchy* InHierarchy, URigHierarchyController* InController)
+{
+	if (!IsProcedural())
+	{
+		AutoCalculateCollision(InHierarchy);
+	}
+
+	if(!CVarControlRigPhysicsAutomaticallySetSolver.GetValueOnAnyThread())
+	{
+		return;
+	}
+	
+	// add a solver if there isn't any
+	TArray<const FRigPhysicsSolverComponent*> SolverComponents = InHierarchy->GetComponents<FRigPhysicsSolverComponent>();
+	if(SolverComponents.IsEmpty())
+	{
+		BodySolverSettings.PhysicsSolverComponentKey = InController->AddTopLevelComponent(
+			FRigPhysicsSolverComponent::StaticStruct(), NAME_None);
+	}
+	else
+	{
+		BodySolverSettings.PhysicsSolverComponentKey = SolverComponents[0]->GetKey();
+	}
+	// let the system know we changed something
+	InHierarchy->Notify(ERigHierarchyNotification::ComponentContentChanged, this);
+}
+
+//======================================================================================================================
+void FRigPhysicsBodyComponent::AutoCalculateCollision(URigHierarchy* InHierarchy, float MinAspectRatio)
+{
+	// Start clean
+	Collision = FRigPhysicsCollision();
+
+	TArray<FVector> Points{ FVector::ZeroVector };
+	FVector MidPoint = FVector::ZeroVector;
+	TArray<FRigElementKey> ChildKeys = InHierarchy->GetChildren(GetElementKey());
+	for (FRigElementKey ChildKey : ChildKeys)
+	{
+		FVector ChildPosition = InHierarchy->GetLocalTransform(ChildKey, true).GetTranslation();
+		if (ChildPosition.SquaredLength() > UE_SMALL_NUMBER)
+		{
+			Points.Add(ChildPosition);
+			MidPoint += ChildPosition;
+		}
+	}
+	MidPoint /= Points.Num();
+
+	if (Points.Num() == 1)
+	{
+		// If there's only one point, then there are no significant children. Make a shape that
+		// duplicates the relationship with our parent, if there is one
+		FTransform TM = InHierarchy->GetLocalTransform(GetElementKey(), true);
+		if (TM.GetTranslation().SquaredLength() < UE_SMALL_NUMBER)
+		{
+			// We have no children, and are co-located with our parent. Make a single, arbitrarily
+			// sized, shape.
+			Collision.Boxes.Add(FRigPhysicsCollisionBox(TM, FVector(10.0)));
+		}
+		else
+		{
+			TM.SetTranslation(InHierarchy->GetLocalTransform(GetElementKey()).GetTranslation() * 0.5);
+			// We're going to make a box that goes along the X axis
+			TM.SetRotation(FQuat::FindBetweenVectors(FVector(1, 0, 0), TM.GetTranslation()));
+			FVector Extents;
+			Extents.X = TM.GetTranslation().Length() * 2.0;
+			Extents.Y = Extents.X * MinAspectRatio;
+			Extents.Z = Extents.Y;
+			Collision.Boxes.Add(FRigPhysicsCollisionBox(TM, Extents));
+		}
+	}
+	else if (Points.Num() == 2)
+	{
+		// There's just one child - easier to do this by hand than the more complex eigenvector
+		// based calculation, and better than using an OBB which may not be aligned with the two points.
+		FTransform TM(MidPoint);
+		// We're going to make a capsule, and they are defined as extending along the Z axis
+		TM.SetRotation(FQuat::FindBetweenVectors(FVector(0, 0, 1), Points[1]));
+
+		float Length = TM.GetTranslation().Length() * 2.0f;
+		float Radius = Length * MinAspectRatio;
+		Length -= 2.0f * Radius;
+		Length = FMath::Max(0.0f, Length);
+		Collision.Capsules.Add(FRigPhysicsCollisionCapsule(TM, Radius, Length));
+	}
+	else
+	{
+		// We could calculate the ideal orientation of a box by calculating the eigenvectors of the
+		// covariance matrix that represents all the joint positions relative to the centroid.
+		// However, for now just use a box orientated with the current join - it is simpler, and
+		// also avoids generating "messy" orientations.
+		UE::Geometry::FAxisAlignedBox3d Box;
+		Box.Contain(Points);
+
+		FTransform TM(Box.Center());
+		FVector Extents = Box.Extents() * 2.0; // The returned extents are half the extents!
+		double MaxExtent = Extents.GetAbsMax();
+		Extents.X = FMath::Max(Extents.X, MaxExtent * MinAspectRatio);
+		Extents.Y = FMath::Max(Extents.Y, MaxExtent * MinAspectRatio);
+		Extents.Z = FMath::Max(Extents.Z, MaxExtent * MinAspectRatio);
+		Collision.Boxes.Add(FRigPhysicsCollisionBox(TM, Extents));
+	}
+
+	InHierarchy->Notify(ERigHierarchyNotification::ComponentContentChanged, this);
+}
+
+//======================================================================================================================
+void FRigPhysicsBodyComponent::OnRigHierarchyKeyChanged(const FRigHierarchyKey& InOldKey, const FRigHierarchyKey& InNewKey)
+{
+	FRigBaseComponent::OnRigHierarchyKeyChanged(InOldKey, InNewKey);
+
+	BodySolverSettings.OnRigHierarchyKeyChanged(InOldKey, InNewKey);
+}
diff --git a/Engine/Plugins/Experimental/PhysicsControl/Source/ControlRigPhysics/Private/RigPhysicsBodyExecution.cpp b/Engine/Plugins/Experimental/PhysicsControl/Source/ControlRigPhysics/Private/RigPhysicsBodyExecution.cpp
new file mode 100644
--- /dev/null
+++ b/Engine/Plugins/Experimental/PhysicsControl/Source/ControlRigPhysics/Private/RigPhysicsBodyExecution.cpp
@@ -0,0 +1,269 @@
+// Copyright Epic Games, Inc. All Rights Reserved.
+
+#include "RigPhysicsBodyExecution.h"
+
+#include "Rigs/RigHierarchyController.h"
+#include "Units/RigUnitContext.h"
+#include "Units/Execution/RigUnit_PrepareForExecution.h"
+#include "ControlRigPhysicsModule.h"
+
+#include "ControlRig.h"
+#if WITH_EDITOR
+#include "RigVMModel/RigVMController.h"
+#include "ControlRigBlueprint.h"
+#endif
+
+#include UE_INLINE_GENERATED_CPP_BY_NAME(RigPhysicsBodyExecution)
+
+//======================================================================================================================
+FRigUnit_AddPhysicsBody_Execute()
+{
+	if (!ExecuteContext.Hierarchy)
+	{
+		return;
+	}
+	if (ExecuteContext.GetEventName() != FRigUnit_PrepareForExecution::EventName)
+	{
+		UE_CONTROLRIG_RIGUNIT_REPORT_ERROR(TEXT("AddPhysicsBody can only be used during Setup"));
+	}
+
+	if (URigHierarchyController* Controller = ExecuteContext.Hierarchy->GetController())
+	{
+		FRigHierarchyControllerInstructionBracket InstructionBracket(
+			Controller, ExecuteContext.GetInstructionIndex());
+		PhysicsBodyComponentKey = Controller->AddComponent(
+			FRigPhysicsBodyComponent::StaticStruct(), FRigPhysicsBodyComponent::GetDefaultName(), Owner);
+		if (PhysicsBodyComponentKey.IsValid())
+		{
+			if (FRigPhysicsBodyComponent* Component = Cast<FRigPhysicsBodyComponent>(
+				ExecuteContext.Hierarchy->FindComponent(PhysicsBodyComponentKey)))
+			{
+				Component->BodySolverSettings = Solver;
+				Component->Dynamics = Dynamics;
+				Component->BodyData = BodyData;
+				Component->Collision = Collision;
+				if (Collision.IsEmpty())
+				{
+					Component->AutoCalculateCollision(ExecuteContext.Hierarchy);
+				}
+			}
+		}
+	}
+}
+
+//======================================================================================================================
+FRigUnit_HierarchyAutoCalculateCollision_Execute()
+{
+	DECLARE_SCOPE_HIERARCHICAL_COUNTER_RIGUNIT();
+	if (!ExecuteContext.Hierarchy)
+	{
+		return;
+	}
+	if (ExecuteContext.GetEventName() != FRigUnit_PrepareForExecution::EventName)
+	{
+		UE_CONTROLRIG_RIGUNIT_REPORT_ERROR(TEXT("AutoCalculateCollision can only be used during Setup"));
+	}
+
+	if (FRigPhysicsBodyComponent* Component = Cast<FRigPhysicsBodyComponent>(
+		ExecuteContext.Hierarchy->FindComponent(PhysicsBodyComponentKey)))
+	{
+		Component->AutoCalculateCollision(ExecuteContext.Hierarchy, MinAspectRatio);
+	}
+}
+
+//======================================================================================================================
+FRigUnit_HierarchySetDynamics_Execute()
+{
+	DECLARE_SCOPE_HIERARCHICAL_COUNTER_RIGUNIT();
+	if (!ExecuteContext.Hierarchy)
+	{
+		return;
+	}
+	if (FRigPhysicsBodyComponent* Component = Cast<FRigPhysicsBodyComponent>(
+		ExecuteContext.Hierarchy->FindComponent(PhysicsBodyComponentKey)))
+	{
+		Component->Dynamics = Dynamics;
+	}
+}
+
+//======================================================================================================================
+FRigUnit_HierarchySetCollision_Execute()
+{
+	DECLARE_SCOPE_HIERARCHICAL_COUNTER_RIGUNIT();
+	if (!ExecuteContext.Hierarchy)
+	{
+		return;
+	}
+	if (FRigPhysicsBodyComponent* Component = Cast<FRigPhysicsBodyComponent>(
+		ExecuteContext.Hierarchy->FindComponent(PhysicsBodyComponentKey)))
+	{
+		Component->Collision = Collision;
+	}
+}
+
+//======================================================================================================================
+FRigUnit_HierarchyDisableCollisionBetween_Execute()
+{
+	DECLARE_SCOPE_HIERARCHICAL_COUNTER_RIGUNIT();
+	if (!ExecuteContext.Hierarchy)
+	{
+		return;
+	}
+	if (FRigPhysicsBodyComponent* Component1 = Cast<FRigPhysicsBodyComponent>(
+		ExecuteContext.Hierarchy->FindComponent(PhysicsBodyComponentKey1)))
+	{
+		Component1->NoCollisionBodies.Add(PhysicsBodyComponentKey2);
+	}
+}
+
+
+//======================================================================================================================
+FRigUnit_HierarchySetPhysicsBodySourceBone_Execute()
+{
+	DECLARE_SCOPE_HIERARCHICAL_COUNTER_RIGUNIT();
+	if (!ExecuteContext.Hierarchy)
+	{
+		return;
+	}
+	if (FRigPhysicsBodyComponent* Component = Cast<FRigPhysicsBodyComponent>(
+		ExecuteContext.Hierarchy->FindComponent(PhysicsBodyComponentKey)))
+	{
+		Component->BodySolverSettings.SourceBone = SourceBone;
+	}
+}
+
+//======================================================================================================================
+FRigUnit_HierarchySetPhysicsBodyTargetBone_Execute()
+{
+	DECLARE_SCOPE_HIERARCHICAL_COUNTER_RIGUNIT();
+	if (!ExecuteContext.Hierarchy)
+	{
+		return;
+	}
+	if (FRigPhysicsBodyComponent* Component = Cast<FRigPhysicsBodyComponent>(
+		ExecuteContext.Hierarchy->FindComponent(PhysicsBodyComponentKey)))
+	{
+		Component->BodySolverSettings.TargetBone = TargetBone;
+	}
+}
+
+//======================================================================================================================
+FRigUnit_HierarchySetPhysicsBodyKinematicTarget_Execute()
+{
+	DECLARE_SCOPE_HIERARCHICAL_COUNTER_RIGUNIT();
+	if (!ExecuteContext.Hierarchy)
+	{
+		return;
+	}
+	if (FRigPhysicsBodyComponent* Component = Cast<FRigPhysicsBodyComponent>(
+		ExecuteContext.Hierarchy->FindComponent(PhysicsBodyComponentKey)))
+	{
+		Component->KinematicTarget = KinematicTarget;
+		Component->KinematicTargetSpace = KinematicTargetSpace;
+	}
+}
+
+//======================================================================================================================
+FRigUnit_HierarchySetPhysicsBodySparseData_Execute()
+{
+	DECLARE_SCOPE_HIERARCHICAL_COUNTER_RIGUNIT();
+	if (!ExecuteContext.Hierarchy)
+	{
+		return;
+	}
+	if (FRigPhysicsBodyComponent* Component = Cast<FRigPhysicsBodyComponent>(
+		ExecuteContext.Hierarchy->FindComponent(PhysicsBodyComponentKey)))
+	{
+		Component->BodyData.UpdateFromSparseData(Data);
+	}
+}
+
+//======================================================================================================================
+FRigUnit_HierarchySetPhysicsBodyMovementType_Execute()
+{
+	DECLARE_SCOPE_HIERARCHICAL_COUNTER_RIGUNIT();
+	if (!ExecuteContext.Hierarchy)
+	{
+		return;
+	}
+	if (FRigPhysicsBodyComponent* Component = Cast<FRigPhysicsBodyComponent>(
+		ExecuteContext.Hierarchy->FindComponent(PhysicsBodyComponentKey)))
+	{
+		Component->BodyData.MovementType = MovementType;
+	}
+}
+
+//======================================================================================================================
+FRigUnit_HierarchySetPhysicsBodyCollisionType_Execute()
+{
+	DECLARE_SCOPE_HIERARCHICAL_COUNTER_RIGUNIT();
+	if (!ExecuteContext.Hierarchy)
+	{
+		return;
+	}
+	if (FRigPhysicsBodyComponent* Component = Cast<FRigPhysicsBodyComponent>(
+		ExecuteContext.Hierarchy->FindComponent(PhysicsBodyComponentKey)))
+	{
+		Component->BodyData.CollisionType = CollisionType;
+	}
+}
+
+//======================================================================================================================
+FRigUnit_HierarchySetPhysicsBodyGravityMultiplier_Execute()
+{
+	DECLARE_SCOPE_HIERARCHICAL_COUNTER_RIGUNIT();
+	if (!ExecuteContext.Hierarchy)
+	{
+		return;
+	}
+	if (FRigPhysicsBodyComponent* Component = Cast<FRigPhysicsBodyComponent>(
+		ExecuteContext.Hierarchy->FindComponent(PhysicsBodyComponentKey)))
+	{
+		Component->BodyData.GravityMultiplier = GravityMultiplier;
+	}
+}
+
+//======================================================================================================================
+FRigUnit_HierarchySetPhysicsBodyPhysicsBlendWeight_Execute()
+{
+	DECLARE_SCOPE_HIERARCHICAL_COUNTER_RIGUNIT();
+	if (!ExecuteContext.Hierarchy)
+	{
+		return;
+	}
+	if (FRigPhysicsBodyComponent* Component = Cast<FRigPhysicsBodyComponent>(
+		ExecuteContext.Hierarchy->FindComponent(PhysicsBodyComponentKey)))
+	{
+		Component->BodyData.PhysicsBlendWeight = PhysicsBlendWeight;
+	}
+}
+
+//======================================================================================================================
+FRigUnit_HierarchySetPhysicsBodyUseSkeletalAnimation_Execute()
+{
+	DECLARE_SCOPE_HIERARCHICAL_COUNTER_RIGUNIT();
+	if (!ExecuteContext.Hierarchy)
+	{
+		return;
+	}
+	if (FRigPhysicsBodyComponent* Component = Cast<FRigPhysicsBodyComponent>(
+		ExecuteContext.Hierarchy->FindComponent(PhysicsBodyComponentKey)))
+	{
+		Component->BodyData.bUseSkeletalAnimation = bUseSkeletalAnimation;
+	}
+}
+
+//======================================================================================================================
+FRigUnit_HierarchySetPhysicsBodyUpdateKinematicFromSimulation_Execute()
+{
+	DECLARE_SCOPE_HIERARCHICAL_COUNTER_RIGUNIT();
+	if (!ExecuteContext.Hierarchy)
+	{
+		return;
+	}
+	if (FRigPhysicsBodyComponent* Component = Cast<FRigPhysicsBodyComponent>(
+		ExecuteContext.Hierarchy->FindComponent(PhysicsBodyComponentKey)))
+	{
+		Component->BodyData.bUpdateKinematicFromSimulation = bUpdateKinematicFromSimulation;
+	}
+}
diff --git a/Engine/Plugins/Experimental/PhysicsControl/Source/ControlRigPhysics/Private/RigPhysicsComponent.cpp b/Engine/Plugins/Experimental/PhysicsControl/Source/ControlRigPhysics/Private/RigPhysicsComponent.cpp
deleted file mode 100644
diff --git a/Engine/Plugins/Experimental/PhysicsControl/Source/ControlRigPhysics/Private/RigPhysicsControlComponent.cpp b/Engine/Plugins/Experimental/PhysicsControl/Source/ControlRigPhysics/Private/RigPhysicsControlComponent.cpp
--- a/Engine/Plugins/Experimental/PhysicsControl/Source/ControlRigPhysics/Private/RigPhysicsControlComponent.cpp
+++ b/Engine/Plugins/Experimental/PhysicsControl/Source/ControlRigPhysics/Private/RigPhysicsControlComponent.cpp
@@ -10,18 +10,11 @@ void FRigPhysicsControlComponent::Save(FArchive& Ar)
 
 	FRigBaseComponent::Save(Ar);
 	Ar << ParentBodyComponentKey;
-	if (Ar.CustomVer(FPhysicsControlObjectVersion::GUID) >=
-		FPhysicsControlObjectVersion::ControlRigControlAddChildBodyComponentKey)
-	{
-		Ar << ChildBodyComponentKey;
-	}
+	Ar << bUseParentBodyAsDefault;
+	Ar << ChildBodyComponentKey;
 	Ar << ControlData;
 	Ar << ControlTarget;
-	if (Ar.CustomVer(FPhysicsControlObjectVersion::GUID) >=
-		FPhysicsControlObjectVersion::ControlRigControlAddChildBodyComponentKey)
-	{
-		Ar << ControlMultiplier;
-	}
+	Ar << ControlMultiplier;
 }
 
 //======================================================================================================================
@@ -29,6 +22,11 @@ void FRigPhysicsControlComponent::Load(FArchive& Ar)
 {
 	FRigBaseComponent::Load(Ar);
 	Ar << ParentBodyComponentKey;
+	if (Ar.CustomVer(FPhysicsControlObjectVersion::GUID) >=
+		FPhysicsControlObjectVersion::ControlRigSeparateOutJointFromBody)
+	{
+		Ar << bUseParentBodyAsDefault;
+	}
 	if (Ar.CustomVer(FPhysicsControlObjectVersion::GUID) >=
 		FPhysicsControlObjectVersion::ControlRigControlAddChildBodyComponentKey)
 	{
@@ -47,6 +45,14 @@ void FRigPhysicsControlComponent::Load(FArchive& Ar)
 bool FRigPhysicsControlComponent::CanBeAddedTo(
 	const FRigElementKey& InElementKey, const URigHierarchy* InHierarchy, FString* OutFailureReason) const
 {
+	if (InElementKey.Type != ERigElementType::Bone)
+	{
+		if (OutFailureReason)
+		{
+			*OutFailureReason = TEXT("Physics control components can only be added to bones.");
+		}
+		return false;
+	}
 	return true;
 }
 
diff --git a/Engine/Plugins/Experimental/PhysicsControl/Source/ControlRigPhysics/Private/RigPhysicsControlExecution.cpp b/Engine/Plugins/Experimental/PhysicsControl/Source/ControlRigPhysics/Private/RigPhysicsControlExecution.cpp
--- a/Engine/Plugins/Experimental/PhysicsControl/Source/ControlRigPhysics/Private/RigPhysicsControlExecution.cpp
+++ b/Engine/Plugins/Experimental/PhysicsControl/Source/ControlRigPhysics/Private/RigPhysicsControlExecution.cpp
@@ -14,36 +14,35 @@
 //======================================================================================================================
 FRigUnit_AddPhysicsControl_Execute()
 {
-	ControlComponentKey.Reset();
+	if (!ExecuteContext.Hierarchy)
+	{
+		return;
+	}
+	if (ExecuteContext.GetEventName() != FRigUnit_PrepareForExecution::EventName)
+	{
+		UE_CONTROLRIG_RIGUNIT_REPORT_ERROR(TEXT("AddPhysicsControl can only be used during Setup"));
+	}
 
-	if (ExecuteContext.GetEventName() == FRigUnit_PrepareForExecution::EventName)
+	if (URigHierarchyController* Controller = ExecuteContext.Hierarchy->GetController())
 	{
-		if (URigHierarchy* Hierarchy = ExecuteContext.Hierarchy)
+		FRigHierarchyControllerInstructionBracket InstructionBracket(
+			Controller, ExecuteContext.GetInstructionIndex());
+
+		ControlComponentKey = Controller->AddComponent(FRigPhysicsControlComponent::StaticStruct(), 
+			FRigPhysicsControlComponent::GetDefaultName(), Owner);
+		if (ControlComponentKey.IsValid())
 		{
-			if (URigHierarchyController* Controller = Hierarchy->GetController())
+			if (FRigPhysicsControlComponent* Component = Cast<FRigPhysicsControlComponent>(
+				ExecuteContext.Hierarchy->FindComponent(ControlComponentKey)))
 			{
-				FRigHierarchyControllerInstructionBracket InstructionBracket(
-					Controller, ExecuteContext.GetInstructionIndex());
-				ControlComponentKey = Controller->AddComponent(FRigPhysicsControlComponent::StaticStruct(), Name, Owner);
-				if (ControlComponentKey.IsValid())
-				{
-					if (FRigPhysicsControlComponent* Component = Cast<FRigPhysicsControlComponent>(
-						Hierarchy->FindComponent(ControlComponentKey)))
-					{
-						Component->ParentBodyComponentKey = ParentBodyComponentKey;
-						Component->ChildBodyComponentKey = ChildBodyComponentKey;
-						Component->ControlData = ControlData;
-						Component->ControlMultiplier = ControlMultiplier;
-						Component->ControlTarget = ControlTarget;
-					}
-				}
+				Component->ParentBodyComponentKey = ParentBodyComponentKey;
+				Component->ChildBodyComponentKey = ChildBodyComponentKey;
+				Component->ControlData = ControlData;
+				Component->ControlMultiplier = ControlMultiplier;
+				Component->ControlTarget = ControlTarget;
 			}
 		}
 	}
-	else
-	{
-		UE_CONTROLRIG_RIGUNIT_REPORT_ERROR(TEXT("AddPhysicsControl can only be used during Setup"));
-	}
 }
 
 //======================================================================================================================
diff --git a/Engine/Plugins/Experimental/PhysicsControl/Source/ControlRigPhysics/Private/RigPhysicsData.cpp b/Engine/Plugins/Experimental/PhysicsControl/Source/ControlRigPhysics/Private/RigPhysicsData.cpp
--- a/Engine/Plugins/Experimental/PhysicsControl/Source/ControlRigPhysics/Private/RigPhysicsData.cpp
+++ b/Engine/Plugins/Experimental/PhysicsControl/Source/ControlRigPhysics/Private/RigPhysicsData.cpp
@@ -2,25 +2,187 @@
 
 #include "RigPhysicsData.h"
 #include "PhysicsControlObjectVersion.h"
+#include "PhysicsControlHelpers.h"
+#include "Chaos/ChaosConstraintSettings.h"
 
 //======================================================================================================================
-FArchive& operator <<(FArchive& Ar, FRigPhysicsJoint& Data)
+FArchive& ArchiveConstraintBaseParams(FArchive& Ar, FConstraintBaseParams& Data)
 {
-	Ar << Data.bEnable;
-	Ar << Data.ParentBody;
+	Ar << Data.Stiffness;
+	Ar << Data.Damping;
+	Ar << Data.Restitution;
+	Ar << Data.ContactDistance;
+	bool bSoftConstraint = Data.bSoftConstraint;
+	Ar << bSoftConstraint;
+	Data.bSoftConstraint = bSoftConstraint;
+	return Ar;
+}
+
+//======================================================================================================================
+FArchive& operator <<(FArchive& Ar, FLinearConstraint& Data)
+{
+	ArchiveConstraintBaseParams(Ar, Data);
+	Ar << Data.Limit;
+	Ar << Data.XMotion;
+	Ar << Data.YMotion;
+	Ar << Data.ZMotion;
+	return Ar;
+}
+
+//======================================================================================================================
+FArchive& operator <<(FArchive& Ar, FConeConstraint& Data)
+{
+	ArchiveConstraintBaseParams(Ar, Data);
+	Ar << Data.Swing1LimitDegrees;
+	Ar << Data.Swing2LimitDegrees;
+	Ar << Data.Swing1Motion;
+	Ar << Data.Swing2Motion;
+	return Ar;
+}
+
+//======================================================================================================================
+FArchive& operator <<(FArchive& Ar, FTwistConstraint& Data)
+{
+	ArchiveConstraintBaseParams(Ar, Data);
+	Ar << Data.TwistLimitDegrees;
+	Ar << Data.TwistMotion;
+	return Ar;
+}
+
+//======================================================================================================================
+FArchive& operator <<(FArchive& Ar, FConstraintDrive& Data)
+{
+	Ar << Data.Stiffness;
+	Ar << Data.Damping;
+	Ar << Data.MaxForce;
+	bool bEnablePositionDrive = Data.bEnablePositionDrive;
+	bool bEnableVelocityDrive = Data.bEnableVelocityDrive;
+	Ar << bEnablePositionDrive;
+	Ar << bEnableVelocityDrive;
+	Data.bEnablePositionDrive = bEnablePositionDrive;
+	Data.bEnableVelocityDrive = bEnableVelocityDrive;
+	return Ar;
+}
+
+//======================================================================================================================
+FArchive& operator <<(FArchive& Ar, FLinearDriveConstraint& Data)
+{
+	Ar << Data.PositionTarget;
+	Ar << Data.VelocityTarget;
+	Ar << Data.XDrive;
+	Ar << Data.YDrive;
+	Ar << Data.ZDrive;
+	Ar << Data.bAccelerationMode;
+	return Ar;
+}
+
+//======================================================================================================================
+FArchive& operator <<(FArchive& Ar, FAngularDriveConstraint& Data)
+{
+	Ar << Data.TwistDrive;
+	Ar << Data.SwingDrive;
+	Ar << Data.SlerpDrive;
+	Ar << Data.OrientationTarget;
+	Ar << Data.AngularVelocityTarget;
+	Ar << Data.AngularDriveMode;
+	Ar << Data.bAccelerationMode;
+	return Ar;
+}
+
+//======================================================================================================================
+FArchive& operator <<(FArchive& Ar, FRigPhysicsJointData& Data)
+{
+	if (Ar.CustomVer(FPhysicsControlObjectVersion::GUID) < FPhysicsControlObjectVersion::ControlRigSeparateOutJointFromBody)
+	{
+		bool bEnable;
+		FRigComponentKey ParentBody;
+		Ar << bEnable;
+		Ar << ParentBody;
+	}
 	Ar << Data.bAutoCalculateParentOffset;
 	Ar << Data.bAutoCalculateChildOffset;
 	Ar << Data.ExtraParentOffset;
 	Ar << Data.ExtraChildOffset;
-	Ar << Data.LinearLimit;
-	Ar << Data.AngularLimit;
+	if (Ar.CustomVer(FPhysicsControlObjectVersion::GUID) < FPhysicsControlObjectVersion::ControlRigSupportFullConstraintData)
+	{
+		float LinearLimit = 0.0f;
+		FVector AngularLimit = FVector(-1.0);
+		Ar << LinearLimit;
+		Ar << AngularLimit;
+		if (Ar.IsLoading())
+		{
+			Data.LinearConstraint.bSoftConstraint = false;
+			Data.ConeConstraint.bSoftConstraint = false;
+			Data.TwistConstraint.bSoftConstraint = false;
+
+			Data.LinearConstraint.Limit = LinearLimit;
+			Data.TwistConstraint.TwistLimitDegrees = AngularLimit.X;
+			Data.ConeConstraint.Swing1LimitDegrees = AngularLimit.Y;
+			Data.ConeConstraint.Swing2LimitDegrees = AngularLimit.Z;
+
+			Data.LinearConstraint.XMotion = ELinearConstraintMotion::LCM_Limited;
+			Data.LinearConstraint.YMotion = ELinearConstraintMotion::LCM_Limited;
+			Data.LinearConstraint.ZMotion = ELinearConstraintMotion::LCM_Limited;
+
+			if (Data.LinearConstraint.Limit == 0)
+			{
+				Data.LinearConstraint.XMotion = ELinearConstraintMotion::LCM_Locked;
+			}
+			else if (Data.LinearConstraint.Limit < 0)
+			{
+				Data.LinearConstraint.XMotion = ELinearConstraintMotion::LCM_Free;
+			}
+
+			Data.TwistConstraint.TwistMotion = EAngularConstraintMotion::ACM_Limited;
+			Data.ConeConstraint.Swing1Motion = EAngularConstraintMotion::ACM_Limited;
+			Data.ConeConstraint.Swing2Motion = EAngularConstraintMotion::ACM_Limited;
+
+			if (Data.TwistConstraint.TwistLimitDegrees == 0)
+			{
+				Data.TwistConstraint.TwistMotion = EAngularConstraintMotion::ACM_Locked;
+			}
+			else if (Data.TwistConstraint.TwistLimitDegrees < 0)
+			{
+				Data.TwistConstraint.TwistMotion = EAngularConstraintMotion::ACM_Free;
+			}
+
+			if (Data.ConeConstraint.Swing1LimitDegrees == 0)
+			{
+				Data.ConeConstraint.Swing1Motion = EAngularConstraintMotion::ACM_Locked;
+			}
+			else if (Data.ConeConstraint.Swing1LimitDegrees < 0)
+			{
+				Data.ConeConstraint.Swing1Motion = EAngularConstraintMotion::ACM_Free;
+			}
+
+			if (Data.ConeConstraint.Swing2LimitDegrees == 0)
+			{
+				Data.ConeConstraint.Swing2Motion = EAngularConstraintMotion::ACM_Locked;
+			}
+			else if (Data.ConeConstraint.Swing2LimitDegrees < 0)
+			{
+				Data.ConeConstraint.Swing2Motion = EAngularConstraintMotion::ACM_Free;
+			}
+
+		}
+	}
+	else
+	{
+		Ar << Data.LinearConstraint;
+		Ar << Data.ConeConstraint;
+		Ar << Data.TwistConstraint;
+	}
 	Ar << Data.bDisableCollision;
 	Ar << Data.LinearProjectionAmount;
 	Ar << Data.AngularProjectionAmount;
 	Ar << Data.ParentInverseMassScale;
 	if (Ar.CustomVer(FPhysicsControlObjectVersion::GUID) >= FPhysicsControlObjectVersion::ControlRigIncludeDriveInJoint)
 	{
-		Ar << Data.Drive;
+		if (Ar.CustomVer(FPhysicsControlObjectVersion::GUID) < FPhysicsControlObjectVersion::ControlRigSeparateOutJointFromBody)
+		{
+			FRigPhysicsDriveData Drive;
+			Ar << Drive;
+		}
 	}
 	return Ar;
 }
@@ -74,17 +236,99 @@ FArchive& operator <<(FArchive& Ar, FRigPhysicsSolverSettings& Data)
 }
 
 //======================================================================================================================
-FArchive& operator <<(FArchive& Ar, FRigPhysicsDrive& Data)
+FArchive& operator <<(FArchive& Ar, FRigPhysicsDriveData& Data)
 {
-	Ar << Data.bEnable;
-	Ar << Data.LinearStrength;
-	Ar << Data.LinearDampingRatio;
-	Ar << Data.LinearExtraDamping;
-	Ar << Data.MaxForce;
-	Ar << Data.AngularStrength;
-	Ar << Data.AngularDampingRatio;
-	Ar << Data.AngularExtraDamping;
-	Ar << Data.MaxTorque;
+	if (Ar.CustomVer(FPhysicsControlObjectVersion::GUID) >= FPhysicsControlObjectVersion::ControlRigSupportFullDriveConstraintData)
+	{
+		Ar << Data.LinearDriveConstraint;
+		Ar << Data.AngularDriveConstraint;
+	}
+	else
+	{
+		bool bEnable = false;
+		float LinearStrength = 0.0f;
+		float LinearDampingRatio = 1.0f;
+		float LinearExtraDamping = 0.0f;
+		float MaxForce = 0.0;
+		float AngularStrength = 0.0f;
+		float AngularDampingRatio = 1.0f;
+		float AngularExtraDamping = 0.0f;
+		float MaxTorque = 0.0f;
+		
+		Ar << bEnable;
+		Ar << LinearStrength;
+		Ar << LinearDampingRatio;
+		Ar << LinearExtraDamping;
+		Ar << MaxForce;
+		Ar << AngularStrength;
+		Ar << AngularDampingRatio;
+		Ar << AngularExtraDamping;
+		Ar << MaxTorque;
+
+		if (Ar.IsLoading())
+		{
+			// Convert to the constraint drive params
+			float AngularSpring;
+			float AngularDamping;
+
+			float LinearSpring;
+			float LinearDamping;
+
+			UE::PhysicsControl::ConvertStrengthToSpringParams(
+				AngularSpring, AngularDamping,
+				AngularStrength, AngularDampingRatio, AngularExtraDamping);
+			UE::PhysicsControl::ConvertStrengthToSpringParams(
+				LinearSpring, LinearDamping,
+				LinearStrength, LinearDampingRatio, LinearExtraDamping);
+
+			// Unfortunately, the physics engine will apply scalings to these values, so we need to
+			// counter that.
+			LinearSpring /= Chaos::ConstraintSettings::LinearDriveStiffnessScale();
+			LinearDamping /= Chaos::ConstraintSettings::LinearDriveDampingScale();
+			AngularSpring /= Chaos::ConstraintSettings::AngularDriveStiffnessScale();
+			AngularDamping /= Chaos::ConstraintSettings::AngularDriveDampingScale();
+
+			Data.LinearDriveConstraint.XDrive.Stiffness = LinearSpring;
+			Data.LinearDriveConstraint.YDrive.Stiffness = LinearSpring;
+			Data.LinearDriveConstraint.ZDrive.Stiffness = LinearSpring;
+
+			Data.LinearDriveConstraint.XDrive.Damping = LinearDamping;
+			Data.LinearDriveConstraint.YDrive.Damping = LinearDamping;
+			Data.LinearDriveConstraint.ZDrive.Damping = LinearDamping;
+
+			Data.LinearDriveConstraint.XDrive.MaxForce = MaxTorque;
+			Data.LinearDriveConstraint.YDrive.MaxForce = MaxTorque;
+			Data.LinearDriveConstraint.ZDrive.MaxForce = MaxTorque;
+
+			Data.AngularDriveConstraint.SlerpDrive.Stiffness = AngularSpring;
+			Data.AngularDriveConstraint.SwingDrive.Stiffness = AngularSpring;
+			Data.AngularDriveConstraint.TwistDrive.Stiffness = AngularSpring;
+
+			Data.AngularDriveConstraint.SlerpDrive.Damping = AngularDamping;
+			Data.AngularDriveConstraint.SwingDrive.Damping = AngularDamping;
+			Data.AngularDriveConstraint.TwistDrive.Damping = AngularDamping;
+
+			Data.AngularDriveConstraint.SlerpDrive.MaxForce = MaxTorque;
+			Data.AngularDriveConstraint.SwingDrive.MaxForce = MaxTorque;
+			Data.AngularDriveConstraint.TwistDrive.MaxForce = MaxTorque;
+
+			Data.LinearDriveConstraint.XDrive.bEnablePositionDrive = bEnable;
+			Data.LinearDriveConstraint.YDrive.bEnablePositionDrive = bEnable;
+			Data.LinearDriveConstraint.ZDrive.bEnablePositionDrive = bEnable;
+
+			Data.LinearDriveConstraint.XDrive.bEnableVelocityDrive = bEnable;
+			Data.LinearDriveConstraint.YDrive.bEnableVelocityDrive = bEnable;
+			Data.LinearDriveConstraint.ZDrive.bEnableVelocityDrive = bEnable;
+
+			Data.AngularDriveConstraint.SlerpDrive.bEnablePositionDrive = bEnable;
+			Data.AngularDriveConstraint.SwingDrive.bEnablePositionDrive = bEnable;
+			Data.AngularDriveConstraint.TwistDrive.bEnablePositionDrive = bEnable;
+
+			Data.AngularDriveConstraint.SlerpDrive.bEnableVelocityDrive = bEnable;
+			Data.AngularDriveConstraint.SwingDrive.bEnableVelocityDrive = bEnable;
+			Data.AngularDriveConstraint.TwistDrive.bEnableVelocityDrive = bEnable;
+		}
+	}
 	Ar << Data.SkeletalAnimationVelocityMultiplier;
 	return Ar;
 }
diff --git a/Engine/Plugins/Experimental/PhysicsControl/Source/ControlRigPhysics/Private/RigPhysicsExecution.cpp b/Engine/Plugins/Experimental/PhysicsControl/Source/ControlRigPhysics/Private/RigPhysicsExecution.cpp
--- a/Engine/Plugins/Experimental/PhysicsControl/Source/ControlRigPhysics/Private/RigPhysicsExecution.cpp
+++ b/Engine/Plugins/Experimental/PhysicsControl/Source/ControlRigPhysics/Private/RigPhysicsExecution.cpp
@@ -2,6 +2,8 @@
 
 #include "RigPhysicsExecution.h"
 #include "RigPhysicsSimulation.h"
+#include "RigPhysicsBodyComponent.h"
+#include "RigPhysicsJointComponent.h"
 #include "RigPhysicsSolverComponent.h"
 #include "RigPhysicsControlComponent.h"
 
@@ -10,6 +12,10 @@
 #include "Units/Execution/RigUnit_PrepareForExecution.h"
 #include "ControlRigPhysicsModule.h"
 
+#include "PhysicsEngine/PhysicsAsset.h"
+#include "PhysicsEngine/SkeletalBodySetup.h"
+#include "PhysicsEngine/PhysicsConstraintTemplate.h"
+
 #include "ControlRig.h"
 #if WITH_EDITOR
 #include "RigVMModel/RigVMController.h"
@@ -19,423 +25,451 @@
 #include UE_INLINE_GENERATED_CPP_BY_NAME(RigPhysicsExecution)
 
 //======================================================================================================================
-#if WITH_EDITOR
-void FRigUnit_PhysicsBaseMutable::OnUnitNodeCreated(FRigVMUnitNodeCreatedContext& InContext) const
+FRigUnit_AddPhysicsSolver_Execute()
 {
-	FRigUnitMutable::OnUnitNodeCreated(InContext);
-	
-	if(!CVarControlRigPhysicsAutomaticallySetSolver.GetValueOnAnyThread())
-	{
-		return;
-	}
-
-	// if the user has dropped the node manually
-	if(InContext.GetReason() != ERigVMNodeCreatedReason::NodeSpawner)
+	if (!ExecuteContext.Hierarchy)
 	{
 		return;
 	}
 
-	const UControlRigBlueprint* Blueprint = Cast<UControlRigBlueprint>(InContext.GetHost());
-	if(Blueprint == nullptr)
+	if (ExecuteContext.GetEventName() != FRigUnit_PrepareForExecution::EventName)
 	{
-		return;
+		UE_CONTROLRIG_RIGUNIT_REPORT_ERROR(TEXT("AddPhysicsSolver can only be used during Setup"));
 	}
 
-	const URigHierarchy* Hierarchy = Blueprint->GetHierarchy();
-	if(Hierarchy == nullptr)
+	if (URigHierarchyController* Controller = ExecuteContext.Hierarchy->GetController())
 	{
-		return;
+		FRigHierarchyControllerInstructionBracket InstructionBracket(
+			Controller, ExecuteContext.GetInstructionIndex());
+		PhysicsSolverComponentKey = Controller->AddComponent(
+			FRigPhysicsSolverComponent::StaticStruct(), TEXT("PhysicsSolver"), Owner);
+		if (PhysicsSolverComponentKey.IsValid())
+		{
+			if (FRigPhysicsSolverComponent* Component = Cast<FRigPhysicsSolverComponent>(
+				ExecuteContext.Hierarchy->FindComponent(PhysicsSolverComponentKey)))
+			{
+				Component->SolverSettings = SolverSettings;
+				Component->SimulationSpaceSettings = SimulationSpaceSettings;
+			}
+		}
 	}
+}
 
-	URigVMController* Controller = InContext.GetController();
-	if(Controller == nullptr)
+//======================================================================================================================
+FRigUnit_StepPhysicsSolver_Execute()
+{
+	DECLARE_SCOPE_HIERARCHICAL_COUNTER_RIGUNIT();
+	if (!ExecuteContext.Hierarchy)
 	{
 		return;
 	}
-
-	const URigVMGraph* Graph = Controller->GetGraph();
-	if(Graph == nullptr)
+	if (UControlRig* ControlRig = Cast<UControlRig>(ExecuteContext.Hierarchy->GetOuter()))
 	{
-		return;
-	}
+		if (FRigPhysicsSolverComponent* SolverComponent = Cast<FRigPhysicsSolverComponent>(
+			ExecuteContext.Hierarchy->FindComponent(PhysicsSolverComponentKey)))
+		{
+			FRigPhysicsSimulation* PhysicsSimulation = Cast<FRigPhysicsSimulation>(
+				ControlRig->GetPhysicsSimulation(PhysicsSolverComponentKey));
 
-	const URigVMNode* MyNode = Graph->FindNodeByName(InContext.GetNodeName());
-	if(MyNode == nullptr)
-	{
-		return;
-	}
+			// Make the solver and instantiate if necessary
+			if (!PhysicsSimulation)
+			{
+				TSharedPtr<FRigPhysicsSimulation> NewPhysicsSimulation = MakeShared<FRigPhysicsSimulation>(ControlRig);
 
-	const URigVMPin* SolverKeyPin = MyNode->FindPin(GET_MEMBER_NAME_STRING_CHECKED(FRigUnit_AddPhysicsSolver, PhysicsSolverComponentKey));
-	if(SolverKeyPin == nullptr)
-	{
-		return;
-	}
+				if (ControlRig->RegisterPhysicsSimulation(NewPhysicsSimulation, PhysicsSolverComponentKey))
+				{
+					NewPhysicsSimulation->Instantiate(SolverComponent);
+				}
+				PhysicsSimulation = NewPhysicsSimulation.Get();
+			}
 
-	if(SolverKeyPin->GetDirection() != ERigVMPinDirection::Input)
-	{
-		return;
-	}
-	
-	TArray<const FRigPhysicsSolverComponent*> SolverComponents = Hierarchy->GetComponents<FRigPhysicsSolverComponent>();
-	if(!SolverComponents.IsEmpty())
-	{
-		FString ContentAsString;
-		FRigComponentKey::StaticStruct()->ExportText(ContentAsString, &SolverComponents[0]->GetKey(), nullptr, nullptr, PPF_ExternalEditor, nullptr);
-		Controller->SetPinDefaultValue(SolverKeyPin->GetPinPath(), ContentAsString);
+			if (PhysicsSimulation)
+			{
+				PhysicsSimulation->StepSimulation(ExecuteContext, SolverComponent, DeltaTimeOverride, Alpha);
+				if (VisualizationSettings.bEnableVisualization)
+				{
+					PhysicsSimulation->Draw(
+						ExecuteContext.GetDrawInterface(), SolverComponent->SolverSettings,
+						VisualizationSettings, ExecuteContext.GetWorld());
+				}
+			}
+		}
 	}
 }
-#endif
 
 //======================================================================================================================
-FRigUnit_AddPhysicsSolver_Execute()
+FRigUnit_InstantiatePhysics_Execute()
 {
-	PhysicsSolverComponentKey.Reset();
-
-	if (ExecuteContext.GetEventName() == FRigUnit_PrepareForExecution::EventName)
+	DECLARE_SCOPE_HIERARCHICAL_COUNTER_RIGUNIT();
+	if (!ExecuteContext.Hierarchy)
+	{
+		return;
+	}
+	if (UControlRig* ControlRig = Cast<UControlRig>(ExecuteContext.Hierarchy->GetOuter()))
 	{
-		if (URigHierarchy* Hierarchy = ExecuteContext.Hierarchy)
+		if (FRigPhysicsSolverComponent* SolverComponent = Cast<FRigPhysicsSolverComponent>(
+			ExecuteContext.Hierarchy->FindComponent(PhysicsSolverComponentKey)))
 		{
-			if (URigHierarchyController* Controller = Hierarchy->GetController())
+			TSharedPtr<FRigPhysicsSimulation> PhysicsSimulation = MakeShared<FRigPhysicsSimulation>(ControlRig);
+
+			if (ControlRig->RegisterPhysicsSimulation(PhysicsSimulation, PhysicsSolverComponentKey))
 			{
-				FRigHierarchyControllerInstructionBracket InstructionBracket(
-					Controller, ExecuteContext.GetInstructionIndex());
-				PhysicsSolverComponentKey = Controller->AddComponent(
-					FRigPhysicsSolverComponent::StaticStruct(), TEXT("PhysicsSolver"), Owner);
-				if (PhysicsSolverComponentKey.IsValid())
-				{
-					if (FRigPhysicsSolverComponent* Component = Cast<FRigPhysicsSolverComponent>(
-						Hierarchy->FindComponent(PhysicsSolverComponentKey)))
-					{
-						Component->SolverSettings = SolverSettings;
-						Component->SimulationSpaceSettings = SimulationSpaceSettings;
-					}
-				}
+				PhysicsSimulation->Instantiate(SolverComponent);
 			}
 		}
 	}
-	else
-	{
-		UE_CONTROLRIG_RIGUNIT_REPORT_ERROR(TEXT("AddPhysicsSolver can only be used during Setup"));
-	}
 }
 
 //======================================================================================================================
-FRigUnit_StepPhysicsSolver_Execute()
+FRigUnit_AddPhysicsComponents_Execute()
 {
 	DECLARE_SCOPE_HIERARCHICAL_COUNTER_RIGUNIT();
+	if (!ExecuteContext.Hierarchy)
+	{
+		return;
+	}
+	if (ExecuteContext.GetEventName() != FRigUnit_PrepareForExecution::EventName)
+	{
+		UE_CONTROLRIG_RIGUNIT_REPORT_ERROR(TEXT("AddPhysicsComponents can only be used during Setup"));
+	}
 
-	if (URigHierarchy* Hierarchy = ExecuteContext.Hierarchy)
+	if (URigHierarchyController* Controller = ExecuteContext.Hierarchy->GetController())
 	{
-		if (UControlRig* ControlRig = Cast<UControlRig>(Hierarchy->GetOuter()))
+		FRigHierarchyControllerInstructionBracket InstructionBracket(
+			Controller, ExecuteContext.GetInstructionIndex());
+		// Add the body
+		PhysicsBodyComponentKey = Controller->AddComponent(
+			FRigPhysicsBodyComponent::StaticStruct(), FRigPhysicsBodyComponent::GetDefaultName(), Owner);
+		if (PhysicsBodyComponentKey.IsValid())
 		{
-			if (FRigPhysicsSolverComponent* SolverComponent = Cast<FRigPhysicsSolverComponent>(
-				Hierarchy->FindComponent(PhysicsSolverComponentKey)))
+			if (FRigPhysicsBodyComponent* Component = Cast<FRigPhysicsBodyComponent>(
+				ExecuteContext.Hierarchy->FindComponent(PhysicsBodyComponentKey)))
 			{
-				FRigPhysicsSimulation* PhysicsSimulation = Cast<FRigPhysicsSimulation>(
-					ControlRig->GetPhysicsSimulation(PhysicsSolverComponentKey));
-
-				// Make the solver and instantiate if necessary
-				if (!PhysicsSimulation)
+				Component->BodySolverSettings = Solver;
+				Component->Dynamics = Dynamics;
+				Component->BodyData = BodyData;
+				Component->Collision = Collision;
+				if (Collision.IsEmpty())
 				{
-					TSharedPtr<FRigPhysicsSimulation> NewPhysicsSimulation = MakeShared<FRigPhysicsSimulation>(ControlRig);
-
-					if (ControlRig->RegisterPhysicsSimulation(NewPhysicsSimulation, PhysicsSolverComponentKey))
-					{
-						NewPhysicsSimulation->Instantiate(SolverComponent);
-					}
-					PhysicsSimulation = NewPhysicsSimulation.Get();
+					Component->AutoCalculateCollision(ExecuteContext.Hierarchy);
 				}
+			}
 
-				if (PhysicsSimulation)
+			if (bAddJoint)
+			{
+				PhysicsJointComponentKey = Controller->AddComponent(
+					FRigPhysicsJointComponent::StaticStruct(), FRigPhysicsJointComponent::GetDefaultName(), Owner);
+				if (PhysicsJointComponentKey.IsValid())
 				{
-					PhysicsSimulation->StepSimulation(ExecuteContext, SolverComponent, DeltaTimeOverride, Alpha);
-					if (VisualizationSettings.bEnableVisualization)
+					if (FRigPhysicsJointComponent* Component = Cast<FRigPhysicsJointComponent>(
+						ExecuteContext.Hierarchy->FindComponent(PhysicsJointComponentKey)))
 					{
-						PhysicsSimulation->Draw(
-							ExecuteContext.GetDrawInterface(), SolverComponent->SolverSettings, 
-							VisualizationSettings, ExecuteContext.GetWorld());
+						// Leave parent as blank to automatically find it
+						Component->ChildBodyComponentKey = PhysicsBodyComponentKey;
+						Component->JointData = JointData;
+						Component->DriveData = DriveData;
 					}
 				}
 			}
-		}
-	}
-}
-
-//======================================================================================================================
-FRigUnit_InstantiatePhysics_Execute()
-{
-	DECLARE_SCOPE_HIERARCHICAL_COUNTER_RIGUNIT();
 
-	if (URigHierarchy* Hierarchy = ExecuteContext.Hierarchy)
-	{
-		if (UControlRig* ControlRig = Cast<UControlRig>(Hierarchy->GetOuter()))
-		{
-			if (FRigPhysicsSolverComponent* SolverComponent = Cast<FRigPhysicsSolverComponent>(
-				Hierarchy->FindComponent(PhysicsSolverComponentKey)))
+			if (bAddSimSpaceControl)
 			{
-				TSharedPtr<FRigPhysicsSimulation> PhysicsSimulation = MakeShared<FRigPhysicsSimulation>(ControlRig);
-
-				if (ControlRig->RegisterPhysicsSimulation(PhysicsSimulation, PhysicsSolverComponentKey))
+				SimSpaceControlComponentKey = Controller->AddComponent(FRigPhysicsControlComponent::StaticStruct(),
+					TEXT("SimSpaceControl"), Owner);
+				if (SimSpaceControlComponentKey.IsValid())
 				{
-					PhysicsSimulation->Instantiate(SolverComponent);
+					if (FRigPhysicsControlComponent* Component = Cast<FRigPhysicsControlComponent>(
+						ExecuteContext.Hierarchy->FindComponent(SimSpaceControlComponentKey)))
+					{
+						Component->bUseParentBodyAsDefault = false;
+						Component->ChildBodyComponentKey = PhysicsBodyComponentKey;
+						Component->ControlData = SimSpaceControlData;
+					}
 				}
 			}
-		}
-	}
-}
 
-//======================================================================================================================
-FRigUnit_AddPhysicsBody_Execute()
-{
-	PhysicsBodyComponentKey.Reset();
-
-	if (ExecuteContext.GetEventName() == FRigUnit_PrepareForExecution::EventName)
-	{
-		if (URigHierarchy* Hierarchy = ExecuteContext.Hierarchy)
-		{
-			if (URigHierarchyController* Controller = Hierarchy->GetController())
+			if (bAddParentSpaceControl)
 			{
-				FRigHierarchyControllerInstructionBracket InstructionBracket(
-					Controller, ExecuteContext.GetInstructionIndex());
-				PhysicsBodyComponentKey = Controller->AddComponent(
-					FRigPhysicsComponent::StaticStruct(), TEXT("PhysicsBody"), Owner);
-				if (PhysicsBodyComponentKey.IsValid())
+				ParentSpaceControlComponentKey = Controller->AddComponent(FRigPhysicsControlComponent::StaticStruct(),
+					TEXT("ParentSpaceControl"), Owner);
+				if (ParentSpaceControlComponentKey.IsValid())
 				{
-					if (FRigPhysicsComponent* Component = Cast<FRigPhysicsComponent>(
-						Hierarchy->FindComponent(PhysicsBodyComponentKey)))
+					if (FRigPhysicsControlComponent* Component = Cast<FRigPhysicsControlComponent>(
+						ExecuteContext.Hierarchy->FindComponent(ParentSpaceControlComponentKey)))
 					{
-						Component->BodySolverSettings = Solver;
-						Component->Dynamics = Dynamics;
-						Component->Joint = Joint;
-						Component->BodyData = BodyData;
-						Component->Collision = Collision;
+						Component->bUseParentBodyAsDefault = true;
+						Component->ChildBodyComponentKey = PhysicsBodyComponentKey;
+						Component->ControlData = ParentSpaceControlData;
 					}
 				}
 			}
 		}
 	}
-	else
-	{
-		UE_CONTROLRIG_RIGUNIT_REPORT_ERROR(TEXT("AddPhysicsBody can only be used during Setup"));
-	}
 }
 
 //======================================================================================================================
-FRigUnit_HierarchyAutoCalculateCollision_Execute()
+static void SetCommonProperties(FRigPhysicsCollisionShape& Shape, const FKShapeElem& ShapeElem)
 {
-	DECLARE_SCOPE_HIERARCHICAL_COUNTER_RIGUNIT();
-
-	if (ExecuteContext.Hierarchy)
-	{
-		if (FRigPhysicsComponent* Component = Cast<FRigPhysicsComponent>(
-			ExecuteContext.Hierarchy->FindComponent(PhysicsBodyComponentKey)))
-		{
-			Component->AutoCalculateCollision(ExecuteContext.Hierarchy, MinAspectRatio);
-		}
-	}
+	Shape.RestOffset = ShapeElem.RestOffset;
+	Shape.Name = ShapeElem.GetName();
+	Shape.bContributeToMass = ShapeElem.GetContributeToMass();
 }
 
 //======================================================================================================================
-FRigUnit_HierarchySetDynamics_Execute()
+static void SetUpJointData(
+	FRigPhysicsJointData&               JointData, 
+	const FConstraintInstance&          ConstraintInstance, 
+	const FConstraintProfileProperties& ConstraintProfileProperties)
 {
-	DECLARE_SCOPE_HIERARCHICAL_COUNTER_RIGUNIT();
+	JointData.bAutoCalculateChildOffset = false;
+	JointData.bAutoCalculateParentOffset = false;
 
-	if (ExecuteContext.Hierarchy)
-	{
-		if (FRigPhysicsComponent* Component = Cast<FRigPhysicsComponent>(
-			ExecuteContext.Hierarchy->FindComponent(PhysicsBodyComponentKey)))
-		{
-			Component->Dynamics = Dynamics;
-		}
-	}
+	JointData.ExtraChildOffset = ConstraintInstance.GetRefFrame(EConstraintFrame::Frame1);
+	JointData.ExtraParentOffset = ConstraintInstance.GetRefFrame(EConstraintFrame::Frame2);
+
+	JointData.LinearConstraint = ConstraintProfileProperties.LinearLimit;
+	JointData.ConeConstraint = ConstraintProfileProperties.ConeLimit;
+	JointData.TwistConstraint = ConstraintProfileProperties.TwistLimit;
 }
 
 //======================================================================================================================
-FRigUnit_HierarchySetCollision_Execute()
+static void SetUpDriveData(
+	FRigPhysicsDriveData&               DriveData, 
+	const FConstraintProfileProperties& ConstraintProfileProperties)
 {
-	DECLARE_SCOPE_HIERARCHICAL_COUNTER_RIGUNIT();
-
-	if (ExecuteContext.Hierarchy)
-	{
-		if (FRigPhysicsComponent* Component = Cast<FRigPhysicsComponent>(
-			ExecuteContext.Hierarchy->FindComponent(PhysicsBodyComponentKey)))
-		{
-			Component->Collision = Collision;
-		}
-	}
+	DriveData.LinearDriveConstraint = ConstraintProfileProperties.LinearDrive;
+	DriveData.AngularDriveConstraint = ConstraintProfileProperties.AngularDrive;
 }
 
 //======================================================================================================================
-FRigUnit_HierarchySetJoint_Execute()
+FRigUnit_HierarchyInstantiateFromPhysicsAsset_Execute()
 {
 	DECLARE_SCOPE_HIERARCHICAL_COUNTER_RIGUNIT();
-
-	if (ExecuteContext.Hierarchy)
+	if (!ExecuteContext.Hierarchy)
 	{
-		if (FRigPhysicsComponent* Component = Cast<FRigPhysicsComponent>(
-			ExecuteContext.Hierarchy->FindComponent(PhysicsBodyComponentKey)))
-		{
-			Component->Joint= Joint;
-		}
+		return;
+	}
+	if (ExecuteContext.GetEventName() != FRigUnit_PrepareForExecution::EventName)
+	{
+		UE_CONTROLRIG_RIGUNIT_REPORT_ERROR(TEXT("InstantiateFromPhysicsAsset can only be used during Setup"));
+		return;
+	}
+	if (!PhysicsAsset)
+	{
+		UE_CONTROLRIG_RIGUNIT_REPORT_ERROR(TEXT("InstantiateFromPhysicsAsset needs a valid physics asset"));
+		return;
 	}
-}
-
-//======================================================================================================================
-FRigUnit_HierarchySetJointDrive_Execute()
-{
-	DECLARE_SCOPE_HIERARCHICAL_COUNTER_RIGUNIT();
 
-	if (ExecuteContext.Hierarchy)
+	URigHierarchyController * Controller = ExecuteContext.Hierarchy->GetController();
+	if (!Controller)
 	{
-		if (FRigPhysicsComponent* Component = Cast<FRigPhysicsComponent>(
-			ExecuteContext.Hierarchy->FindComponent(PhysicsBodyComponentKey)))
-		{
-			Component->Joint.Drive = Drive;
-		}
+		return;
 	}
-}
 
-//======================================================================================================================
-FRigUnit_HierarchySetPhysicsBodySourceBone_Execute()
-{
-	DECLARE_SCOPE_HIERARCHICAL_COUNTER_RIGUNIT();
+	FRigHierarchyControllerInstructionBracket InstructionBracket(
+		Controller, ExecuteContext.GetInstructionIndex());
 
-	if (ExecuteContext.Hierarchy)
+	// When we create physics joints we need to be able to refer back to the bodies that have been created
+	TMap<FName, FRigComponentKey> BoneToBodyComponentKeyMap;
+
+	for (TObjectPtr<USkeletalBodySetup> SkeletalBodySetup : PhysicsAsset->SkeletalBodySetups)
 	{
-		if (FRigPhysicsComponent* Component = Cast<FRigPhysicsComponent>(
-			ExecuteContext.Hierarchy->FindComponent(PhysicsBodyComponentKey)))
+		if (!SkeletalBodySetup)
 		{
-			Component->BodySolverSettings.SourceBone = SourceBone;
+			continue;
 		}
-	}
-}
 
-//======================================================================================================================
-FRigUnit_HierarchySetPhysicsBodyTargetBone_Execute()
-{
-	DECLARE_SCOPE_HIERARCHICAL_COUNTER_RIGUNIT();
-
-	if (ExecuteContext.Hierarchy)
-	{
-		if (FRigPhysicsComponent* Component = Cast<FRigPhysicsComponent>(
-			ExecuteContext.Hierarchy->FindComponent(PhysicsBodyComponentKey)))
+		// Check we can use this based on the bone name
+		FName BoneName = SkeletalBodySetup->BoneName;
+		FRigElementKey OwnerElementKey(BoneName, ERigElementType::Bone);
+		if (!ExecuteContext.Hierarchy->Find(OwnerElementKey))
 		{
-			Component->BodySolverSettings.TargetBone = TargetBone;
+			continue;
+		}
+		if (!BonesToUse.IsEmpty() && BonesToUse.Find(OwnerElementKey) == INDEX_NONE)
+		{
+			continue;
 		}
-	}
-}
 
-//======================================================================================================================
-FRigUnit_HierarchySetPhysicsBodyKinematicTarget_Execute()
-{
-	DECLARE_SCOPE_HIERARCHICAL_COUNTER_RIGUNIT();
+		// Collision
+		FRigPhysicsCollision Collision;
+		const FKAggregateGeom& AggGeom = SkeletalBodySetup->AggGeom;
+		for (const FKBoxElem& BoxElem : AggGeom.BoxElems)
+		{
+			FTransform TM(BoxElem.Rotation, BoxElem.Center);
+			FRigPhysicsCollisionBox Box(TM, FVector(BoxElem.X, BoxElem.Y, BoxElem.Z));
+			SetCommonProperties(Box, BoxElem);
+			Collision.Boxes.Add(Box);
+		}
+		for (const FKSphereElem& SphereElem : AggGeom.SphereElems)
+		{
+			FTransform TM(SphereElem.Center);
+			FRigPhysicsCollisionSphere Sphere(TM, SphereElem.Radius);
+			SetCommonProperties(Sphere, SphereElem);
+			Collision.Spheres.Add(Sphere);
+		}
+		for (const FKSphylElem& SphylElem : AggGeom.SphylElems)
+		{
+			FTransform TM(SphylElem.Rotation, SphylElem.Center);
+			FRigPhysicsCollisionCapsule Capsule(TM, SphylElem.Radius, SphylElem.Length);
+			SetCommonProperties(Capsule, SphylElem);
+			Collision.Capsules.Add(Capsule);
+		}
 
-	if (ExecuteContext.Hierarchy)
-	{
-		if (FRigPhysicsComponent* Component = Cast<FRigPhysicsComponent>(
+		// Dynamics
+		FRigPhysicsDynamics Dynamics;
+		Dynamics.Mass = SkeletalBodySetup->CalculateMass();
+
+		// Body Data
+		FPhysicsControlModifierData BodyData;
+		BodyData.MovementType =
+			SkeletalBodySetup->PhysicsType == EPhysicsType::PhysType_Simulated ?
+			EPhysicsMovementType::Simulated : EPhysicsMovementType::Kinematic;
+		BodyData.CollisionType =
+			SkeletalBodySetup->CollisionReponse == EBodyCollisionResponse::BodyCollision_Enabled ?
+			ECollisionEnabled::PhysicsOnly : ECollisionEnabled::NoCollision;
+
+		// Add the component
+		FRigComponentKey PhysicsBodyComponentKey = Controller->AddComponent(
+			FRigPhysicsBodyComponent::StaticStruct(), TEXT("PhysicsBody"), OwnerElementKey);
+		if (FRigPhysicsBodyComponent* Component = Cast<FRigPhysicsBodyComponent>(
 			ExecuteContext.Hierarchy->FindComponent(PhysicsBodyComponentKey)))
 		{
-			Component->KinematicTarget = KinematicTarget;
-			Component->KinematicTargetSpace = KinematicTargetSpace;
-		}
-	}
-}
+			Component->BodySolverSettings = Solver;
+			Component->Dynamics = Dynamics;
+			Component->BodyData = BodyData;
+			Component->Collision = Collision;
 
-//======================================================================================================================
-FRigUnit_HierarchySetPhysicsBodySparseData_Execute()
-{
-	DECLARE_SCOPE_HIERARCHICAL_COUNTER_RIGUNIT();
+			BoneToBodyComponentKeyMap.Add(BoneName, PhysicsBodyComponentKey);
+			PhysicsBodyComponentKeys.Add(PhysicsBodyComponentKey);
+		}
 
-	if (ExecuteContext.Hierarchy)
-	{
-		if (FRigPhysicsComponent* Component = Cast<FRigPhysicsComponent>(
-			ExecuteContext.Hierarchy->FindComponent(PhysicsBodyComponentKey)))
+		if (bAddSimSpaceControl)
 		{
-			Component->BodyData.UpdateFromSparseData(Data);
+			FRigComponentKey SimSpaceControlComponentKey = Controller->AddComponent(
+				FRigPhysicsControlComponent::StaticStruct(),
+				TEXT("SimSpaceControl"), PhysicsBodyComponentKey.ElementKey);
+			if (FRigPhysicsControlComponent* Component = Cast<FRigPhysicsControlComponent>(
+				ExecuteContext.Hierarchy->FindComponent(SimSpaceControlComponentKey)))
+			{
+				Component->bUseParentBodyAsDefault = false;
+				Component->ChildBodyComponentKey = PhysicsBodyComponentKey;
+				Component->ControlData = SimSpaceControlData;
+
+				SimSpaceControlComponentKeys.Add(SimSpaceControlComponentKey);
+			}
 		}
-	}
-}
 
-//======================================================================================================================
-FRigUnit_HierarchySetPhysicsBodyMovementType_Execute()
-{
-	DECLARE_SCOPE_HIERARCHICAL_COUNTER_RIGUNIT();
-	if (ExecuteContext.Hierarchy)
-	{
-		if (FRigPhysicsComponent* Component = Cast<FRigPhysicsComponent>(
-			ExecuteContext.Hierarchy->FindComponent(PhysicsBodyComponentKey)))
+		if (bAddParentSpaceControl)
 		{
-			Component->BodyData.MovementType = MovementType;
+			FRigComponentKey ParentSpaceControlComponentKey = Controller->AddComponent(
+				FRigPhysicsControlComponent::StaticStruct(),
+				TEXT("ParentSpaceControl"), PhysicsBodyComponentKey.ElementKey);
+			if (FRigPhysicsControlComponent* Component = Cast<FRigPhysicsControlComponent>(
+				ExecuteContext.Hierarchy->FindComponent(ParentSpaceControlComponentKey)))
+			{
+				Component->bUseParentBodyAsDefault = true;
+				Component->ChildBodyComponentKey = PhysicsBodyComponentKey;
+				Component->ControlData = ParentSpaceControlData;
+
+				ParentSpaceControlComponentKeys.Add(ParentSpaceControlComponentKey);
+			}
 		}
-	}
-}
+	} // Loop over SkeletalBodySetup
 
-//======================================================================================================================
-FRigUnit_HierarchySetPhysicsBodyCollisionType_Execute()
-{
-	DECLARE_SCOPE_HIERARCHICAL_COUNTER_RIGUNIT();
-	if (ExecuteContext.Hierarchy)
+	// Now bodies are created we can handle no-collision
+	const TMap<FRigidBodyIndexPair, bool>& CollisionDisableTable = PhysicsAsset->CollisionDisableTable;
+
+	for (const TPair<FRigidBodyIndexPair, bool>& PairPair : CollisionDisableTable)
 	{
-		if (FRigPhysicsComponent* Component = Cast<FRigPhysicsComponent>(
-			ExecuteContext.Hierarchy->FindComponent(PhysicsBodyComponentKey)))
+		if (PairPair.Value)
 		{
-			Component->BodyData.CollisionType = CollisionType;
+			const FRigidBodyIndexPair& Pair = PairPair.Key;
+			int32 I1 = Pair.Indices[0];
+			int32 I2 = Pair.Indices[1];
+			if (I1 >= 0 && I1 < PhysicsAsset->SkeletalBodySetups.Num() && 
+				I2 >= 0 && I2 < PhysicsAsset->SkeletalBodySetups.Num())
+			{
+				if (PhysicsAsset->SkeletalBodySetups[I1] && PhysicsAsset->SkeletalBodySetups[I2])
+				{
+					const FName BoneName1 = PhysicsAsset->SkeletalBodySetups[I1]->BoneName;
+					const FName BoneName2 = PhysicsAsset->SkeletalBodySetups[I2]->BoneName;
+
+					if (FRigComponentKey* BodyComponentKey1 = BoneToBodyComponentKeyMap.Find(BoneName1))
+					{
+						if (FRigPhysicsBodyComponent* Component = Cast<FRigPhysicsBodyComponent>(
+							ExecuteContext.Hierarchy->FindComponent(*BodyComponentKey1)))
+
+						if (FRigComponentKey* BodyComponentKey2 = BoneToBodyComponentKeyMap.Find(BoneName2))
+						{
+							if (ExecuteContext.Hierarchy->FindComponent(*BodyComponentKey2))
+							{
+								Component->NoCollisionBodies.Add(*BodyComponentKey2);
+							}
+						}
+					}
+				}
+			}
 		}
 	}
-}
 
-//======================================================================================================================
-FRigUnit_HierarchySetPhysicsBodyGravityMultiplier_Execute()
-{
-	DECLARE_SCOPE_HIERARCHICAL_COUNTER_RIGUNIT();
-	if (ExecuteContext.Hierarchy)
+	// Physics Joints
+	for (TObjectPtr<UPhysicsConstraintTemplate>& ConstraintTemplate : PhysicsAsset->ConstraintSetup)
 	{
-		if (FRigPhysicsComponent* Component = Cast<FRigPhysicsComponent>(
-			ExecuteContext.Hierarchy->FindComponent(PhysicsBodyComponentKey)))
+		if (!ConstraintTemplate)
 		{
-			Component->BodyData.GravityMultiplier = GravityMultiplier;
+			continue;
 		}
-	}
-}
 
-//======================================================================================================================
-FRigUnit_HierarchySetPhysicsBodyPhysicsBlendWeight_Execute()
-{
-	DECLARE_SCOPE_HIERARCHICAL_COUNTER_RIGUNIT();
-	if (ExecuteContext.Hierarchy)
-	{
-		if (FRigPhysicsComponent* Component = Cast<FRigPhysicsComponent>(
-			ExecuteContext.Hierarchy->FindComponent(PhysicsBodyComponentKey)))
+		// Note that physics assets are normally set up as child/parent for 1/2
+		// However, users can create their own constraints, and some assets will be the other way round.
+		FConstraintInstance& ConstraintInstance = ConstraintTemplate->DefaultInstance;
+		const FName ChildBoneName = ConstraintInstance.ConstraintBone1;
+		const FName ParentBoneName = ConstraintInstance.ConstraintBone2;
+
+		const FRigComponentKey* ChildBodyComponentKey = BoneToBodyComponentKeyMap.Find(ChildBoneName);
+		const FRigComponentKey* ParentBodyComponentKey = BoneToBodyComponentKeyMap.Find(ParentBoneName);
+
+		if (!ChildBodyComponentKey || !ParentBodyComponentKey)
 		{
-			Component->BodyData.PhysicsBlendWeight = PhysicsBlendWeight;
+			if (BonesToUse.IsEmpty())
+			{
+				// If BonesToUse is not empty, it is hard to say whether this is an
+				// error/warning condition. However, if the physics asset is OK and BonesToUse
+				// is empty, then we shouldn't get here.
+				UE_CONTROLRIG_RIGUNIT_REPORT_WARNING(
+					TEXT("InstantiateFromPhysicsAsset - unable to make physics joint between %s and %s"),
+					*ChildBoneName.ToString(), *ParentBoneName.ToString());
+			}
+			continue;
 		}
-	}
-}
 
-//======================================================================================================================
-FRigUnit_HierarchySetPhysicsBodyUseSkeletalAnimation_Execute()
-{
-	DECLARE_SCOPE_HIERARCHICAL_COUNTER_RIGUNIT();
-	if (ExecuteContext.Hierarchy)
-	{
-		if (FRigPhysicsComponent* Component = Cast<FRigPhysicsComponent>(
-			ExecuteContext.Hierarchy->FindComponent(PhysicsBodyComponentKey)))
+		const FConstraintProfileProperties& ConstraintProfileProperties = 
+			ConstraintTemplate->GetConstraintProfilePropertiesOrDefault(ConstraintProfileName);
+
+		FRigPhysicsJointData JointData;
+		SetUpJointData(JointData, ConstraintInstance, ConstraintProfileProperties);
+		FRigPhysicsDriveData DriveData;
+		if (bEnableDrives)
 		{
-			Component->BodyData.bUseSkeletalAnimation = bUseSkeletalAnimation;
+			SetUpDriveData(DriveData, ConstraintProfileProperties);
 		}
-	}
-}
 
-//======================================================================================================================
-FRigUnit_HierarchySetPhysicsBodyUpdateKinematicFromSimulation_Execute()
-{
-	DECLARE_SCOPE_HIERARCHICAL_COUNTER_RIGUNIT();
-	if (ExecuteContext.Hierarchy)
-	{
-		if (FRigPhysicsComponent* Component = Cast<FRigPhysicsComponent>(
-			ExecuteContext.Hierarchy->FindComponent(PhysicsBodyComponentKey)))
+		// Add the component
+		FRigComponentKey PhysicsJointComponentKey = Controller->AddComponent(
+			FRigPhysicsJointComponent::StaticStruct(), TEXT("PhysicsJoint"), ChildBodyComponentKey->ElementKey);
+		if (PhysicsJointComponentKey.IsValid())
 		{
-			Component->BodyData.bUpdateKinematicFromSimulation = bUpdateKinematicFromSimulation;
+			if (FRigPhysicsJointComponent* Component = Cast<FRigPhysicsJointComponent>(
+				ExecuteContext.Hierarchy->FindComponent(PhysicsJointComponentKey)))
+			{
+				Component->JointData = JointData;
+				Component->DriveData = DriveData;
+				Component->ParentBodyComponentKey = *ParentBodyComponentKey;
+				Component->ChildBodyComponentKey = *ChildBodyComponentKey;
+
+				PhysicsJointComponentKeys.Add(PhysicsJointComponentKey);
+			}
 		}
-	}
+	} // Loop over constraint templates
 }
diff --git a/Engine/Plugins/Experimental/PhysicsControl/Source/ControlRigPhysics/Private/RigPhysicsJointComponent.cpp b/Engine/Plugins/Experimental/PhysicsControl/Source/ControlRigPhysics/Private/RigPhysicsJointComponent.cpp
new file mode 100644
--- /dev/null
+++ b/Engine/Plugins/Experimental/PhysicsControl/Source/ControlRigPhysics/Private/RigPhysicsJointComponent.cpp
@@ -0,0 +1,68 @@
+// Copyright Epic Games, Inc. All Rights Reserved.
+
+#include "RigPhysicsJointComponent.h"
+#include "PhysicsControlObjectVersion.h"
+#include "RigPhysicsSolverComponent.h"
+#include "Rigs/RigHierarchy.h"
+#include "Rigs/RigHierarchyController.h"
+#include "ControlRigPhysicsModule.h"
+
+#if WITH_EDITOR
+#include "ControlRigPhysicsEditorStyle.h"
+#endif
+
+//======================================================================================================================
+void FRigPhysicsJointComponent::Save(FArchive& Ar)
+{
+	Ar.UsingCustomVersion(FPhysicsControlObjectVersion::GUID);
+
+	FRigBaseComponent::Save(Ar);
+	Ar << ParentBodyComponentKey;
+	Ar << ChildBodyComponentKey;
+	Ar << JointData;
+	Ar << DriveData;
+}
+
+//======================================================================================================================
+void FRigPhysicsJointComponent::Load(FArchive& Ar)
+{
+	FRigBaseComponent::Load(Ar);
+	Ar << ParentBodyComponentKey;
+	Ar << ChildBodyComponentKey;
+	Ar << JointData;
+	Ar << DriveData;
+}
+
+//======================================================================================================================
+bool FRigPhysicsJointComponent::CanBeAddedTo(
+	const FRigElementKey& InElementKey, const URigHierarchy* InHierarchy, FString* OutFailureReason) const
+{
+	if(InElementKey.Type != ERigElementType::Bone)
+	{
+		if(OutFailureReason)
+		{
+			*OutFailureReason = TEXT("Physics joint components can only be added to bones.");
+		}
+		return false;
+	}
+	return true;
+}
+
+
+//======================================================================================================================
+void FRigPhysicsJointComponent::OnRigHierarchyKeyChanged(const FRigHierarchyKey& InOldKey, const FRigHierarchyKey& InNewKey)
+{
+	FRigBaseComponent::OnRigHierarchyKeyChanged(InOldKey, InNewKey);
+
+	if (InOldKey.IsComponent() && InNewKey.IsComponent())
+	{
+		if (ParentBodyComponentKey == InOldKey.GetComponent())
+		{
+			ParentBodyComponentKey = InNewKey.GetComponent();
+		}
+		if (ChildBodyComponentKey == InOldKey.GetComponent())
+		{
+			ChildBodyComponentKey = InNewKey.GetComponent();
+		}
+	}
+}
diff --git a/Engine/Plugins/Experimental/PhysicsControl/Source/ControlRigPhysics/Private/RigPhysicsJointExecution.cpp b/Engine/Plugins/Experimental/PhysicsControl/Source/ControlRigPhysics/Private/RigPhysicsJointExecution.cpp
new file mode 100644
--- /dev/null
+++ b/Engine/Plugins/Experimental/PhysicsControl/Source/ControlRigPhysics/Private/RigPhysicsJointExecution.cpp
@@ -0,0 +1,80 @@
+// Copyright Epic Games, Inc. All Rights Reserved.
+
+#include "RigPhysicsJointExecution.h"
+#include "RigPhysicsJointComponent.h"
+
+#include "Rigs/RigHierarchyController.h"
+#include "Units/RigUnitContext.h"
+#include "Units/Execution/RigUnit_PrepareForExecution.h"
+#include "ControlRigPhysicsModule.h"
+
+#include "ControlRig.h"
+#if WITH_EDITOR
+#include "RigVMModel/RigVMController.h"
+#include "ControlRigBlueprint.h"
+#endif
+
+#include UE_INLINE_GENERATED_CPP_BY_NAME(RigPhysicsJointExecution)
+
+//======================================================================================================================
+FRigUnit_AddPhysicsJoint_Execute()
+{
+	if (!ExecuteContext.Hierarchy)
+	{
+		return;
+	}
+	if (ExecuteContext.GetEventName() != FRigUnit_PrepareForExecution::EventName)
+	{
+		UE_CONTROLRIG_RIGUNIT_REPORT_ERROR(TEXT("AddPhysicsJoint can only be used during Setup"));
+	}
+
+	if (URigHierarchyController* Controller = ExecuteContext.Hierarchy->GetController())
+	{
+		FRigHierarchyControllerInstructionBracket InstructionBracket(
+			Controller, ExecuteContext.GetInstructionIndex());
+		PhysicsJointComponentKey = Controller->AddComponent(
+			FRigPhysicsJointComponent::StaticStruct(), FRigPhysicsJointComponent::GetDefaultName(), Owner);
+		if (PhysicsJointComponentKey.IsValid())
+		{
+			if (FRigPhysicsJointComponent* Component = Cast<FRigPhysicsJointComponent>(
+				ExecuteContext.Hierarchy->FindComponent(PhysicsJointComponentKey)))
+			{
+				Component->ParentBodyComponentKey = ParentBodyComponentKey;
+				Component->ChildBodyComponentKey = ChildBodyComponentKey;
+				Component->JointData = JointData;
+				Component->DriveData = DriveData;
+			}
+		}
+	}
+}
+
+//======================================================================================================================
+FRigUnit_HierarchySetJointData_Execute()
+{
+	DECLARE_SCOPE_HIERARCHICAL_COUNTER_RIGUNIT();
+	if (!ExecuteContext.Hierarchy)
+	{
+		return;
+	}
+	if (FRigPhysicsJointComponent* Component = Cast<FRigPhysicsJointComponent>(
+		ExecuteContext.Hierarchy->FindComponent(PhysicsJointComponentKey)))
+	{
+		Component->JointData = JointData;
+	}
+}
+
+//======================================================================================================================
+FRigUnit_HierarchySetJointDriveData_Execute()
+{
+	DECLARE_SCOPE_HIERARCHICAL_COUNTER_RIGUNIT();
+	if (!ExecuteContext.Hierarchy)
+	{
+		return;
+	}
+	if (FRigPhysicsJointComponent* Component = Cast<FRigPhysicsJointComponent>(
+		ExecuteContext.Hierarchy->FindComponent(PhysicsJointComponentKey)))
+	{
+		Component->DriveData = DriveData;
+	}
+}
+
diff --git a/Engine/Plugins/Experimental/PhysicsControl/Source/ControlRigPhysics/Private/RigPhysicsSimulation.cpp b/Engine/Plugins/Experimental/PhysicsControl/Source/ControlRigPhysics/Private/RigPhysicsSimulation.cpp
--- a/Engine/Plugins/Experimental/PhysicsControl/Source/ControlRigPhysics/Private/RigPhysicsSimulation.cpp
+++ b/Engine/Plugins/Experimental/PhysicsControl/Source/ControlRigPhysics/Private/RigPhysicsSimulation.cpp
@@ -1,7 +1,8 @@
 // Copyright Epic Games, Inc. All Rights Reserved.
 
 #include "RigPhysicsSimulation.h"
-#include "RigPhysicsComponent.h"
+#include "RigPhysicsBodyComponent.h"
+#include "RigPhysicsJointComponent.h"
 #include "RigPhysicsControlComponent.h"
 #include "RigPhysicsSolverComponent.h"
 
@@ -49,29 +50,6 @@ DEFINE_LOG_CATEGORY(LogRigPhysics);
 constexpr int32 ConstraintChildIndex = 0;
 constexpr int32 ConstraintParentIndex = 1;
 
-//======================================================================================================================
-// Danny TODO remove duplication with RigidBodyWithControl
-static void SetConstraintEnabled(ImmediatePhysics::FJointHandle* const JointHandle, const bool bIsEnabled)
-{
-	if (JointHandle)
-	{
-		if (ImmediatePhysics::FJointHandle::FChaosConstraintHandle* ConstraintHandle = JointHandle->GetConstraint())
-		{
-			ConstraintHandle->SetConstraintEnabled(bIsEnabled);
-
-			if (!bIsEnabled)
-			{
-				// The call above disables the constraint, but if any actor in the simulation is
-				// flagged as dirty, it gets re-enabled! See UE-204006 TODO remove this workaround
-				// when the bug has been fixed.
-				ConstraintHandle->SetDriveParams(
-					Chaos::FVec3::ZeroVector, Chaos::FVec3::ZeroVector, Chaos::FVec3::ZeroVector,
-					Chaos::FVec3::ZeroVector, Chaos::FVec3::ZeroVector, Chaos::FVec3::ZeroVector);
-			}
-		}
-	}
-}
-
 //======================================================================================================================
 FRigPhysicsSimulation::FRigPhysicsSimulation(TObjectPtr<UControlRig> InOwningControlRig)
 	: FRigPhysicsSimulationBase(FRigPhysicsSimulation::StaticStruct()), OwningControlRig(InOwningControlRig)
@@ -86,7 +64,7 @@ bool FRigPhysicsSimulation::ShouldComponentBeInSimulation(
 {
 	const URigHierarchy* Hierarchy = OwningControlRig->GetHierarchy();
 
-	const FRigPhysicsComponent* PhysicsComponent = Cast<FRigPhysicsComponent>(
+	const FRigPhysicsBodyComponent* PhysicsComponent = Cast<FRigPhysicsBodyComponent>(
 		Hierarchy->FindComponent(ComponentKey));
 
 	const FRigPhysicsSolverComponent* PhysicsSolverComponent = Cast<FRigPhysicsSolverComponent>(
@@ -142,6 +120,12 @@ void FRigPhysicsSimulation::InitialiseSimulation(const FRigPhysicsSolverComponen
 
 	Simulation = MakeShared<ImmediatePhysics::FSimulation>();
 
+#if CHAOS_DEBUG_NAME
+	FString SimName = TEXT("ControlRigPhysics-");
+	SimName += OwningControlRig->GetName();
+	Simulation->SetDebugName(FName(SimName));
+#endif
+
 #if WITH_CHAOS_VISUAL_DEBUGGER
 	Simulation->GetChaosVDContextData().Id = FChaosVDRuntimeModule::Get().GenerateUniqueID();
 	Simulation->GetChaosVDContextData().Type = static_cast<int32>(EChaosVDContextType::Solver);
@@ -152,6 +136,10 @@ void FRigPhysicsSimulation::InitialiseSimulation(const FRigPhysicsSolverComponen
 	// behavior (in particular, for controls) that behave fairly independently of the control-rig
 	// tick rate.
 	Simulation->SetRewindVelocities(true);
+
+	// Always create a world actor at the origin, for attaching controls to. 
+	SimulationActorHandle = CreateBody(
+		FName(TEXT("Simulation")), FRigPhysicsCollision(), nullptr, FTransform());
 }
 
 //======================================================================================================================
@@ -169,27 +157,126 @@ void FRigPhysicsSimulation::InitialiseControlRecords(const FRigPhysicsSolverComp
 
 	TArray<FRigComponentKey> AllComponentKeys = Hierarchy->GetAllComponentKeys();
 
-	// All the components in this simulation
-	TArray<FRigComponentKey> UnsortedBodyComponentKeys;
-
 	for (FRigComponentKey ComponentKey : AllComponentKeys)
 	{
 		if (const FRigPhysicsControlComponent* ControlComponent = Cast<FRigPhysicsControlComponent>(
 			Hierarchy->FindComponent(ComponentKey)))
 		{
-			if (ShouldComponentBeInSimulation(SolverComponentKey, ControlComponent->ChildBodyComponentKey))
+			// The authored body components may be blank, in which case we need to find the automatic ones.
+			FRigControlRecord ControlRecord;
+			ControlRecord.ParentBodyComponentKey = ControlComponent->ParentBodyComponentKey;
+			ControlRecord.ChildBodyComponentKey = ControlComponent->ChildBodyComponentKey;
+
+			// Automate the child
+			if (!ControlRecord.ChildBodyComponentKey.IsValid())
+			{
+				TArray<FRigComponentKey> SiblingComponentKeys = Hierarchy->GetComponentKeys(ComponentKey.ElementKey);
+				for (FRigComponentKey SiblingComponentKey : SiblingComponentKeys)
+				{
+					if (ShouldComponentBeInSimulation(SolverComponentKey, SiblingComponentKey))
+					{
+						ControlRecord.ChildBodyComponentKey = SiblingComponentKey;
+						break;
+					}
+				}
+			}
+
+			if (ControlComponent->bUseParentBodyAsDefault)
 			{
-				if (!ControlComponent->ParentBodyComponentKey.IsValid() || 
-					ShouldComponentBeInSimulation(SolverComponentKey, ControlComponent->ParentBodyComponentKey))
+				if (!ControlRecord.ParentBodyComponentKey.IsValid())
+				{
+					FRigElementKey ParentElementKey = Hierarchy->GetFirstParent(ComponentKey.ElementKey);
+					TArray<FRigComponentKey> ParentComponentKeys = Hierarchy->GetComponentKeys(ParentElementKey);
+					for (FRigComponentKey ParentComponentKey : ParentComponentKeys)
+					{
+						if (ShouldComponentBeInSimulation(SolverComponentKey, ParentComponentKey))
+						{
+							ControlRecord.ParentBodyComponentKey = ParentComponentKey;
+							break;
+						}
+					}
+				}
+			}
+
+			if (ShouldComponentBeInSimulation(SolverComponentKey, ControlRecord.ChildBodyComponentKey))
+			{
+				// Here, an invalid parent component key indicates a sim-space control. 
+				if (!ControlRecord.ParentBodyComponentKey.IsValid() ||
+					ShouldComponentBeInSimulation(SolverComponentKey, ControlRecord.ParentBodyComponentKey))
 				{
 					// Just make the record for now - it will be instantiated later
-					ControlRecords.Add(ComponentKey);
+					ControlRecords.Add(ComponentKey, ControlRecord);
 				}
 			}
 		}
 	}
 }
 
+//======================================================================================================================
+void FRigPhysicsSimulation::InitialiseJointRecords(const FRigPhysicsSolverComponent* SolverComponent)
+{
+	ensure(JointRecords.IsEmpty());
+
+	const URigHierarchy* Hierarchy = OwningControlRig->GetHierarchy();
+	if (!Hierarchy || !SolverComponent)
+	{
+		return;
+	}
+
+	FRigComponentKey SolverComponentKey = SolverComponent->GetKey();
+
+	TArray<FRigComponentKey> AllComponentKeys = Hierarchy->GetAllComponentKeys();
+
+	for (FRigComponentKey ComponentKey : AllComponentKeys)
+	{
+		if (const FRigPhysicsJointComponent* JointComponent = Cast<FRigPhysicsJointComponent>(
+			Hierarchy->FindComponent(ComponentKey)))
+		{
+			// The authored body components may be blank, in which case we need to find the automatic ones.
+			FRigJointRecord JointRecord;
+			JointRecord.ParentBodyComponentKey = JointComponent->ParentBodyComponentKey;
+			JointRecord.ChildBodyComponentKey = JointComponent->ChildBodyComponentKey;
+
+			if (!JointRecord.ChildBodyComponentKey.IsValid())
+			{
+				TArray<FRigComponentKey> SiblingComponentKeys = Hierarchy->GetComponentKeys(ComponentKey.ElementKey);
+				for (FRigComponentKey SiblingComponentKey : SiblingComponentKeys)
+				{
+					if (ShouldComponentBeInSimulation(SolverComponentKey, SiblingComponentKey))
+					{
+						JointRecord.ChildBodyComponentKey = SiblingComponentKey;
+						break;
+					}
+				}
+			}
+
+			if (!JointRecord.ParentBodyComponentKey.IsValid())
+			{
+				FRigElementKey ParentElementKey = Hierarchy->GetFirstParent(ComponentKey.ElementKey);
+				TArray<FRigComponentKey> ParentComponentKeys = Hierarchy->GetComponentKeys(ParentElementKey);
+				for (FRigComponentKey ParentComponentKey : ParentComponentKeys)
+				{
+					if (ShouldComponentBeInSimulation(SolverComponentKey, ParentComponentKey))
+					{
+						JointRecord.ParentBodyComponentKey = ParentComponentKey;
+						break;
+					}
+				}
+			}
+
+			if (ShouldComponentBeInSimulation(SolverComponentKey, JointRecord.ChildBodyComponentKey))
+			{
+				if (ShouldComponentBeInSimulation(SolverComponentKey, JointRecord.ParentBodyComponentKey))
+				{
+					// Just make the record for now - it will be instantiated later
+					JointRecords.Add(ComponentKey, JointRecord);
+				}
+			}
+		}
+	}
+}
+
+
 //======================================================================================================================
 void FRigPhysicsSimulation::InitialiseBodyRecords(const FRigPhysicsSolverComponent* SolverComponent)
 {
@@ -239,10 +326,6 @@ void FRigPhysicsSimulation::DestroyPhysicsSimulation()
 	for (TPair<FRigComponentKey, FRigBodyRecord>& BodyRecordPair : BodyRecords)
 	{
 		FRigBodyRecord& Record = BodyRecordPair.Value;
-		if (Record.JointHandle)
-		{
-			Simulation->DestroyJoint(Record.JointHandle);
-		}
 		if (Record.ActorHandle)
 		{
 			Simulation->DestroyActor(Record.ActorHandle);
@@ -250,6 +333,17 @@ void FRigPhysicsSimulation::DestroyPhysicsSimulation()
 	}
 	BodyRecords.Empty();
 
+
+	for (TPair<FRigComponentKey, FRigJointRecord>& JointRecordPair : JointRecords)
+	{
+		FRigJointRecord& Record = JointRecordPair.Value;
+		if (Record.JointHandle)
+		{
+			Simulation->DestroyJoint(Record.JointHandle);
+		}
+	}
+	JointRecords.Empty();
+
 	for (TPair<FRigComponentKey, FRigControlRecord>& ControlRecordPair : ControlRecords)
 	{
 		FRigControlRecord& Record = ControlRecordPair.Value;
@@ -372,6 +466,10 @@ ImmediatePhysics::FActorHandle* FRigPhysicsSimulation::CreateBody(
 	{
 		BodyInstance.bSimulatePhysics = false;
 	}
+
+	// Without this flag then creation only works when it is on the game thread.
+	BodySetup->bNeverNeedsCookedCollisionData = true;
+
 	BodyInstance.InitBody(BodySetup, BodyRelSimSpaceTM, nullptr, nullptr);
 
 	ImmediatePhysics::EActorType ActorType = Dynamics ? ImmediatePhysics::EActorType::DynamicActor : ImmediatePhysics::EActorType::KinematicActor;
@@ -411,26 +509,43 @@ void FRigPhysicsSimulation::Instantiate(const FRigPhysicsSolverComponent* Solver
 
 	InitialiseBodyRecords(SolverComponent);
 
+	InitialiseJointRecords(SolverComponent);
+
 	InitialiseControlRecords(SolverComponent);
 
-	InstantiateSolverCollision(SolverComponent);
+	FRigPhysicsIgnorePairs IgnorePairs;
 
-	InstantiatePhysicsBodies(SolverComponent);
+	InstantiatePhysicsBodies(SolverComponent, IgnorePairs);
 
-	InstantiatePhysicsJoints(SolverComponent);
+	InstantiatePhysicsJoints(SolverComponent, IgnorePairs);
 
-	InstantiateControls(SolverComponent);
+	InstantiateControls(SolverComponent, IgnorePairs);
+
+	// This is done last as it applies IgnorePairs
+	InstantiateSolverCollision(SolverComponent, IgnorePairs);
 
 	bNeedToInstantiate = false;
 }
 
 //======================================================================================================================
-void FRigPhysicsSimulation::InstantiateSolverCollision(const FRigPhysicsSolverComponent* SolverComponent)
+ImmediatePhysics::FActorHandle* FRigPhysicsSimulation::GetActor(const FRigComponentKey& ComponentKey) const
 {
-	// Always create a world actor at the origin, for attaching controls to. 
-	SimulationActorHandle = CreateBody(
-		FName(TEXT("Simulation")), FRigPhysicsCollision(), nullptr, FTransform());
+	if (const FRigBodyRecord* BodyRecord = BodyRecords.Find(ComponentKey))
+	{
+		return BodyRecord->ActorHandle;
+	}
+	if (ComponentKey == PhysicsSolverComponentKey)
+	{
+		return SimulationActorHandle;
+	}
+	return nullptr;
+}
 
+//======================================================================================================================
+void FRigPhysicsSimulation::InstantiateSolverCollision(
+	const FRigPhysicsSolverComponent* SolverComponent,
+	FRigPhysicsIgnorePairs&           IgnorePairs)
+{
 	// Optionally create an object to contain environment collision
 	if (!SolverComponent->SolverSettings.Collision.IsEmpty())
 	{
@@ -442,10 +557,29 @@ void FRigPhysicsSimulation::InstantiateSolverCollision(const FRigPhysicsSolverCo
 		CollisionActorHandle = CreateBody(
 			FName(TEXT("Environment")), SolverComponent->SolverSettings.Collision, nullptr, BodyRelSimSpaceTM);
 	}
+
+
+	// Add no-collision pairs
+	TArray<ImmediatePhysics_Chaos::FSimulation::FIgnorePair> ChaosIgnorePairs;
+
+	for (const FRigPhysicsIgnorePair& IgnorePair : IgnorePairs)
+	{
+		ImmediatePhysics::FSimulation::FIgnorePair ChaosPair;
+		ChaosPair.A = GetActor(IgnorePair.Key);
+		ChaosPair.B = GetActor(IgnorePair.Value);
+		if (ChaosPair.A && ChaosPair.B)
+		{
+			ChaosIgnorePairs.Add(ChaosPair);
+		}
+	}
+	Simulation->SetIgnoreCollisionPairTable(ChaosIgnorePairs);
+
 }
 
 //======================================================================================================================
-void FRigPhysicsSimulation::InstantiatePhysicsBodies(const FRigPhysicsSolverComponent* SolverComponent)
+void FRigPhysicsSimulation::InstantiatePhysicsBodies(
+	const FRigPhysicsSolverComponent* SolverComponent,
+	FRigPhysicsIgnorePairs&           IgnorePairs)
 {
 	const URigHierarchy* Hierarchy = OwningControlRig->GetHierarchy();
 
@@ -457,7 +591,7 @@ void FRigPhysicsSimulation::InstantiatePhysicsBodies(const FRigPhysicsSolverComp
 		const FRigComponentKey& ComponentKey = BodyRecordPair.Key;
 		FRigBodyRecord& Record = BodyRecordPair.Value;
 
-		if (const FRigPhysicsComponent* PhysicsComponent = Cast<FRigPhysicsComponent>(
+		if (const FRigPhysicsBodyComponent* PhysicsComponent = Cast<FRigPhysicsBodyComponent>(
 			Hierarchy->FindComponent(ComponentKey)))
 		{
 			const FRigPhysicsCollision& Collision = PhysicsComponent->Collision;
@@ -478,177 +612,121 @@ void FRigPhysicsSimulation::InstantiatePhysicsBodies(const FRigPhysicsSolverComp
 					ConvertComponentSpaceTransformToSimSpace(SolverComponent->SolverSettings, SourceComponentSpaceTM);
 				Record.ActorHandle = CreateBody(ComponentKey.ElementKey.Name, Collision, &Dynamics, SourceSimulationSpaceTM);
 			}
+
+			for (const FRigComponentKey& NoCollisionKey : PhysicsComponent->NoCollisionBodies)
+			{
+				IgnorePairs.Add(FRigPhysicsIgnorePair(ComponentKey, NoCollisionKey));
+			}
 		}
 	}
 }
 
 //======================================================================================================================
-void FRigPhysicsSimulation::InstantiatePhysicsJoints(const FRigPhysicsSolverComponent* SolverComponent)
+void FRigPhysicsSimulation::InstantiatePhysicsJoints(
+	const FRigPhysicsSolverComponent* SolverComponent,
+	FRigPhysicsIgnorePairs&           IgnorePairs)
 {
 	const URigHierarchy* Hierarchy = OwningControlRig->GetHierarchy();
 
 	const FRigPhysicsSolverSettings& SolverSettings = SolverComponent->SolverSettings;
 	const FRigPhysicsSimulationSpaceSettings& SimulationSpaceSettings = SolverComponent->SimulationSpaceSettings;
 
-	TArray<ImmediatePhysics::FSimulation::FIgnorePair> IgnorePairs;
-
 	// Once all the bodies are created, we can make their physics joints
-	for (TPair<FRigComponentKey, FRigBodyRecord>& BodyRecordPair : BodyRecords)
+	for (TPair<FRigComponentKey, FRigJointRecord>& JointRecordPair : JointRecords)
 	{
-		const FRigComponentKey& ChildBodyComponentKey = BodyRecordPair.Key;
-		FRigBodyRecord& ChildBodyRecord = BodyRecordPair.Value;
+		const FRigComponentKey& JointComponentKey = JointRecordPair.Key;
+		FRigJointRecord& JointRecord = JointRecordPair.Value;
 
-		ImmediatePhysics::FActorHandle* ChildActorHandle = ChildBodyRecord.ActorHandle;
-		if (!ChildActorHandle)
+		if (const FRigPhysicsJointComponent* PhysicsJointComponent = Cast<FRigPhysicsJointComponent>(
+			Hierarchy->FindComponent(JointComponentKey)))
 		{
-			continue;
-		}
+			const FRigPhysicsJointData& JointData = PhysicsJointComponent->JointData;
 
-		if (const FRigPhysicsComponent* ChildPhysicsComponent = Cast<FRigPhysicsComponent>(
-			Hierarchy->FindComponent(ChildBodyComponentKey)))
-		{
-			const FRigPhysicsJoint& PhysicsJoint = ChildPhysicsComponent->Joint;
-			if (!PhysicsJoint.bEnable)
-			{
-				continue;
-			}
+			FRigComponentKey ChildBodyComponentKey = JointRecord.ChildBodyComponentKey;
+			FRigElementKey ChildBoneKey = ChildBodyComponentKey.ElementKey;
 
-			// Now we need to find the parent body (etc) for the physics joint
-			FRigComponentKey ParentBodyComponentKey = PhysicsJoint.ParentBody;
+			FRigComponentKey ParentBodyComponentKey = JointRecord.ParentBodyComponentKey;
+			FRigElementKey ParentBoneKey = ParentBodyComponentKey.ElementKey;
 
-			// And along the way we will discover various other things
-			FRigElementKey ChildBoneKey = ChildBodyComponentKey.ElementKey;
-			if (!ChildBoneKey.IsValid())
+			// Joints require both parent and child to exist
+			const FRigPhysicsBodyComponent* ChildPhysicsComponent =
+				Cast<FRigPhysicsBodyComponent>(Hierarchy->FindComponent(ChildBodyComponentKey));
+			FRigBodyRecord* ChildBodyRecord = BodyRecords.Find(ChildBodyComponentKey);
+			ImmediatePhysics::FActorHandle* ChildActorHandle = ChildBodyRecord ? ChildBodyRecord->ActorHandle : nullptr;
+			if (!ChildActorHandle)
 			{
 				continue;
 			}
 
-			if (!ParentBodyComponentKey.IsValid())
-			{
-				// Physics bodies are parented to a bone - we want the parent of the bone
-				FRigElementKey ParentBoneKey = Hierarchy->GetFirstParent(ChildBoneKey);
-				// Now we have the regular bone, we want the physics element that is the child
-				if (ParentBoneKey.IsValid())
-				{
-					TArray<FRigComponentKey> ComponentKeys = Hierarchy->GetComponentKeys(ParentBoneKey);
-					for (const FRigComponentKey& ComponentKey : ComponentKeys)
-					{
-						const FRigBaseComponent* BaseComponent = Hierarchy->FindComponent(ComponentKey);
-						// Danny TODO what's the best way to find/check this?!
-						if (Cast<FRigPhysicsComponent>(BaseComponent))
-						{
-							ParentBodyComponentKey = ComponentKey;
-							break;
-						}
-					}
-				}
-			}
-			const FRigPhysicsComponent* ParentPhysicsComponent =
-				Cast<FRigPhysicsComponent>(Hierarchy->FindComponent(ParentBodyComponentKey));
-			if (!ParentPhysicsComponent)
+			const FRigPhysicsBodyComponent* ParentPhysicsComponent =
+				Cast<FRigPhysicsBodyComponent>(Hierarchy->FindComponent(ParentBodyComponentKey));
+			FRigBodyRecord* ParentBodyRecord = BodyRecords.Find(ParentBodyComponentKey);
+			ImmediatePhysics::FActorHandle* ParentActorHandle = ParentBodyRecord ? ParentBodyRecord->ActorHandle : nullptr;
+			if (!ParentActorHandle)
 			{
 				continue;
 			}
 
-			ChildBodyRecord.ParentBodyComponentKey = ParentBodyComponentKey;
+			// Make the physics joint (joint constraint). 
+			// UE likes to treat Body1 (index 0) as the child, and Body2 (index 1) as the parent
+			{
+				const FTransform ParentCoMTransform = ParentActorHandle->GetLocalCoMTransform();
+				const FTransform ChildCoMTransform = ChildActorHandle->GetLocalCoMTransform();
 
-			const FRigBodyRecord* ParentRecord = nullptr;
-			ImmediatePhysics::FActorHandle* ParentActorHandle = nullptr;
+				Chaos::FPBDJointSettings JointSettings;
+				if (JointData.bAutoCalculateChildOffset)
+				{
+					// Danny TODO - is the "non-override" actually always simply the identity TM?
+					JointSettings.ConnectorTransforms[0] = FTransform();
+				}
+				JointSettings.ConnectorTransforms[0] =
+					JointData.ExtraChildOffset * JointSettings.ConnectorTransforms[0];
 
-			// Find the record that corresponds to the parent, so we can see if there is a parent actor
-			for (const TPair<FRigComponentKey, FRigBodyRecord>& OtherBodyRecordPair : BodyRecords)
-			{
-				const FRigComponentKey& OtherBodyComponentKey = OtherBodyRecordPair.Key;
-				const FRigBodyRecord& OtherBodyRecord = OtherBodyRecordPair.Value;
-				if (OtherBodyComponentKey == ParentPhysicsComponent->GetKey())
+				if (JointData.bAutoCalculateParentOffset)
 				{
-					ParentRecord = &OtherBodyRecord;
-					ParentActorHandle = OtherBodyRecord.ActorHandle;
-					break;
+					// Danny TODO is it right to use the TM of the child bone relative to parent bone?
+					// What about if the body has got a TM too?
+					JointSettings.ConnectorTransforms[1] = Hierarchy->GetLocalTransform(ChildBoneKey, true);
 				}
-			}
-			if (!ParentRecord)
-			{
-				continue;
-			}
+				JointSettings.ConnectorTransforms[1] =
+					JointData.ExtraParentOffset * JointSettings.ConnectorTransforms[1];
 
-			// Make the physics joint (joint constraint). 
-			// UE likes to treat Body1 (index 0) as the child, and Body2 (index 1) as the parent
+				ImmediatePhysics_Chaos::UpdateJointSettingsFromLinearConstraint(JointData.LinearConstraint, JointSettings);
+				ImmediatePhysics_Chaos::UpdateJointSettingsFromConeConstraint(JointData.ConeConstraint, JointSettings);
+				ImmediatePhysics_Chaos::UpdateJointSettingsFromTwistConstraint(JointData.TwistConstraint, JointSettings);
 
-			const FTransform ParentCoMTransform = ParentActorHandle->GetLocalCoMTransform();
-			const FTransform ChildCoMTransform = ChildActorHandle->GetLocalCoMTransform();
+				// The physics setting is backwards, because we can't enable collision on bodies that
+				// are set to not collide for other reasons.
+				JointSettings.bCollisionEnabled = !JointData.bDisableCollision;
+				JointSettings.bProjectionEnabled = 
+					JointData.LinearProjectionAmount > 0.0f || JointData.AngularProjectionAmount > 0.0f;
+				JointSettings.AngularProjection = JointData.AngularProjectionAmount;
+				JointSettings.LinearProjection = JointData.LinearProjectionAmount;
+				JointSettings.ParentInvMassScale = JointData.ParentInverseMassScale;
 
-			Chaos::FPBDJointSettings ConstraintSettings;
-			if (PhysicsJoint.bAutoCalculateChildOffset)
-			{
-				// Danny TODO - is the "non-override" actually always simply the identity TM?
-				ConstraintSettings.ConnectorTransforms[0] = FTransform();
-			}
-			ConstraintSettings.ConnectorTransforms[0] =
-				PhysicsJoint.ExtraChildOffset * ConstraintSettings.ConnectorTransforms[0];
+				JointSettings.bUseLinearSolver = SolverSettings.bUseLinearJointSolver;
 
-			if (PhysicsJoint.bAutoCalculateParentOffset)
-			{
-				// Danny TODO is it right to use the TM of the child bone relative to parent bone?
-				// What about if the body has got a TM too?
-				ConstraintSettings.ConnectorTransforms[1] = Hierarchy->GetLocalTransform(ChildBoneKey, true);
-			}
-			ConstraintSettings.ConnectorTransforms[1] =
-				PhysicsJoint.ExtraParentOffset * ConstraintSettings.ConnectorTransforms[1];
-
-			// Danny TODO The linear limit is likely implemented as a box limit (check) - it's just
-			// not exposed as one. We should change this.
-			ConstraintSettings.LinearLimit = PhysicsJoint.LinearLimit;
-			Chaos::EJointMotionType LinearMotionType = GetMotionTypeFromLimitValue(ConstraintSettings.LinearLimit);
-			ConstraintSettings.LinearMotionTypes = { LinearMotionType, LinearMotionType, LinearMotionType };
-
-			ConstraintSettings.AngularLimits = FMath::DegreesToRadians(PhysicsJoint.AngularLimit);
-			ConstraintSettings.AngularMotionTypes = {
-				GetMotionTypeFromLimitValue(PhysicsJoint.AngularLimit[0]),
-				GetMotionTypeFromLimitValue(PhysicsJoint.AngularLimit[1]),
-				GetMotionTypeFromLimitValue(PhysicsJoint.AngularLimit[2])
-			};
-
-			// The physics setting is backwards, because we can't enable collision on bodies that
-			// are set to not collide for other reasons.
-			ConstraintSettings.bCollisionEnabled = !PhysicsJoint.bDisableCollision;
-			ConstraintSettings.bProjectionEnabled = 
-				PhysicsJoint.LinearProjectionAmount > 0.0f || PhysicsJoint.AngularProjectionAmount > 0.0f;
-			ConstraintSettings.AngularProjection = PhysicsJoint.AngularProjectionAmount;
-			ConstraintSettings.LinearProjection = PhysicsJoint.LinearProjectionAmount;
-			ConstraintSettings.ParentInvMassScale = PhysicsJoint.ParentInverseMassScale;
-
-			ConstraintSettings.bUseLinearSolver = SolverSettings.bUseLinearJointSolver;
-
-			ChildBodyRecord.JointHandle = Simulation->CreateJoint(
-				ImmediatePhysics::FJointSetup(ConstraintSettings, ChildActorHandle, ParentActorHandle));
-
-			// Now set up the drive, if requested
-			const FRigPhysicsDrive& Drive = PhysicsJoint.Drive;
-			if (ChildBodyRecord.JointHandle && Drive.bEnable)
-			{
-				if (Chaos::FPBDJointConstraintHandle* Constraint = ChildBodyRecord.JointHandle->GetConstraint())
+				JointRecord.JointHandle = Simulation->CreateJoint(
+					ImmediatePhysics::FJointSetup(JointSettings, ChildActorHandle, ParentActorHandle));
+
+				if (Chaos::FPBDJointConstraintHandle* Constraint = JointRecord.JointHandle->GetConstraint())
 				{
 					Chaos::FPBDJointSettings Settings = Constraint->GetSettings();
 					if (!Settings.bCollisionEnabled)
 					{
-						// Immediate physics doesn't actually honor the flag in settings, so do it manually
-						ImmediatePhysics::FSimulation::FIgnorePair Pair;
-						Pair.A = ParentActorHandle;
-						Pair.B = ChildActorHandle;
-						IgnorePairs.Add(Pair);
+						IgnorePairs.Add(FRigPhysicsIgnorePair(
+							JointRecord.ChildBodyComponentKey, JointRecord.ParentBodyComponentKey));
 					}
-
-					// The target and properties will be set in the pre-physics update
 				}
 			}
 		}
 	}
-	Simulation->SetIgnoreCollisionPairTable(IgnorePairs);
 }
 
 //======================================================================================================================
-void FRigPhysicsSimulation::InstantiateControls(const FRigPhysicsSolverComponent* SolverComponent)
+void FRigPhysicsSimulation::InstantiateControls(
+	const FRigPhysicsSolverComponent* SolverComponent,
+	FRigPhysicsIgnorePairs&           IgnorePairs)
 {
 	const URigHierarchy* Hierarchy = OwningControlRig->GetHierarchy();
 
@@ -660,10 +738,10 @@ void FRigPhysicsSimulation::InstantiateControls(const FRigPhysicsSolverComponent
 		if (const FRigPhysicsControlComponent* ControlComponent = Cast<FRigPhysicsControlComponent>(
 			Hierarchy->FindComponent(ComponentKey)))
 		{
-			if (FRigBodyRecord* ChildBodyRecord = BodyRecords.Find(ControlComponent->ChildBodyComponentKey))
+			if (FRigBodyRecord* ChildBodyRecord = BodyRecords.Find(ControlRecord.ChildBodyComponentKey))
 			{
 				// this can be nullptr - just means a global control
-				FRigBodyRecord* ParentBodyRecord = BodyRecords.Find(ControlComponent->ParentBodyComponentKey);
+				FRigBodyRecord* ParentBodyRecord = BodyRecords.Find(ControlRecord.ParentBodyComponentKey);
 
 				ImmediatePhysics::FActorHandle* const ChildBodyHandle = ChildBodyRecord->ActorHandle;
 				ImmediatePhysics::FActorHandle* const ParentBodyHandle = 
@@ -677,6 +755,12 @@ void FRigPhysicsSimulation::InstantiateControls(const FRigPhysicsSolverComponent
 					UE_LOG(LogRigPhysics, Warning,
 						TEXT("Unable to create control constraint for %s"), *ControlComponent->GetKey().ToString());
 				}
+
+				if (ControlComponent->ControlData.bDisableCollision)
+				{
+					IgnorePairs.Add(FRigPhysicsIgnorePair(
+						ControlRecord.ChildBodyComponentKey, ControlRecord.ParentBodyComponentKey));
+				}
 			}
 		}
 	}
@@ -687,10 +771,14 @@ void FRigPhysicsSimulation::UpdateBodyPrePhysics(
 	const FRigVMExecuteContext& ExecuteContext,
 	const FRigPhysicsSolverComponent* SolverComponent,
 	const FRigBodyRecord& Record, 
-	const FRigPhysicsComponent* PhysicsComponent,
+	const FRigPhysicsBodyComponent* PhysicsComponent,
 	const FSimulationSpaceData& SimulationSpaceData)
 {
 	URigHierarchy* Hierarchy = OwningControlRig->GetHierarchy();
+	if (!Record.ActorHandle || !Hierarchy)
+	{
+		return;
+	}
 	const FRigComponentKey ComponentKey = PhysicsComponent->GetKey();
 
 	UpdateBodyFromModifierData(
@@ -766,181 +854,128 @@ void FRigPhysicsSimulation::UpdateBodyPrePhysics(
 
 //======================================================================================================================
 void FRigPhysicsSimulation::UpdateJointPrePhysics(
-	FRigBodyRecord&             Record, 
-	const FRigPhysicsComponent* PhysicsComponent, 
-	const URigHierarchy*        Hierarchy, 
-	const float                 DeltaTime)
+	FRigJointRecord&                    Record, 
+	const FRigPhysicsJointComponent*    PhysicsJointComponent, 
+	const URigHierarchy*                Hierarchy, 
+	const float                         DeltaTime)
 {
-	const FRigComponentKey ComponentKey = PhysicsComponent->GetKey();
-
 	// Now update the joint targets
-	if (Record.JointHandle)
+	if (!Record.JointHandle)
+	{
+		return;
+	}
+	if (Chaos::FPBDJointConstraintHandle* Constraint = Record.JointHandle->GetConstraint())
 	{
-		if (Chaos::FPBDJointConstraintHandle* Constraint = Record.JointHandle->GetConstraint())
+		const FRigComponentKey ComponentKey = PhysicsJointComponent->GetKey();
+
+		// Set the drive strength etc
+		const FRigPhysicsJointData& JointData = PhysicsJointComponent->JointData;
+		const FRigPhysicsDriveData& DriveData = PhysicsJointComponent->DriveData;
+
+		Chaos::FPBDJointSettings Settings = Constraint->GetSettings();
+		ImmediatePhysics_Chaos::UpdateJointSettingsFromLinearDriveConstraint(DriveData.LinearDriveConstraint, Settings);
+		ImmediatePhysics_Chaos::UpdateJointSettingsFromAngularDriveConstraint(DriveData.AngularDriveConstraint, Settings);
+		Constraint->SetSettings(Settings);
+
+		// Now set the actual target
+		if (Settings.AngularDriveStiffness.SquaredLength() > 0.0f ||
+			Settings.AngularDriveDamping.SquaredLength() > 0.0f ||
+			Settings.LinearDriveStiffness.SquaredLength() > 0.0f ||
+			Settings.LinearDriveDamping.SquaredLength() > 0.0f)
 		{
-			// Set the drive strength etc
-			const FRigPhysicsDrive& Drive = PhysicsComponent->Joint.Drive;
-			if (Drive.bEnable)
+			// Multiplier on the velocity calculated from the current and previous target
+			float TargetVelocityMultiplier = 1.0f;
+
+			FRigElementKey ChildSourceKey;
+			if (const FRigPhysicsBodyComponent* ChildPhysicsComponent = Cast<FRigPhysicsBodyComponent>(
+				Hierarchy->FindComponent(Record.ChildBodyComponentKey)))
 			{
-				float AngularSpring;
-				float AngularDamping;
-				const float MaxTorque = Drive.MaxTorque;
-
-				float LinearSpring;
-				float LinearDamping;
-				const float MaxForce = Drive.MaxForce;
-
-				UE::PhysicsControl::ConvertStrengthToSpringParams(
-					AngularSpring, AngularDamping,
-					Drive.AngularStrength, Drive.AngularDampingRatio, Drive.AngularExtraDamping);
-				UE::PhysicsControl::ConvertStrengthToSpringParams(
-					LinearSpring, LinearDamping,
-					Drive.LinearStrength, Drive.LinearDampingRatio, Drive.LinearExtraDamping);
-
-				Chaos::FPBDJointSettings Settings = Constraint->GetSettings();
-
-				// Linear
-				Settings.LinearDriveForceMode = Chaos::EJointForceMode::Acceleration;
-				const bool bLinearPositionDriveEnabled = LinearSpring > 0.0f;
-				const bool bLinearVelocityDriveEnabled = LinearDamping > 0.0f;
-				Settings.bLinearPositionDriveEnabled = {
-					bLinearPositionDriveEnabled, bLinearPositionDriveEnabled, bLinearPositionDriveEnabled };
-				Settings.bLinearVelocityDriveEnabled = {
-					bLinearVelocityDriveEnabled, bLinearVelocityDriveEnabled, bLinearVelocityDriveEnabled };
-
-				// Angular
-				Settings.AngularDriveForceMode = Chaos::EJointForceMode::Acceleration;
-				// SLERP is generally better, but doesn't work if anything is locked
-				if (Settings.AngularMotionTypes[0] == Chaos::EJointMotionType::Locked ||
-					Settings.AngularMotionTypes[1] == Chaos::EJointMotionType::Locked ||
-					Settings.AngularMotionTypes[2] == Chaos::EJointMotionType::Locked)
-				{
-					Settings.bAngularSwingPositionDriveEnabled = AngularSpring > 0.0f;
-					Settings.bAngularTwistPositionDriveEnabled = AngularSpring > 0.0f;
-					Settings.bAngularSwingVelocityDriveEnabled = AngularDamping > 0.0f;
-					Settings.bAngularTwistVelocityDriveEnabled = AngularDamping > 0.0f;
-				}
-				else
+				ChildSourceKey = ChildPhysicsComponent->BodySolverSettings.SourceBone;
+				if (!ChildSourceKey.IsValid())
 				{
-					Settings.bAngularSLerpPositionDriveEnabled = AngularSpring > 0.0f;
-					Settings.bAngularSLerpVelocityDriveEnabled = AngularDamping > 0.0f;
+					ChildSourceKey = ComponentKey.ElementKey;
 				}
-
-				Constraint->SetSettings(Settings);
-				Constraint->SetDriveParams(
-					Chaos::FVec3(LinearSpring), Chaos::FVec3(LinearDamping), Chaos::FVec3(MaxForce),
-					Chaos::FVec3(AngularSpring), Chaos::FVec3(AngularDamping), Chaos::FVec3(MaxTorque));
+				TargetVelocityMultiplier = DriveData.SkeletalAnimationVelocityMultiplier;
 			}
-			else
-			{
-				Constraint->SetDriveParams(
-					Chaos::FVec3(0), Chaos::FVec3(0), Chaos::FVec3(0),
-					Chaos::FVec3(0), Chaos::FVec3(0), Chaos::FVec3(0));
-			}
-
-			// Now set the actual target
-			const Chaos::FPBDJointSettings& Settings = Constraint->GetSettings();
 
-			if (Settings.AngularDriveStiffness.SquaredLength() > 0.0f ||
-				Settings.AngularDriveDamping.SquaredLength() > 0.0f ||
-				Settings.LinearDriveStiffness.SquaredLength() > 0.0f ||
-				Settings.LinearDriveDamping.SquaredLength() > 0.0f)
+			FRigElementKey ParentSourceKey;
+			if (const FRigPhysicsBodyComponent* ParentPhysicsComponent = Cast<FRigPhysicsBodyComponent>(
+				Hierarchy->FindComponent(Record.ParentBodyComponentKey)))
 			{
-				// Multiplier on the velocity calculated from the current and previous target
-				float TargetVelocityMultiplier = 1.0f;
-
-				FRigElementKey ChildSourceKey;
-				if (const FRigPhysicsComponent* ChildPhysicsComponent = Cast<FRigPhysicsComponent>(
-					Hierarchy->FindComponent(ComponentKey)))
-				{
-					ChildSourceKey = ChildPhysicsComponent->BodySolverSettings.SourceBone;
-					if (!ChildSourceKey.IsValid())
-					{
-						ChildSourceKey = ComponentKey.ElementKey;
-					}
-					TargetVelocityMultiplier =
-						ChildPhysicsComponent->Joint.Drive.SkeletalAnimationVelocityMultiplier;
-				}
-
-				FRigElementKey ParentSourceKey;
-				if (const FRigPhysicsComponent* ParentPhysicsComponent = Cast<FRigPhysicsComponent>(
-					Hierarchy->FindComponent(Record.ParentBodyComponentKey)))
+				ParentSourceKey = ParentPhysicsComponent->BodySolverSettings.SourceBone;
+				if (!ParentSourceKey.IsValid())
 				{
-					ParentSourceKey = ParentPhysicsComponent->BodySolverSettings.SourceBone;
-					if (!ParentSourceKey.IsValid())
-					{
-						ParentSourceKey = Record.ParentBodyComponentKey.ElementKey;
-					}
+					ParentSourceKey = Record.ParentBodyComponentKey.ElementKey;
 				}
+			}
 
-				if (!ChildSourceKey.IsValid() || !ParentSourceKey.IsValid())
-				{
-					return;
-				}
+			if (!ChildSourceKey.IsValid() || !ParentSourceKey.IsValid())
+			{
+				return;
+			}
 
-				// Note that the drive operates between a parent and child part, so we
-				// don't need to worry about global/component (etc) space.
-				const FTransform ChildTM = Hierarchy->GetGlobalTransform(ChildSourceKey);
-				const FTransform ParentTM = Hierarchy->GetGlobalTransform(ParentSourceKey);
+			// Note that the drive operates between a parent and child part, so we
+			// don't need to worry about global/component (etc) space.
+			const FTransform ChildTM = Hierarchy->GetGlobalTransform(ChildSourceKey);
+			const FTransform ParentTM = Hierarchy->GetGlobalTransform(ParentSourceKey);
 
-				const UE::PhysicsControl::FPosQuat ComponentSpaceParentFrameTM = 
-					Settings.ConnectorTransforms[1] * ParentTM;
-				const UE::PhysicsControl::FPosQuat ComponentSpaceChildFrameTM = 
-					Settings.ConnectorTransforms[0] * ChildTM;
+			const UE::PhysicsControl::FPosQuat ComponentSpaceParentFrameTM = 
+				Settings.ConnectorTransforms[1] * ParentTM;
+			const UE::PhysicsControl::FPosQuat ComponentSpaceChildFrameTM = 
+				Settings.ConnectorTransforms[0] * ChildTM;
 
-				const UE::PhysicsControl::FPosQuat DriveTargetTM =
-					ComponentSpaceParentFrameTM.Inverse() * ComponentSpaceChildFrameTM;
+			const UE::PhysicsControl::FPosQuat DriveTargetTM =
+				ComponentSpaceParentFrameTM.Inverse() * ComponentSpaceChildFrameTM;
 
-				Constraint->SetLinearDrivePositionTarget(DriveTargetTM.GetTranslation());
-				Constraint->SetAngularDrivePositionTarget(DriveTargetTM.GetRotation());
+			Constraint->SetLinearDrivePositionTarget(DriveTargetTM.GetTranslation());
+			Constraint->SetAngularDrivePositionTarget(DriveTargetTM.GetRotation());
 
-				if (Record.PreviousDriveTargetUpdateCounter + 1 == UpdateCounter &&
-					TargetVelocityMultiplier > 0.0f)
+			if (Record.PreviousDriveTargetUpdateCounter + 1 == UpdateCounter &&
+				TargetVelocityMultiplier > 0.0f)
+			{
+				if (DeltaTime > SMALL_NUMBER)
 				{
-					if (DeltaTime > SMALL_NUMBER)
+					const UE::PhysicsControl::FPosQuat DriveTargetTMDelta =
+						DriveTargetTM * Record.PreviousDriveTargetTM.Inverse();
+					const FVector Velocity = DriveTargetTMDelta.GetTranslation() / DeltaTime;
+					const FVector AngularVelocity =
+						DriveTargetTMDelta.GetRotation().GetShortestArcWith(
+							FQuat::Identity).ToRotationVector() / DeltaTime;
+					if (!Velocity.ContainsNaN() && !AngularVelocity.ContainsNaN())
 					{
-						const UE::PhysicsControl::FPosQuat DriveTargetTMDelta =
-							DriveTargetTM * Record.PreviousDriveTargetTM.Inverse();
-						const FVector Velocity = DriveTargetTMDelta.GetTranslation() / DeltaTime;
-						const FVector AngularVelocity =
-							DriveTargetTMDelta.GetRotation().GetShortestArcWith(
-								FQuat::Identity).ToRotationVector() / DeltaTime;
-						if (!Velocity.ContainsNaN() && !AngularVelocity.ContainsNaN())
-						{
-							Constraint->SetLinearDriveVelocityTarget(Velocity * TargetVelocityMultiplier);
-							Constraint->SetAngularDriveVelocityTarget(AngularVelocity * TargetVelocityMultiplier);
-						}
+						Constraint->SetLinearDriveVelocityTarget(Velocity * TargetVelocityMultiplier);
+						Constraint->SetAngularDriveVelocityTarget(AngularVelocity * TargetVelocityMultiplier);
 					}
 				}
-				else
-				{
-					Constraint->SetLinearDriveVelocityTarget(Chaos::FVec3(0));
-					Constraint->SetAngularDriveVelocityTarget(Chaos::FVec3(0));
-				}
-				Record.PreviousDriveTargetUpdateCounter = UpdateCounter;
-				Record.PreviousDriveTargetTM = DriveTargetTM;
 			}
+			else
+			{
+				Constraint->SetLinearDriveVelocityTarget(Chaos::FVec3(0));
+				Constraint->SetAngularDriveVelocityTarget(Chaos::FVec3(0));
+			}
+			Record.PreviousDriveTargetUpdateCounter = UpdateCounter;
+			Record.PreviousDriveTargetTM = DriveTargetTM;
 		}
 	}
 }
 
 //======================================================================================================================
 static UE::PhysicsControl::FPosQuat CalculateTargetTM(
-	const URigHierarchy* Hierarchy,
-	const Chaos::FPBDJointSettings& JointSettings,
-	const FRigPhysicsControlComponent* ControlComponent)
+	const URigHierarchy*               Hierarchy,
+	const Chaos::FPBDJointSettings&    JointSettings,
+	const FRigControlRecord&           Record)
 {
-	const FTransform ChildTM = ControlComponent->ChildBodyComponentKey.IsValid() ?
-		Hierarchy->GetGlobalTransform(ControlComponent->ChildBodyComponentKey.ElementKey) :
+	const FTransform ChildTM = Record.ChildBodyComponentKey.IsValid() ?
+		Hierarchy->GetGlobalTransform(Record.ChildBodyComponentKey.ElementKey) :
 		FTransform();
 
 	const UE::PhysicsControl::FPosQuat ChildTargetTM =
 		UE::PhysicsControl::FPosQuat(ChildTM) * 
 		UE::PhysicsControl::FPosQuat(JointSettings.ConnectorTransforms[ConstraintChildIndex]);
 
-	if (ControlComponent->ParentBodyComponentKey.IsValid())
+	if (Record.ParentBodyComponentKey.IsValid())
 	{
-		const FTransform ParentTM = Hierarchy->GetGlobalTransform(ControlComponent->ParentBodyComponentKey.ElementKey);
+		const FTransform ParentTM = Hierarchy->GetGlobalTransform(Record.ParentBodyComponentKey.ElementKey);
 		const UE::PhysicsControl::FPosQuat ParentTargetTM =
 			UE::PhysicsControl::FPosQuat(ParentTM) *
 			UE::PhysicsControl::FPosQuat(JointSettings.ConnectorTransforms[ConstraintParentIndex]);
@@ -973,122 +1008,141 @@ void FRigPhysicsSimulation::UpdatePrePhysics(
 		const FRigComponentKey& ComponentKey = BodyRecordPair.Key;
 		FRigBodyRecord& Record = BodyRecordPair.Value;
 
-		if (const FRigPhysicsComponent* PhysicsComponent = Cast<FRigPhysicsComponent>(
+		if (const FRigPhysicsBodyComponent* PhysicsComponent = Cast<FRigPhysicsBodyComponent>(
 			Hierarchy->FindComponent(ComponentKey)))
 		{
-			if (Record.ActorHandle)
-			{
-				UpdateBodyPrePhysics(ExecuteContext, SolverComponent, Record, PhysicsComponent, SimulationSpaceData);
+			UpdateBodyPrePhysics(ExecuteContext, SolverComponent, Record, PhysicsComponent, SimulationSpaceData);
+		}
+	}
 
-				UpdateJointPrePhysics(Record, PhysicsComponent, Hierarchy, DeltaTime);
-			}
+	for (TPair<FRigComponentKey, FRigJointRecord>& JointRecordPair : JointRecords)
+	{
+		const FRigComponentKey& ComponentKey = JointRecordPair.Key;
+		FRigJointRecord& JointRecord = JointRecordPair.Value;
+
+		if (const FRigPhysicsJointComponent* JointComponent = Cast<FRigPhysicsJointComponent>(
+			Hierarchy->FindComponent(ComponentKey)))
+		{
+			UpdateJointPrePhysics(JointRecord, JointComponent, Hierarchy, DeltaTime);
 		}
 	}
 
 	for (TPair<FRigComponentKey, FRigControlRecord>& ControlRecordPair : ControlRecords)
 	{
 		const FRigComponentKey& ComponentKey = ControlRecordPair.Key;
-		FRigControlRecord& Record = ControlRecordPair.Value;
+		FRigControlRecord& ControlRecord = ControlRecordPair.Value;
 
 		if (const FRigPhysicsControlComponent* ControlComponent = Cast<FRigPhysicsControlComponent>(
 			Hierarchy->FindComponent(ComponentKey)))
 		{
-			if (ImmediatePhysics::FJointHandle * JointHandle = Record.JointHandle)
-			{
-				Chaos::FPBDJointConstraintHandle* Constraint = JointHandle->GetConstraint();
-				if (Constraint)
-				{
-					float ThisDeltaTime = DeltaTime;
+			UpdateControlPrePhysics(ControlRecord, ControlComponent, Hierarchy, DeltaTime);
+		}
+	}
+}
 
-					if (Record.PreviousTargetUpdateCounter + 1 != UpdateCounter)
-					{
-						// If we missed some intermediate updates, then we don't want to use the previous
-						// positions etc to calculate velocities. This will mean velocity/damping will be
-						// incorrect for one frame, but that's probably OK.
-						ThisDeltaTime = 0.0f;
-					}
+//======================================================================================================================
+void FRigPhysicsSimulation::UpdateControlPrePhysics(
+	FRigControlRecord&                 ControlRecord,
+	const FRigPhysicsControlComponent* ControlComponent,
+	const URigHierarchy*               Hierarchy,
+	const float                        DeltaTime)
+{
+	ImmediatePhysics::FJointHandle* JointHandle = ControlRecord.JointHandle;
+	if (!JointHandle)
+	{
+		return;
+	}
+	Chaos::FPBDJointConstraintHandle* Constraint = JointHandle->GetConstraint();
+	if (!Constraint)
+	{
+		return;
+	}
+	float ThisDeltaTime = DeltaTime;
+
+	if (ControlRecord.PreviousTargetUpdateCounter + 1 != UpdateCounter)
+	{
+		// If we missed some intermediate updates, then we don't want to use the previous
+		// positions etc to calculate velocities. This will mean velocity/damping will be
+		// incorrect for one frame, but that's probably OK.
+		ThisDeltaTime = 0.0f;
+	}
 
-					Constraint->SetCollisionEnabled(!ControlComponent->ControlData.bDisableCollision);
-					Constraint->SetParentInvMassScale(ControlComponent->ControlData.bOnlyControlChildObject ? 0 : 1);
+	Constraint->SetCollisionEnabled(!ControlComponent->ControlData.bDisableCollision);
+	Constraint->SetParentInvMassScale(ControlComponent->ControlData.bOnlyControlChildObject ? 0 : 1);
 
-					const Chaos::FPBDJointSettings& JointSettings = Constraint->GetSettings();
-					if (UpdateDriveSpringDamperSettings(
-						JointHandle, JointSettings, ControlComponent->ControlData, ControlComponent->ControlMultiplier))
-					{
-						const ImmediatePhysics::FActorHandle* ChildActorHandle = 
-							JointHandle->GetActorHandles()[ConstraintChildIndex];
-						const ImmediatePhysics::FActorHandle* ParentActorHandle = 
-							JointHandle->GetActorHandles()[ConstraintParentIndex];
+	const Chaos::FPBDJointSettings& JointSettings = Constraint->GetSettings();
+	if (UpdateDriveSpringDamperSettings(
+		JointHandle, JointSettings, ControlComponent->ControlData, ControlComponent->ControlMultiplier))
+	{
+		const ImmediatePhysics::FActorHandle* ChildActorHandle =
+			JointHandle->GetActorHandles()[ConstraintChildIndex];
+		const ImmediatePhysics::FActorHandle* ParentActorHandle =
+			JointHandle->GetActorHandles()[ConstraintParentIndex];
 
-						if (ChildActorHandle && ParentActorHandle)
-						{
-							// TODO
-							// - cache settings / previous input parameters to avoid unnecessary repeating
-							//   calculations and making physics API calls every update.
-
-							// Update the target point on the child
-							Constraint->SetChildConnectorLocation(
-								ControlComponent->ControlData.GetControlPoint(ChildActorHandle));
-
-							UE::PhysicsControl::FPosQuat TargetTM(
-								ControlComponent->ControlTarget.TargetOrientation,
-								ControlComponent->ControlTarget.TargetPosition);
-
-							if (ControlComponent->ControlData.bUseSkeletalAnimation)
-							{
-								UE::PhysicsControl::FPosQuat AnimTargetTM = CalculateTargetTM(
-									Hierarchy, JointSettings, ControlComponent);
-								TargetTM = AnimTargetTM * TargetTM;
-							}
-
-							Constraint->SetLinearDrivePositionTarget(TargetTM.GetTranslation());
-							Constraint->SetAngularDrivePositionTarget(TargetTM.GetRotation());
-
-							if ((ThisDeltaTime * ControlComponent->ControlData.LinearTargetVelocityMultiplier) != 0)
-							{
-								FVector Velocity =
-									(TargetTM.GetTranslation() - Record.PreviousTargetTM.GetTranslation()) / ThisDeltaTime;
-								Constraint->SetLinearDriveVelocityTarget(
-									Velocity * ControlComponent->ControlData.LinearTargetVelocityMultiplier);
-							}
-							else
-							{
-								Constraint->SetLinearDriveVelocityTarget(Chaos::FVec3(0));
-							}
-
-							if ((ThisDeltaTime * ControlComponent->ControlData.AngularTargetVelocityMultiplier) != 0)
-							{
-								// Note that quats multiply in the opposite order to TMs, and must be in the same hemisphere.
-								const FQuat Q = TargetTM.GetRotation();
-								FQuat PrevQ = Record.PreviousTargetTM.GetRotation();
-								PrevQ.EnforceShortestArcWith(Q);
-								const FQuat DeltaQ = Q * PrevQ.Inverse();
-								const FVector AngularVelocity = DeltaQ.ToRotationVector() / ThisDeltaTime;
-
-								Constraint->SetAngularDriveVelocityTarget(
-									AngularVelocity * ControlComponent->ControlData.AngularTargetVelocityMultiplier);
-							}
-							else
-							{
-								Constraint->SetAngularDriveVelocityTarget(Chaos::FVec3(0));
-							}
-
-
-							Record.PreviousTargetTM = TargetTM;
-							Record.PreviousTargetUpdateCounter = UpdateCounter;
-						}
-						else
-						{
-							// Note that if we don't have any strength, then we don't calculate the targets.
-							// However, make sure that we don't apply velocities using the wrong calculation
-							// when the strength/damping is increased in the future
-						}
-					}
-				}
+		if (ChildActorHandle && ParentActorHandle)
+		{
+			// TODO
+			// - cache settings / previous input parameters to avoid unnecessary repeating
+			//   calculations and making physics API calls every update.
+
+			// Update the target point on the child
+			Constraint->SetChildConnectorLocation(
+				ControlComponent->ControlData.GetControlPoint(ChildActorHandle));
+
+			UE::PhysicsControl::FPosQuat TargetTM(
+				ControlComponent->ControlTarget.TargetOrientation,
+				ControlComponent->ControlTarget.TargetPosition);
+
+			if (ControlComponent->ControlData.bUseSkeletalAnimation)
+			{
+				UE::PhysicsControl::FPosQuat AnimTargetTM = CalculateTargetTM(
+					Hierarchy, JointSettings, ControlRecord);
+				TargetTM = AnimTargetTM * TargetTM;
+			}
+
+			Constraint->SetLinearDrivePositionTarget(TargetTM.GetTranslation());
+			Constraint->SetAngularDrivePositionTarget(TargetTM.GetRotation());
+
+			if ((ThisDeltaTime * ControlComponent->ControlData.LinearTargetVelocityMultiplier) != 0)
+			{
+				FVector Velocity =
+					(TargetTM.GetTranslation() - ControlRecord.PreviousTargetTM.GetTranslation()) / ThisDeltaTime;
+				Constraint->SetLinearDriveVelocityTarget(
+					Velocity * ControlComponent->ControlData.LinearTargetVelocityMultiplier);
+			}
+			else
+			{
+				Constraint->SetLinearDriveVelocityTarget(Chaos::FVec3(0));
+			}
+
+			if ((ThisDeltaTime * ControlComponent->ControlData.AngularTargetVelocityMultiplier) != 0)
+			{
+				// Note that quats multiply in the opposite order to TMs, and must be in the same hemisphere.
+				const FQuat Q = TargetTM.GetRotation();
+				FQuat PrevQ = ControlRecord.PreviousTargetTM.GetRotation();
+				PrevQ.EnforceShortestArcWith(Q);
+				const FQuat DeltaQ = Q * PrevQ.Inverse();
+				const FVector AngularVelocity = DeltaQ.ToRotationVector() / ThisDeltaTime;
+
+				Constraint->SetAngularDriveVelocityTarget(
+					AngularVelocity * ControlComponent->ControlData.AngularTargetVelocityMultiplier);
 			}
+			else
+			{
+				Constraint->SetAngularDriveVelocityTarget(Chaos::FVec3(0));
+			}
+
+
+			ControlRecord.PreviousTargetTM = TargetTM;
+			ControlRecord.PreviousTargetUpdateCounter = UpdateCounter;
+		}
+		else
+		{
+			// Note that if we don't have any strength, then we don't calculate the targets.
+			// However, make sure that we don't apply velocities using the wrong calculation
+			// when the strength/damping is increased in the future
 		}
 	}
-
 }
 
 //======================================================================================================================
@@ -1113,7 +1167,7 @@ void FRigPhysicsSimulation::UpdatePostPhysics(const FRigPhysicsSolverComponent*
 		const FRigBodyRecord& Record = *BodyRecords.Find(ComponentKey);
 		if (Record.ActorHandle)
 		{
-			if (const FRigPhysicsComponent* PhysicsComponent = Cast<FRigPhysicsComponent>(
+			if (const FRigPhysicsBodyComponent* PhysicsComponent = Cast<FRigPhysicsBodyComponent>(
 				Hierarchy->FindComponent(ComponentKey)))
 			{
 				// We will only have created this body record if it's associated with our solver, so
@@ -1169,7 +1223,7 @@ void FRigPhysicsSimulation::UpdatePostPhysics(const FRigPhysicsSolverComponent*
 			const FRigComponentKey& ComponentKey = BodyRecordPair.Key;
 			const FRigBodyRecord& Record = BodyRecordPair.Value;
 
-			if (const FRigPhysicsComponent* PhysicsComponent = Cast<FRigPhysicsComponent>(
+			if (const FRigPhysicsBodyComponent* PhysicsComponent = Cast<FRigPhysicsBodyComponent>(
 				Hierarchy->FindComponent(ComponentKey)))
 			{
 				if (Record.ActorHandle && !Record.ActorHandle->GetIsKinematic())
diff --git a/Engine/Plugins/Experimental/PhysicsControl/Source/ControlRigPhysics/Private/RigPhysicsSimulation.h b/Engine/Plugins/Experimental/PhysicsControl/Source/ControlRigPhysics/Private/RigPhysicsSimulation.h
--- a/Engine/Plugins/Experimental/PhysicsControl/Source/ControlRigPhysics/Private/RigPhysicsSimulation.h
+++ b/Engine/Plugins/Experimental/PhysicsControl/Source/ControlRigPhysics/Private/RigPhysicsSimulation.h
@@ -22,7 +22,9 @@ struct FRigPhysicsSolverComponent;
 
 struct FRigPhysicsCollision;
 struct FRigPhysicsDynamics;
-struct FRigPhysicsComponent;
+struct FRigPhysicsBodyComponent;
+struct FRigPhysicsControlComponent;
+struct FRigPhysicsJointComponent;
 
 #if UE_BUILD_SHIPPING || UE_BUILD_TEST
 DECLARE_LOG_CATEGORY_EXTERN(LogRigPhysics, Log, Warning);
@@ -30,6 +32,29 @@ DECLARE_LOG_CATEGORY_EXTERN(LogRigPhysics, Log, Warning);
 DECLARE_LOG_CATEGORY_EXTERN(LogRigPhysics, Log, All);
 #endif
 
+//======================================================================================================================
+// For internal use only
+//======================================================================================================================
+typedef TPair<FRigComponentKey, FRigComponentKey> FRigPhysicsIgnorePair;
+inline bool operator<(const FRigPhysicsIgnorePair& A, const FRigPhysicsIgnorePair& B)
+{
+	if (A.Key < B.Key)
+	{
+		return true;
+	}
+	else if (A.Key > B.Key)
+	{
+		return false;
+	}
+	// keys are the same...
+	if (A.Value < B.Value)
+	{
+		return true;
+	}
+	return false;
+}
+typedef TSet<FRigPhysicsIgnorePair> FRigPhysicsIgnorePairs;
+
 //======================================================================================================================
 // For internal use only, to keep track of objects we have created
 //======================================================================================================================
@@ -40,10 +65,25 @@ struct FRigBodyRecord
 
 	// Things that are set during instantiation
 	ImmediatePhysics::FActorHandle* ActorHandle = nullptr;
+
+	// Things that are updated as the simulation progresses
+};
+
+//======================================================================================================================
+// For internal use only, to keep track of joints we have created
+//======================================================================================================================
+USTRUCT()
+struct FRigJointRecord
+{
+	GENERATED_BODY()
+
+	// Things that are set during instantiation
 	ImmediatePhysics::FJointHandle* JointHandle = nullptr;
-	// When we have a joint, this is the key for the parent body. Note that the child body component
-	// key is stored as the lookup key for the records, so we don't need to duplicate it here.
+
+	// These keys are filled in when the record is created, even if the original key is set to pick
+	// up the components automatically.
 	FRigComponentKey                ParentBodyComponentKey;
+	FRigComponentKey                ChildBodyComponentKey;
 
 	// Things that are updated as the simulation progresses
 
@@ -64,9 +104,10 @@ struct FRigControlRecord
 	// Things that are set during instantiation
 	ImmediatePhysics::FJointHandle* JointHandle = nullptr;
 
-	// When we have a joint, this is the key for the parent body. Note that the child body component
-	// key is stored as the lookup key for the records, so we don't need to duplicate it here.
+	// These keys are filled in when the record is created, even if the original key is set to pick
+	// up the components automatically.
 	FRigComponentKey                ParentBodyComponentKey;
+	FRigComponentKey                ChildBodyComponentKey;
 
 	// Things that are updated as the simulation progresses
 
@@ -148,20 +189,31 @@ private:
 	// This collects all the bodies associated with the solver and makes records for them
 	void InitialiseBodyRecords(const FRigPhysicsSolverComponent* SolverComponent);
 
+	// This collects all the joints associated with the solver and makes records for them
+	void InitialiseJointRecords(const FRigPhysicsSolverComponent* SolverComponent);
+
 	// This collects all the bodies associated with the solver and makes records for them
 	void InitialiseControlRecords(const FRigPhysicsSolverComponent* SolverComponent);
 
-	// Creates collision shapes associated with the solver
-	void InstantiateSolverCollision(const FRigPhysicsSolverComponent* SolverComponent);
+	// Creates low level bodies
+	void InstantiatePhysicsBodies(
+		const FRigPhysicsSolverComponent* SolverComponent,
+		FRigPhysicsIgnorePairs&           IgnorePairs);
 
 	// Creates low level physics joints 
-	void InstantiatePhysicsJoints(const FRigPhysicsSolverComponent* SolverComponent);
-
-	// Creates low level bodies
-	void InstantiatePhysicsBodies(const FRigPhysicsSolverComponent* SolverComponent);
+	void InstantiatePhysicsJoints(
+		const FRigPhysicsSolverComponent* SolverComponent,
+		FRigPhysicsIgnorePairs&           IgnorePairs);
 
 	// Creates low level controls
-	void InstantiateControls(const FRigPhysicsSolverComponent* SolverComponent);
+	void InstantiateControls(
+		const FRigPhysicsSolverComponent* SolverComponent,
+		FRigPhysicsIgnorePairs&           IgnorePairs);
+
+	// Creates collision shapes associated with the solver. Also applies IgnorePairs
+	void InstantiateSolverCollision(
+		const FRigPhysicsSolverComponent* SolverComponent,
+		FRigPhysicsIgnorePairs&           IgnorePairs);
 
 	// Returns true if the component is physics, and its solver matches the solver component
 	// (directly, or automatically)
@@ -185,20 +237,26 @@ private:
 		const float                       DeltaTime, 
 		const FSimulationSpaceData&       SimulationSpaceData);
 
+	void UpdateBodyPrePhysics(
+		const FRigVMExecuteContext&       ExecuteContext,
+		const FRigPhysicsSolverComponent* SolverComponent,
+		const FRigBodyRecord&             Record, 
+		const FRigPhysicsBodyComponent*       PhysicsComponent,
+		const FSimulationSpaceData&       SimulationSpaceData);
+
 	void UpdateJointPrePhysics(
-		FRigBodyRecord&             Record,
-		const FRigPhysicsComponent* PhysicsComponent, 
-		const URigHierarchy*        Hierarchy, 
-		const float                 DeltaTime);
+		FRigJointRecord&                 JointRecord,
+		const FRigPhysicsJointComponent* JointComponent, 
+		const URigHierarchy*             Hierarchy, 
+		const float                      DeltaTime);
 
-	void UpdatePostPhysics(const FRigPhysicsSolverComponent* SolverComponent, const float Alpha);
+	void UpdateControlPrePhysics(
+		FRigControlRecord&                 ControlRecord,
+		const FRigPhysicsControlComponent* ControlComponent,
+		const URigHierarchy*               Hierarchy,
+		const float                        DeltaTime);
 
-	void UpdateBodyPrePhysics(
-		const FRigVMExecuteContext& ExecuteContext,
-		const FRigPhysicsSolverComponent* SolverComponent,
-		const FRigBodyRecord& Record, 
-		const FRigPhysicsComponent* PhysicsComponent,
-		const FSimulationSpaceData& SimulationSpaceData);
+	void UpdatePostPhysics(const FRigPhysicsSolverComponent* SolverComponent, const float Alpha);
 
 	void InitSimulationSpace(const FTransform& ComponentToWorld, const FTransform& BoneToComponent);
 
@@ -228,6 +286,10 @@ private:
 	ImmediatePhysics::FJointHandle* CreateConstraint(
 		ImmediatePhysics::FActorHandle* ChildActorHandle, ImmediatePhysics::FActorHandle* ParentActorHandle);
 
+	// Gets the simulation actor handle for a component key. Note that the component key could be a
+	// body or a solver component. Can return nullptr
+	ImmediatePhysics::FActorHandle* GetActor(const FRigComponentKey& ComponentKey) const;
+
 private:
 	TObjectPtr<UControlRig> OwningControlRig;
 	FRigComponentKey        PhysicsSolverComponentKey;
@@ -238,13 +300,16 @@ private:
 	// Ordering so that we can traverse from root to leaf bones
 	TArray<FRigComponentKey> SortedBodyComponentKeys;
 
+	// All the joints
+	TMap<FRigComponentKey, FRigJointRecord> JointRecords;
+
 	// All the controls
 	TMap<FRigComponentKey, FRigControlRecord> ControlRecords;
 
 	TSharedPtr<ImmediatePhysics::FSimulation> Simulation;
 
-	// Used to store things like collision shapes. May be offset from the origin if collision is in
-	// a different space to the simulation.
+	// Used to store things the simulation collision shape. May be offset from the origin if
+	// collision is in a different space to the simulation.
 	ImmediatePhysics::FActorHandle* CollisionActorHandle = nullptr;
 
 	// Used to make controls when they're not attached to another simulated body. Will always be at
diff --git a/Engine/Plugins/Experimental/PhysicsControl/Source/ControlRigPhysics/Private/RigPhysicsSimulationVisualization.cpp b/Engine/Plugins/Experimental/PhysicsControl/Source/ControlRigPhysics/Private/RigPhysicsSimulationVisualization.cpp
--- a/Engine/Plugins/Experimental/PhysicsControl/Source/ControlRigPhysics/Private/RigPhysicsSimulationVisualization.cpp
+++ b/Engine/Plugins/Experimental/PhysicsControl/Source/ControlRigPhysics/Private/RigPhysicsSimulationVisualization.cpp
@@ -336,16 +336,20 @@ void FRigPhysicsSimulation::Draw(
 	{
 		const FRigBodyRecord& Record = BodyRecordPair.Value;
 		const ImmediatePhysics::FActorHandle* ActorHandle = Record.ActorHandle;
-		const ImmediatePhysics::FJointHandle* JointHandle = Record.JointHandle;
 
-		if (JointHandle)
+		if (ActorHandle)
 		{
-			DrawJoint(DI, VisualizationSettings, SpaceTransform, *JointHandle);
+			DrawActor(DI, VisualizationSettings, SpaceTransform, *ActorHandle);
 		}
+	}
+	for (const TPair<FRigComponentKey, FRigJointRecord>& JointRecordPair : JointRecords)
+	{
+		const FRigJointRecord& Record = JointRecordPair.Value;
+		const ImmediatePhysics::FJointHandle* JointHandle = Record.JointHandle;
 
-		if (ActorHandle)
+		if (JointHandle)
 		{
-			DrawActor(DI, VisualizationSettings, SpaceTransform, *ActorHandle);
+			DrawJoint(DI, VisualizationSettings, SpaceTransform, *JointHandle);
 		}
 	}
 }
diff --git a/Engine/Plugins/Experimental/PhysicsControl/Source/ControlRigPhysics/Public/RigPhysicsBodyComponent.h b/Engine/Plugins/Experimental/PhysicsControl/Source/ControlRigPhysics/Public/RigPhysicsBodyComponent.h
new file mode 100644
--- /dev/null
+++ b/Engine/Plugins/Experimental/PhysicsControl/Source/ControlRigPhysics/Public/RigPhysicsBodyComponent.h
@@ -0,0 +1,72 @@
+// Copyright Epic Games, Inc. All Rights Reserved.
+
+#pragma once
+
+#include "RigPhysicsData.h"
+
+#include "PhysicsControlData.h"
+
+#include "Rigs/RigPhysics.h"
+#include "Rigs/RigHierarchyDefines.h"
+
+#include "RigPhysicsBodyComponent.generated.h"
+
+/**
+ * A component that can be added to a joint/element that defines how a physical body can be "attached" to it.
+ * The body supports dynamic movement, collision, and a physics joint with this body's parent in the hierarchy.
+ */
+USTRUCT(BlueprintType)
+struct CONTROLRIGPHYSICS_API FRigPhysicsBodyComponent : public FRigBaseComponent
+{
+public:
+
+	GENERATED_BODY()
+	DECLARE_RIG_COMPONENT_METHODS(FRigPhysicsBodyComponent)
+
+	UPROPERTY(BlueprintReadOnly, EditAnywhere, Category = Physics, meta = (ShowOnlyInnerProperties))
+	FRigPhysicsBodySolverSettings BodySolverSettings;
+
+	UPROPERTY(BlueprintReadOnly, EditAnywhere, Category = Physics, meta = (ShowOnlyInnerProperties))
+	FRigPhysicsDynamics Dynamics;
+
+	UPROPERTY(BlueprintReadOnly, EditAnywhere, Category = Physics, meta = (ShowOnlyInnerProperties))
+	FRigPhysicsCollision Collision;
+
+	UPROPERTY(BlueprintReadOnly, EditAnywhere, Category = Physics)
+	FPhysicsControlModifierData BodyData;
+
+	// The target for when this body is kinematic
+	UPROPERTY(BlueprintReadOnly, EditAnywhere, Category = Physics)
+	FTransform KinematicTarget;
+
+	UPROPERTY(BlueprintReadOnly, EditAnywhere, Category = Physics)
+	ERigPhysicsKinematicTargetSpace KinematicTargetSpace = ERigPhysicsKinematicTargetSpace::OffsetInBoneSpace;
+
+	// A list of body components with which we should not collide. The solver component can also be included.
+	UPROPERTY(BlueprintReadOnly, EditAnywhere, Category = Physics)
+	TArray<FRigComponentKey> NoCollisionBodies;
+
+	// Removes any existing collision, and replaces it with a shape calculated from the joint
+	// positions (if possible). The shape will be a single box.
+	// @param MinAspectRatio the minimum box extent, as a proportion of the maximum box extent.
+	void AutoCalculateCollision(URigHierarchy* InHierarchy, float MinAspectRatio = 0.25f);
+
+	virtual void Save(FArchive& Ar) override;
+	virtual void Load(FArchive& Ar) override;
+
+#if WITH_EDITOR
+	virtual const FSlateIcon& GetIconForUI() const override;
+#endif
+
+	virtual bool CanBeAddedTo(
+		const FRigElementKey& InElementKey, const URigHierarchy* InHierarchy, FString* OutFailureReason) const override;
+
+	virtual void OnAddedToHierarchy(URigHierarchy* InHierarchy, URigHierarchyController* InController) override;
+
+	virtual void OnRigHierarchyKeyChanged(const FRigHierarchyKey& InOldKey, const FRigHierarchyKey& InNewKey) override;
+
+	virtual FName GetDefaultComponentName() const override { return GetDefaultName(); }
+
+	static FName GetDefaultName() { return TEXT("PhysicsBody"); }
+};
+
diff --git a/Engine/Plugins/Experimental/PhysicsControl/Source/ControlRigPhysics/Public/RigPhysicsBodyExecution.h b/Engine/Plugins/Experimental/PhysicsControl/Source/ControlRigPhysics/Public/RigPhysicsBodyExecution.h
new file mode 100644
--- /dev/null
+++ b/Engine/Plugins/Experimental/PhysicsControl/Source/ControlRigPhysics/Public/RigPhysicsBodyExecution.h
@@ -0,0 +1,416 @@
+// Copyright Epic Games, Inc. All Rights Reserved.
+
+#pragma once
+
+#include "RigPhysicsBodyComponent.h"
+#include "RigPhysicsSolverComponent.h"
+#include "RigPhysicsSimulation.h"
+#include "PhysicsControlData.h"
+
+#include "Units/Execution/RigUnit_DynamicHierarchy.h"
+
+#include "RigPhysicsBodyExecution.generated.h"
+
+/**
+ * Base struct for all other mutable physics nodes
+ */
+USTRUCT(meta = (Category = "RigPhysics", NodeColor = "1.0 0.6 0.3", Keywords = "Physics"))
+struct CONTROLRIGPHYSICS_API FRigUnit_PhysicsBodyBase: public FRigUnitMutable
+{
+	GENERATED_BODY()
+};
+
+/**
+ * Adds a new physics body as a component on the owner element.
+ * Note: This node only runs as part of the construction event.
+ */
+USTRUCT(meta = (DisplayName = "Add Physics Body Component", Keywords = "Construction,Create,New,Body,Skeleton", Varying))
+struct CONTROLRIGPHYSICS_API FRigUnit_AddPhysicsBody : public FRigUnit_PhysicsBodyBase
+{
+	GENERATED_BODY()
+
+	FRigUnit_AddPhysicsBody()
+	{
+		Owner.Type = ERigElementType::Bone;
+		Solver.PhysicsSolverComponentKey.ElementKey.Type = ERigElementType::Bone;
+		Solver.PhysicsSolverComponentKey.Name = FRigPhysicsBodyComponent::GetDefaultName();
+	}
+
+	RIGVM_METHOD()
+	virtual void Execute() override;
+
+	/**
+	 * The owner of the newly created component (must be set/valid)
+	 */
+	UPROPERTY(meta = (Input, BoneName))
+	FRigElementKey Owner;
+
+	UPROPERTY(meta = (Output))
+	FRigComponentKey PhysicsBodyComponentKey;
+
+	/*
+	 * The solver to relate this new physics element to
+	 */
+	UPROPERTY(meta = (Input))
+	FRigPhysicsBodySolverSettings Solver;
+
+	// The dynamics properties of the new physics element	
+	UPROPERTY(meta = (Input))
+	FRigPhysicsDynamics Dynamics;
+
+	// The collision properties of the new physics element
+	UPROPERTY(meta = (Input))
+	FRigPhysicsCollision Collision;
+
+	// The runtime modifiable data
+	UPROPERTY(meta = (Input))
+	FPhysicsControlModifierData BodyData;
+};
+
+// Discards any existing collision data and replaces it with a box based on the joint positions.
+// Note that this must be called before the physics solver is instantiated/stepped.
+USTRUCT(meta = (DisplayName = "Automatically Calculate Collision", Varying))
+struct CONTROLRIGPHYSICS_API FRigUnit_HierarchyAutoCalculateCollision : public FRigUnit_PhysicsBodyBase
+{
+	GENERATED_BODY()
+
+	FRigUnit_HierarchyAutoCalculateCollision()
+	{
+		PhysicsBodyComponentKey.ElementKey.Type = ERigElementType::Bone;
+		PhysicsBodyComponentKey.Name = FRigPhysicsBodyComponent::GetDefaultName();
+	}
+
+	RIGVM_METHOD()
+	virtual void Execute() override;
+
+	UPROPERTY(meta = (Input))
+	FRigComponentKey PhysicsBodyComponentKey;
+
+	// The minimum box extent, as a proportion of the maximum box extent.
+	UPROPERTY(meta = (Input))
+	float MinAspectRatio = 0.25f;
+};
+
+// Sets the mass etc for a physics component body
+USTRUCT(meta = (DisplayName = "Set Physics Body Dynamics Properties", Varying))
+struct CONTROLRIGPHYSICS_API FRigUnit_HierarchySetDynamics : public FRigUnit_PhysicsBodyBase
+{
+	GENERATED_BODY()
+
+	FRigUnit_HierarchySetDynamics()
+	{
+		PhysicsBodyComponentKey.ElementKey.Type = ERigElementType::Bone;
+		PhysicsBodyComponentKey.Name = FRigPhysicsBodyComponent::GetDefaultName();
+	}
+
+	RIGVM_METHOD()
+	virtual void Execute() override;
+
+	UPROPERTY(meta = (Input))
+	FRigComponentKey PhysicsBodyComponentKey;
+
+	UPROPERTY(meta = (Input))
+	FRigPhysicsDynamics Dynamics;
+};
+
+// Sets the collision for a physics component body
+USTRUCT(meta = (DisplayName = "Set Physics Body Collision Properties", Varying))
+struct CONTROLRIGPHYSICS_API FRigUnit_HierarchySetCollision : public FRigUnit_PhysicsBodyBase
+{
+	GENERATED_BODY()
+
+	FRigUnit_HierarchySetCollision()
+	{
+		PhysicsBodyComponentKey.ElementKey.Type = ERigElementType::Bone;
+		PhysicsBodyComponentKey.Name = FRigPhysicsBodyComponent::GetDefaultName();
+	}
+
+	RIGVM_METHOD()
+	virtual void Execute() override;
+
+	UPROPERTY(meta = (Input))
+	FRigComponentKey PhysicsBodyComponentKey;
+
+	UPROPERTY(meta = (Input))
+	FRigPhysicsCollision Collision;
+};
+
+
+// Disables collision between two bodies
+USTRUCT(meta = (DisplayName = "Disable Collision Between", Varying))
+struct CONTROLRIGPHYSICS_API FRigUnit_HierarchyDisableCollisionBetween : public FRigUnit_PhysicsBodyBase
+{
+	GENERATED_BODY()
+
+	FRigUnit_HierarchyDisableCollisionBetween()
+	{
+		PhysicsBodyComponentKey1.ElementKey.Type = ERigElementType::Bone;
+		PhysicsBodyComponentKey1.Name = FRigPhysicsBodyComponent::GetDefaultName();
+		PhysicsBodyComponentKey2.ElementKey.Type = ERigElementType::Bone;
+		PhysicsBodyComponentKey2.Name = FRigPhysicsBodyComponent::GetDefaultName();
+	}
+
+	RIGVM_METHOD()
+	virtual void Execute() override;
+
+	UPROPERTY(meta = (Input))
+	FRigComponentKey PhysicsBodyComponentKey1;
+
+	UPROPERTY(meta = (Input))
+	FRigComponentKey PhysicsBodyComponentKey2;
+};
+
+/**
+ * Sets what bone is used as a source transform for the physics body. This is used as a kinematic target, and when
+ * initializing the simulation.
+ */
+USTRUCT(meta = (DisplayName = "Set Physics Body Source Bone", Varying))
+struct CONTROLRIGPHYSICS_API FRigUnit_HierarchySetPhysicsBodySourceBone : public FRigUnit_PhysicsBodyBase
+{
+	GENERATED_BODY()
+
+	FRigUnit_HierarchySetPhysicsBodySourceBone()
+	{
+		PhysicsBodyComponentKey.ElementKey.Type = ERigElementType::Bone;
+		PhysicsBodyComponentKey.Name = FRigPhysicsBodyComponent::GetDefaultName();
+		SourceBone.Type = ERigElementType::Bone;
+	}
+
+	RIGVM_METHOD()
+	virtual void Execute() override;
+
+	UPROPERTY(meta = (Input))
+	FRigComponentKey PhysicsBodyComponentKey;
+
+	UPROPERTY(meta = (Input))
+	FRigElementKey SourceBone;
+};
+
+/**
+ * Sets what bone is targeted by the simulation - i.e. where the simulation output is written to.
+ */
+USTRUCT(meta = (DisplayName = "Set Physics Body Target Bone", Varying))
+struct CONTROLRIGPHYSICS_API FRigUnit_HierarchySetPhysicsBodyTargetBone : public FRigUnit_PhysicsBodyBase
+{
+	GENERATED_BODY()
+
+	FRigUnit_HierarchySetPhysicsBodyTargetBone()
+	{
+		PhysicsBodyComponentKey.ElementKey.Type = ERigElementType::Bone;
+		PhysicsBodyComponentKey.Name = FRigPhysicsBodyComponent::GetDefaultName();
+		TargetBone.Type = ERigElementType::Bone;
+	}
+
+	RIGVM_METHOD()
+	virtual void Execute() override;
+
+	UPROPERTY(meta = (Input))
+	FRigComponentKey PhysicsBodyComponentKey;
+
+	UPROPERTY(meta = (Input))
+	FRigElementKey TargetBone;
+};
+
+
+/**
+ * Sets all the data on a body - but in a sparse way so you can decide which parameters get applied.
+ * Danny TODO - Note that the sparse data does not get displayed correctly, so this is largely 
+ * unusable - the flags that enable/disable all end up getting reset if the user attempts to change them.
+ */
+USTRUCT(meta = (DisplayName = "Set Physics Body Data", Varying))
+struct CONTROLRIGPHYSICS_API FRigUnit_HierarchySetPhysicsBodySparseData : public FRigUnit_PhysicsBodyBase
+{
+	GENERATED_BODY()
+
+	FRigUnit_HierarchySetPhysicsBodySparseData()
+	{
+		PhysicsBodyComponentKey.ElementKey.Type = ERigElementType::Bone;
+		PhysicsBodyComponentKey.Name = FRigPhysicsBodyComponent::GetDefaultName();
+	}
+
+	RIGVM_METHOD()
+	virtual void Execute() override;
+
+	UPROPERTY(meta = (Input))
+	FRigComponentKey PhysicsBodyComponentKey;
+
+	UPROPERTY(meta = (Input))
+	FPhysicsControlModifierSparseData Data;
+};
+
+/**
+ * Sets the kinematic target for a body - note that this won't actually make the body kinematic
+ */
+USTRUCT(meta = (DisplayName = "Set Physics Body Kinematic Target", Varying))
+struct CONTROLRIGPHYSICS_API FRigUnit_HierarchySetPhysicsBodyKinematicTarget : public FRigUnit_PhysicsBodyBase
+{
+	GENERATED_BODY()
+
+	FRigUnit_HierarchySetPhysicsBodyKinematicTarget()
+	{
+		PhysicsBodyComponentKey.ElementKey.Type = ERigElementType::Bone;
+		PhysicsBodyComponentKey.Name = FRigPhysicsBodyComponent::GetDefaultName();
+	}
+
+	RIGVM_METHOD()
+	virtual void Execute() override;
+
+	UPROPERTY(meta = (Input))
+	FRigComponentKey PhysicsBodyComponentKey;
+
+	UPROPERTY(meta = (Input))
+	ERigPhysicsKinematicTargetSpace KinematicTargetSpace = ERigPhysicsKinematicTargetSpace::OffsetInBoneSpace;
+
+	UPROPERTY(meta = (Input))
+	FTransform KinematicTarget;
+};
+
+/**
+ * Sets the movement type for this body. 
+ * Danny TODO explain the different between kinematic and static
+ */
+USTRUCT(meta = (DisplayName = "Set Physics Body Movement Type", Varying))
+struct CONTROLRIGPHYSICS_API FRigUnit_HierarchySetPhysicsBodyMovementType : public FRigUnit_PhysicsBodyBase
+{
+	GENERATED_BODY()
+
+	FRigUnit_HierarchySetPhysicsBodyMovementType()
+	{
+		PhysicsBodyComponentKey.ElementKey.Type = ERigElementType::Bone;
+		PhysicsBodyComponentKey.Name = FRigPhysicsBodyComponent::GetDefaultName();
+	}
+
+	RIGVM_METHOD()
+	virtual void Execute() override;
+
+	UPROPERTY(meta = (Input))
+	FRigComponentKey PhysicsBodyComponentKey;
+
+	UPROPERTY(meta = (Input))
+	EPhysicsMovementType MovementType = EPhysicsMovementType::Simulated;
+};
+
+/**
+ * Sets what collision is used for this body
+ */
+USTRUCT(meta = (DisplayName = "Set Physics Body Collision Type", Varying))
+struct CONTROLRIGPHYSICS_API FRigUnit_HierarchySetPhysicsBodyCollisionType : public FRigUnit_PhysicsBodyBase
+{
+	GENERATED_BODY()
+
+	FRigUnit_HierarchySetPhysicsBodyCollisionType()
+	{
+		PhysicsBodyComponentKey.ElementKey.Type = ERigElementType::Bone;
+		PhysicsBodyComponentKey.Name = FRigPhysicsBodyComponent::GetDefaultName();
+	}
+
+	RIGVM_METHOD()
+	virtual void Execute() override;
+
+	UPROPERTY(meta = (Input))
+	FRigComponentKey PhysicsBodyComponentKey;
+
+	UPROPERTY(meta = (Input))
+	TEnumAsByte<ECollisionEnabled::Type> CollisionType = ECollisionEnabled::QueryAndPhysics;
+};
+
+/**
+ * Sets the multiplier on gravity that should be applied to the body.
+ */
+USTRUCT(meta = (DisplayName = "Set Physics Body Gravity Multiplier", Varying))
+struct CONTROLRIGPHYSICS_API FRigUnit_HierarchySetPhysicsBodyGravityMultiplier : public FRigUnit_PhysicsBodyBase
+{
+	GENERATED_BODY()
+
+	FRigUnit_HierarchySetPhysicsBodyGravityMultiplier()
+	{
+		PhysicsBodyComponentKey.ElementKey.Type = ERigElementType::Bone;
+		PhysicsBodyComponentKey.Name = FRigPhysicsBodyComponent::GetDefaultName();
+	}
+
+	RIGVM_METHOD()
+	virtual void Execute() override;
+
+	UPROPERTY(meta = (Input))
+	FRigComponentKey PhysicsBodyComponentKey;
+
+	UPROPERTY(meta = (Input))
+	float GravityMultiplier = 1.0f;
+};
+
+/**
+ * Controls the amount that the simulation is blended back into the target bones. 
+ * Danny TODO implement
+ */
+USTRUCT(meta = (DisplayName = "Set Physics Body Physics Blend Weight", Varying))
+struct CONTROLRIGPHYSICS_API FRigUnit_HierarchySetPhysicsBodyPhysicsBlendWeight : public FRigUnit_PhysicsBodyBase
+{
+	GENERATED_BODY()
+
+	FRigUnit_HierarchySetPhysicsBodyPhysicsBlendWeight()
+	{
+		PhysicsBodyComponentKey.ElementKey.Type = ERigElementType::Bone;
+		PhysicsBodyComponentKey.Name = FRigPhysicsBodyComponent::GetDefaultName();
+	}
+
+	RIGVM_METHOD()
+	virtual void Execute() override;
+
+	UPROPERTY(meta = (Input))
+	FRigComponentKey PhysicsBodyComponentKey;
+
+	UPROPERTY(meta = (Input))
+	float PhysicsBlendWeight = 1.0f;
+};
+
+/**
+ * Danny TODO I think this is redundant given the ability to specify the space when kinematic 
+ * targets are set. Investigate propagating that into the rest of Physics Control.
+ */
+USTRUCT(meta = (DisplayName = "Set Physics Body Use Skeletal Animation", Varying))
+struct CONTROLRIGPHYSICS_API FRigUnit_HierarchySetPhysicsBodyUseSkeletalAnimation : public FRigUnit_PhysicsBodyBase
+{
+	GENERATED_BODY()
+
+	FRigUnit_HierarchySetPhysicsBodyUseSkeletalAnimation()
+	{
+		PhysicsBodyComponentKey.ElementKey.Type = ERigElementType::Bone;
+		PhysicsBodyComponentKey.Name = FRigPhysicsBodyComponent::GetDefaultName();
+	}
+
+	RIGVM_METHOD()
+	virtual void Execute() override;
+
+	UPROPERTY(meta = (Input))
+	FRigComponentKey PhysicsBodyComponentKey;
+
+	UPROPERTY(meta = (Input))
+	bool bUseSkeletalAnimation = true;
+};
+
+/**
+ * If true, then kinematic objects will be written back from simulation to the bones. This only 
+ * necessary when either kinematic targets are being used, or when the target bone differs from the source bone.
+ * Danny TODO - check the implementation, and also can we auto-calculate this at runtime?
+ */
+USTRUCT(meta = (DisplayName = "Set Physics Body Update Kinematic From Simulation", Varying))
+struct CONTROLRIGPHYSICS_API FRigUnit_HierarchySetPhysicsBodyUpdateKinematicFromSimulation : public FRigUnit_PhysicsBodyBase
+{
+	GENERATED_BODY()
+
+	FRigUnit_HierarchySetPhysicsBodyUpdateKinematicFromSimulation()
+	{
+		PhysicsBodyComponentKey.ElementKey.Type = ERigElementType::Bone;
+		PhysicsBodyComponentKey.Name = FRigPhysicsBodyComponent::GetDefaultName();
+	}
+
+	RIGVM_METHOD()
+	virtual void Execute() override;
+
+	UPROPERTY(meta = (Input))
+	FRigComponentKey PhysicsBodyComponentKey;
+
+	UPROPERTY(meta = (Input))
+	bool bUpdateKinematicFromSimulation = true;
+};
+
diff --git a/Engine/Plugins/Experimental/PhysicsControl/Source/ControlRigPhysics/Public/RigPhysicsComponent.h b/Engine/Plugins/Experimental/PhysicsControl/Source/ControlRigPhysics/Public/RigPhysicsComponent.h
deleted file mode 100644
diff --git a/Engine/Plugins/Experimental/PhysicsControl/Source/ControlRigPhysics/Public/RigPhysicsControlComponent.h b/Engine/Plugins/Experimental/PhysicsControl/Source/ControlRigPhysics/Public/RigPhysicsControlComponent.h
--- a/Engine/Plugins/Experimental/PhysicsControl/Source/ControlRigPhysics/Public/RigPhysicsControlComponent.h
+++ b/Engine/Plugins/Experimental/PhysicsControl/Source/ControlRigPhysics/Public/RigPhysicsControlComponent.h
@@ -3,6 +3,7 @@
 #pragma once
 
 #include "PhysicsControlData.h"
+#include "RigPhysicsBodyComponent.h"
 
 #include "Rigs/RigPhysics.h"
 
@@ -20,11 +21,25 @@ public:
 	GENERATED_BODY()
 	DECLARE_RIG_COMPONENT_METHODS(FRigPhysicsControlComponent)
 
+	FRigPhysicsControlComponent()
+	{
+		ParentBodyComponentKey.ElementKey.Type = ERigElementType::Bone;
+		ParentBodyComponentKey.Name = FRigPhysicsBodyComponent::GetDefaultName();
+		ChildBodyComponentKey.ElementKey.Type = ERigElementType::Bone;
+		ChildBodyComponentKey.Name = FRigPhysicsBodyComponent::GetDefaultName();
+	}
+
 	// The body that controls the body being controlled. If this is dynamic, it will be affected
 	// too. If unset, then it implies a global control.
 	UPROPERTY(BlueprintReadOnly, EditAnywhere, Category = Control)
 	FRigComponentKey ParentBodyComponentKey;
 
+	// If true, then if the parent body component key is not set, the the default parent body comes
+	// from the parent joint. If it is false, then this search is not done, so the control will be
+	// in simulation space.
+	UPROPERTY(BlueprintReadOnly, EditAnywhere, Category = Control)
+	bool bUseParentBodyAsDefault = false;
+
 	// The body being controlled
 	UPROPERTY(BlueprintReadOnly, EditAnywhere, Category = Control)
 	FRigComponentKey ChildBodyComponentKey;
diff --git a/Engine/Plugins/Experimental/PhysicsControl/Source/ControlRigPhysics/Public/RigPhysicsControlExecution.h b/Engine/Plugins/Experimental/PhysicsControl/Source/ControlRigPhysics/Public/RigPhysicsControlExecution.h
--- a/Engine/Plugins/Experimental/PhysicsControl/Source/ControlRigPhysics/Public/RigPhysicsControlExecution.h
+++ b/Engine/Plugins/Experimental/PhysicsControl/Source/ControlRigPhysics/Public/RigPhysicsControlExecution.h
@@ -13,7 +13,7 @@
  * Base struct for all other mutable physics control nodes
  */
 USTRUCT(meta = (Category = "RigPhysics", NodeColor = "0.8 0.8 0.3", Keywords = "Physics, Control"))
-struct CONTROLRIGPHYSICS_API FRigUnit_PhysicsControlBaseMutable : public FRigUnitMutable
+struct CONTROLRIGPHYSICS_API FRigUnit_PhysicsControlBase : public FRigUnitMutable
 {
 	GENERATED_BODY()
 };
@@ -23,7 +23,7 @@ struct CONTROLRIGPHYSICS_API FRigUnit_PhysicsControlBaseMutable : public FRigUni
  * Note: This node only runs as part of the construction event.
  */
 USTRUCT(meta = (Category = "RigPhysics", NodeColor = "1.0 0.6 0.3", DisplayName = "Add Physics Control Component", Keywords = "Construction,Create,New,Joint,Control,Physics", Varying))
-struct CONTROLRIGPHYSICS_API FRigUnit_AddPhysicsControl : public FRigUnit_PhysicsControlBaseMutable
+struct CONTROLRIGPHYSICS_API FRigUnit_AddPhysicsControl : public FRigUnit_PhysicsControlBase
 {
 	GENERATED_BODY()
 
@@ -62,15 +62,11 @@ struct CONTROLRIGPHYSICS_API FRigUnit_AddPhysicsControl : public FRigUnit_Physic
 	/** Which set to include the control in (optional). Note that it automatically gets added to the set "All" */
 	//UPROPERTY(meta = (Input))
 	//FName Set;
-
-	/** NamePrefix Name of the control component. */
-	UPROPERTY(meta = (Input))
-	FName Name = TEXT("Physics Control");
 };
 
 // Sets the control data for a physics control
 USTRUCT(meta = (DisplayName = "Set Physics Control Data", Varying))
-struct CONTROLRIGPHYSICS_API FRigUnit_HierarchySetControlData : public FRigUnit_PhysicsControlBaseMutable
+struct CONTROLRIGPHYSICS_API FRigUnit_HierarchySetControlData : public FRigUnit_PhysicsControlBase
 {
 	GENERATED_BODY()
 
@@ -92,7 +88,7 @@ struct CONTROLRIGPHYSICS_API FRigUnit_HierarchySetControlData : public FRigUnit_
 
 // Sets the target for a physics control
 USTRUCT(meta = (DisplayName = "Set Physics Control Target", Varying))
-struct CONTROLRIGPHYSICS_API FRigUnit_HierarchySetControlTarget : public FRigUnit_PhysicsControlBaseMutable
+struct CONTROLRIGPHYSICS_API FRigUnit_HierarchySetControlTarget : public FRigUnit_PhysicsControlBase
 {
 	GENERATED_BODY()
 
@@ -114,7 +110,7 @@ struct CONTROLRIGPHYSICS_API FRigUnit_HierarchySetControlTarget : public FRigUni
 
 // Sets the multipliers for a physics control
 USTRUCT(meta = (DisplayName = "Set Physics Control Multiplier", Varying))
-struct CONTROLRIGPHYSICS_API FRigUnit_HierarchySetControlMultiplier : public FRigUnit_PhysicsControlBaseMutable
+struct CONTROLRIGPHYSICS_API FRigUnit_HierarchySetControlMultiplier : public FRigUnit_PhysicsControlBase
 {
 	GENERATED_BODY()
 
diff --git a/Engine/Plugins/Experimental/PhysicsControl/Source/ControlRigPhysics/Public/RigPhysicsData.h b/Engine/Plugins/Experimental/PhysicsControl/Source/ControlRigPhysics/Public/RigPhysicsData.h
--- a/Engine/Plugins/Experimental/PhysicsControl/Source/ControlRigPhysics/Public/RigPhysicsData.h
+++ b/Engine/Plugins/Experimental/PhysicsControl/Source/ControlRigPhysics/Public/RigPhysicsData.h
@@ -7,6 +7,9 @@
 #include "Rigs/RigPhysics.h"
 #include "Rigs/RigHierarchyDefines.h"
 
+#include "PhysicsEngine/ConstraintTypes.h"
+#include "PhysicsEngine/ConstraintDrives.h"
+
 #include "RigPhysicsData.generated.h"
 
 struct FRigVMExecuteContext;
@@ -23,13 +26,13 @@ struct FRigPhysicsVisualizationSettings
 	GENERATED_BODY()
 
 	/* Enable visualization (includes enabling low-level chaos debug draw) */
-	UPROPERTY(BlueprintReadWrite, EditAnywhere, Category = Visualization)
+	UPROPERTY(EditAnywhere, Category = Visualization)
 	bool bEnableVisualization = true;
 
-	UPROPERTY(BlueprintReadWrite, EditAnywhere, Category = Visualization, meta = (ClampMin = "0.0"))
+	UPROPERTY(EditAnywhere, Category = Visualization, meta = (ClampMin = "0.0"))
 	float LineThickness = 1.0f;
 
-	UPROPERTY(BlueprintReadWrite, EditAnywhere, Category = Visualization, meta = (ClampMin = "0.0"))
+	UPROPERTY(EditAnywhere, Category = Visualization, meta = (ClampMin = "0.0"))
 	int32 ShapeDetail = 16;
 };
 
@@ -122,11 +125,11 @@ struct CONTROLRIGPHYSICS_API FRigPhysicsCollisionBox : public FRigPhysicsCollisi
 		return Ar;
 	}
 
-	UPROPERTY(BlueprintReadWrite, EditAnywhere, Category = Shapes)
+	UPROPERTY(EditAnywhere, Category = Shapes)
 	FTransform TM;
 
 	// These are the full extents of the box in each axis
-	UPROPERTY(BlueprintReadWrite, EditAnywhere, Category = Shapes)
+	UPROPERTY(EditAnywhere, Category = Shapes)
 	FVector Extents = FVector::OneVector * RigPhysicsDefaultShapeSize;
 };
 
@@ -150,10 +153,10 @@ struct CONTROLRIGPHYSICS_API FRigPhysicsCollisionSphere : public FRigPhysicsColl
 		return Ar;
 	}
 
-	UPROPERTY(BlueprintReadWrite, EditAnywhere, Category = Shapes)
+	UPROPERTY(EditAnywhere, Category = Shapes)
 	FTransform TM;
 
-	UPROPERTY(BlueprintReadWrite, EditAnywhere, Category = Shapes, meta = (ClampMin = "0.0"))
+	UPROPERTY(EditAnywhere, Category = Shapes, meta = (ClampMin = "0.0"))
 	float Radius = RigPhysicsDefaultShapeSize;
 };
 
@@ -177,14 +180,14 @@ struct CONTROLRIGPHYSICS_API FRigPhysicsCollisionCapsule : public FRigPhysicsCol
 		return Ar;
 	}
 
-	UPROPERTY(BlueprintReadWrite, EditAnywhere, Category = Shapes)
+	UPROPERTY(EditAnywhere, Category = Shapes)
 	FTransform TM;
 
-	UPROPERTY(BlueprintReadWrite, EditAnywhere, Category = Shapes, meta = (ClampMin = "0.0"))
+	UPROPERTY(EditAnywhere, Category = Shapes, meta = (ClampMin = "0.0"))
 	float Radius = RigPhysicsDefaultShapeSize;
 
 	// This is the length of the core part of the capsule. The total length will be Length + 2 * Radius
-	UPROPERTY(BlueprintReadWrite, EditAnywhere, Category = Shapes, meta = (ClampMin = "0.0"))
+	UPROPERTY(EditAnywhere, Category = Shapes, meta = (ClampMin = "0.0"))
 	float Length = RigPhysicsDefaultShapeSize;
 };
 
@@ -198,13 +201,13 @@ struct CONTROLRIGPHYSICS_API FRigPhysicsCollision
 
 	friend FArchive& operator <<(FArchive& Ar, FRigPhysicsCollision& Data);
 
-	UPROPERTY(BlueprintReadWrite, EditAnywhere, Category = Shapes)
+	UPROPERTY(EditAnywhere, Category = Shapes)
 	TArray<FRigPhysicsCollisionBox> Boxes;
 
-	UPROPERTY(BlueprintReadWrite, EditAnywhere, Category = Shapes)
+	UPROPERTY(EditAnywhere, Category = Shapes)
 	TArray<FRigPhysicsCollisionSphere> Spheres;
 
-	UPROPERTY(BlueprintReadWrite, EditAnywhere, Category = Shapes)
+	UPROPERTY(EditAnywhere, Category = Shapes)
 	TArray<FRigPhysicsCollisionCapsule> Capsules;
 
 	bool IsEmpty() const { return Boxes.IsEmpty() && Spheres.IsEmpty() && Capsules.IsEmpty(); }
@@ -223,16 +226,16 @@ struct FRigPhysicsSimulationSpaceSettings
 	friend FArchive& operator <<(FArchive& Ar, FRigPhysicsSimulationSpaceSettings& Data);
 
 	// Overall multiplier on the effects of simulation space movement on the simulation
-	UPROPERTY(EditAnywhere, BlueprintReadWrite, Category = Settings, meta = (ClampMin = "0.0", ClampMax = "1.0"))
+	UPROPERTY(EditAnywhere, Category = Settings, meta = (ClampMin = "0.0", ClampMax = "1.0"))
 	float SpaceMovementAmount = 1.0f;
 
 	// Multiplier on the components of vertical velocity of the simulation space that is passed to the simulation.
 	// Usually from 0.0 to 1.0 to reduce the effects of jumping and crouching on the simulation, but
 	// it can be higher than 1.0 if you need to exaggerate this motion for some reason.
-	UPROPERTY(EditAnywhere, BlueprintReadWrite, Category = Settings)
+	UPROPERTY(EditAnywhere, Category = Settings)
 	float VelocityScaleZ = 1.0f;
 
-	UPROPERTY(EditAnywhere, BlueprintReadWrite, Category = Settings, meta = (InlineEditConditionToggle))
+	UPROPERTY(EditAnywhere, Category = Settings, meta = (InlineEditConditionToggle))
 	bool bClampLinearVelocity = false;
 
 	// A clamp on the effective world-space velocity that is passed to the simulation. The default
@@ -241,10 +244,10 @@ struct FRigPhysicsSimulationSpaceSettings
 	// that have LinearDrag set to non-zero in the physics asset). Expected values in this case
 	// would be somewhat less than the usual velocities of your object which is commonly a few
 	// hundred for a character.
-	UPROPERTY(EditAnywhere, BlueprintReadWrite, Category = Settings, meta = (ClampMin = "0.0", EditCondition = "bClampLinearVelocity"))
+	UPROPERTY(EditAnywhere, Category = Settings, meta = (ClampMin = "0.0", EditCondition = "bClampLinearVelocity"))
 	float MaxLinearVelocity = 10000;
 
-	UPROPERTY(EditAnywhere, BlueprintReadWrite, Category = Settings, meta = (InlineEditConditionToggle))
+	UPROPERTY(EditAnywhere, Category = Settings, meta = (InlineEditConditionToggle))
 	bool bClampAngularVelocity = false;
 
 	// A clamp on the effective world-space angular velocity that is passed to the simulation. Units
@@ -253,10 +256,10 @@ struct FRigPhysicsSimulationSpaceSettings
 	// how much bodies "fly out" when the actor spins on the spot. This is especially useful if you
 	// have characters than can rotate very quickly and you would probably want values around or
 	// less than 10 in this case.
-	UPROPERTY(EditAnywhere, BlueprintReadWrite, Category = Settings, meta = (ClampMin = "0.0", EditCondition = "bClampAngularVelocity"))
+	UPROPERTY(EditAnywhere, Category = Settings, meta = (ClampMin = "0.0", EditCondition = "bClampAngularVelocity"))
 	float MaxAngularVelocity = 10000;
 
-	UPROPERTY(EditAnywhere, BlueprintReadWrite, Category = Settings, meta = (InlineEditConditionToggle))
+	UPROPERTY(EditAnywhere, Category = Settings, meta = (InlineEditConditionToggle))
 	bool bClampLinearAcceleration = false;
 
 	// A clamp on the effective world-space acceleration that is passed to the simulation. The
@@ -264,31 +267,31 @@ struct FRigPhysicsSimulationSpaceSettings
 	// simulation flying out when suddenly changing linear speed. It is useful when you have
 	// characters than can changes from stationary to running very quickly such as in an FPS. A
 	// common value for a character might be in the few hundreds.
-	UPROPERTY(EditAnywhere, BlueprintReadWrite, Category = Settings, meta = (ClampMin = "0.0", EditCondition = "bClampLinearAcceleration"))
+	UPROPERTY(EditAnywhere, Category = Settings, meta = (ClampMin = "0.0", EditCondition = "bClampLinearAcceleration"))
 	float MaxLinearAcceleration = 10000;
 
-	UPROPERTY(EditAnywhere, BlueprintReadWrite, Category = Settings, meta = (InlineEditConditionToggle))
+	UPROPERTY(EditAnywhere, Category = Settings, meta = (InlineEditConditionToggle))
 	bool bClampAngularAcceleration = false;
 
 	// A clamp on the effective world-space angular acceleration that is passed to the simulation.
 	// Units are radian/s/s. The default value effectively means "unlimited". This has a similar
 	// effect to MaxAngularVelocity, except that it is related to the flying out of bodies when the
 	// rotation speed suddenly changes. A typical value for a character might be around 100.
-	UPROPERTY(EditAnywhere, BlueprintReadWrite, Category = Settings, meta = (ClampMin = "0.0", EditCondition = "bClampAngularAcceleration"))
+	UPROPERTY(EditAnywhere, Category = Settings, meta = (ClampMin = "0.0", EditCondition = "bClampAngularAcceleration"))
 	float MaxAngularAcceleration = 10000;
 
 	// If the linear acceleration is above this threshold, the movement is treated as a teleport.
 	// The calculated accelerations and velocities will be set to zero, without resetting the
 	// simulation state.
 	// A value of zero will disable this detection
-	UPROPERTY(EditAnywhere, BlueprintReadWrite, Category = Settings, meta = (ClampMin = "0.0"))
+	UPROPERTY(EditAnywhere, Category = Settings, meta = (ClampMin = "0.0"))
 	float LinearAccelerationThresholdForTeleport = 100000;
 
 	// If the angular acceleration (rand/s/s) is above this threshold, the movement is treated as a teleport.
 	// The calculated accelerations and velocities will be set to zero, without resetting the
 	// simulation state.
 	// A value of zero will disable this detection
-	UPROPERTY(EditAnywhere, BlueprintReadWrite, Category = Settings, meta = (ClampMin = "0.0"))
+	UPROPERTY(EditAnywhere, Category = Settings, meta = (ClampMin = "0.0"))
 	float AngularAccelerationThresholdForTeleport = 100000;
 
 	// Additional linear drag applied to every body in addition to linear drag specified on them in
@@ -297,21 +300,21 @@ struct FRigPhysicsSimulationSpaceSettings
 	// physics asset. The result is that each body has a force equal to -ExternalLinearDragV *
 	// ExternalLinearVelocity applied to it, in addition to all other forces. The vector is in
 	// simulation local space.
-	UPROPERTY(EditAnywhere, BlueprintReadWrite, Category = Settings)
+	UPROPERTY(EditAnywhere, Category = Settings)
 	FVector ExternalLinearDragV = FVector::ZeroVector;
 
 	// Additional velocity that is added to the component velocity so the simulation acts as if the
 	// actor is moving at speed, even when stationary. The vector is in world space. This could be
 	// used for wind effects etc. Typical values are similar to the velocity of the object or
 	// effect, and usually around or less than 1000 for characters/wind.
-	UPROPERTY(EditAnywhere, BlueprintReadWrite, Category = Settings)
+	UPROPERTY(EditAnywhere, Category = Settings)
 	FVector ExternalLinearVelocity = FVector::ZeroVector;
 
 	// Additional angular velocity that is added to the component angular velocity. This can be used
 	// to make the simulation act as if the actor is rotating even when it is not. E.g., to apply
 	// physics to a character on a podium as the camera rotates around it, to emulate the podium
 	// itself rotating. Vector is in world space. Units are rad/s.
-	UPROPERTY(EditAnywhere, BlueprintReadWrite, Category = Settings)
+	UPROPERTY(EditAnywhere, Category = Settings)
 	FVector ExternalAngularVelocity = FVector::ZeroVector;
 };
 
@@ -327,7 +330,7 @@ struct CONTROLRIGPHYSICS_API FRigPhysicsSolverSettings
 
 	// If true, then any physics component below (the element owning this one) will be automatically
 	// added to this solver, if its UseAutomaticSolver flag is set.
-	UPROPERTY(BlueprintReadWrite, EditAnywhere, Category = SolverSettings)
+	UPROPERTY(EditAnywhere, Category = SolverSettings)
 	bool bAutomaticallyAddPhysicsComponents = true;
 
 	UPROPERTY(BlueprintReadOnly, EditAnywhere, Category = SolverSettings)
@@ -345,21 +348,21 @@ struct CONTROLRIGPHYSICS_API FRigPhysicsSolverSettings
 	UPROPERTY(BlueprintReadOnly, EditAnywhere, Category = SolverSettings, meta = (ShowOnlyInnerProperties))
 	FRigPhysicsCollision Collision;
 
-	UPROPERTY(EditAnywhere, BlueprintReadWrite, Category = SolverSettings)
+	UPROPERTY(EditAnywhere, Category = SolverSettings)
 	FVector Gravity = FVector(0, 0, -981.0);
 
 	/**
 	 * The number of position iterations to run. The position solve is responsible for de-penetration.
 	 * Increasing this will improve simulation stability, but increase the cost.
 	*/
-	UPROPERTY(EditAnywhere, BlueprintReadWrite, Category = SolverSettings, meta = (ClampMin = 0))
+	UPROPERTY(EditAnywhere, Category = SolverSettings, meta = (ClampMin = 0))
 	int32 PositionIterations = 6;
 
 	/**
 	 * The number of velocity iterations to run. The velocity solve is responsible for restitution (bounce) and friction.
 	 * This should usually be 1, but could be 0 if you don't care about friction and restitution.
 	*/
-	UPROPERTY(EditAnywhere, BlueprintReadWrite, Category = SolverSettings, meta = (ClampMin = 0))
+	UPROPERTY(EditAnywhere, Category = SolverSettings, meta = (ClampMin = 0))
 	int32 VelocityIterations = 1;
 
 	/**
@@ -369,18 +372,18 @@ struct CONTROLRIGPHYSICS_API FRigPhysicsSolverSettings
 	 * can be controlled per-constraint in the constraint settings (assuming ProjectionIteration is not zero).
 	 * This should be left as 1 in almost all cases.
 	*/
-	UPROPERTY(EditAnywhere, BlueprintReadWrite, Category = SolverSettings, meta = (ClampMin = 0))
+	UPROPERTY(EditAnywhere, Category = SolverSettings, meta = (ClampMin = 0))
 	int32 ProjectionIterations = 1;
 
 	// This sets how or if the step size should be smoothed. A value of one will disable smoothing,
 	// so that the physics simulation will match the Control Rig delta time. 
-	UPROPERTY(EditAnywhere, BlueprintReadWrite, Category = SolverSettings, meta = (ClampMin = 1))
+	UPROPERTY(EditAnywhere, Category = SolverSettings, meta = (ClampMin = 1))
 	int32 MaxNumRollingAverageStepTimes = 1;
 
 	/**
 	 * When bodies are penetrating, this is the maximum velocity delta that can be applied in one frame.
 	 */
-	UPROPERTY(EditAnywhere, BlueprintReadWrite, Category = SolverSettings, meta = (ClampMin = 0))
+	UPROPERTY(EditAnywhere, Category = SolverSettings, meta = (ClampMin = 0))
 	float MaxDepenetrationVelocity = 0.0f;
 
 	/**
@@ -388,18 +391,18 @@ struct CONTROLRIGPHYSICS_API FRigPhysicsSolverSettings
 	 * NOTE: If this value is non-zero and less than the current frame time, the simulation will step multiple times
 	 * which increases the cost.
 	 */
-	UPROPERTY(EditAnywhere, BlueprintReadWrite, Category = SolverSettings, meta = (ClampMin = 0))
+	UPROPERTY(EditAnywhere, Category = SolverSettings, meta = (ClampMin = 0))
 	float FixedTimeStep = 0.01666667f;
 
 	// The maximum number of solver steps that can be made
-	UPROPERTY(EditAnywhere, BlueprintReadWrite, Category = SolverSettings, meta = (ClampMin = 1))
+	UPROPERTY(EditAnywhere, Category = SolverSettings, meta = (ClampMin = 1))
 	int32 MaxTimeSteps = 10;
 
 	/**
 	 * If a variable timestep is in use, this is the maximum delta time that can be used. If this is smaller than the
 	 * requested total delta time then multiple steps will be used.
 	 */
-	UPROPERTY(EditAnywhere, BlueprintReadWrite, Category = SolverSettings, meta = (ClampMin = 0))
+	UPROPERTY(EditAnywhere, Category = SolverSettings, meta = (ClampMin = 0))
 	float MaxDeltaTime = 1.f / 30.f;
 
 	/**
@@ -408,26 +411,26 @@ struct CONTROLRIGPHYSICS_API FRigPhysicsSolverSettings
 	 * In general you should try to use the linear solver and increase the PositionIterations to improve stability if
 	 * possible, only using the non-linear solver as a last resort.
 	 */
-	UPROPERTY(EditAnywhere, BlueprintReadWrite, Category = SolverSettings)
+	UPROPERTY(EditAnywhere, Category = SolverSettings)
 	bool bUseLinearJointSolver = true;
 
 	/**
 	 * When solving joints, whether to solve the positions last (as opposed to the orientations)
 	 */
-	UPROPERTY(EditAnywhere, BlueprintReadWrite, Category = SolverSettings)
+	UPROPERTY(EditAnywhere, Category = SolverSettings)
 	bool bSolveJointPositionsLast = true;
 
 	/**
 	 * It enables the use of multi-point contact manifolds, which are created only once at the start of each tick.
 	 * When disabled, a single-point contact is generated in each solver iteration which is more expensive.
 	 */
-	UPROPERTY(EditAnywhere, BlueprintReadWrite, Category = CollisionSettings)
+	UPROPERTY(EditAnywhere, Category = CollisionSettings)
 	bool bUseManifolds = true;
 
 	// If any object in the simulation exceeds this distance from the simulation origin, we will
 	// reset the entire simulation. This is to detect problems/explosions.
 	// Disabled when the value is zero.
-	UPROPERTY(EditAnywhere, BlueprintReadWrite, Category = Settings, meta = (ClampMin = "0.0"))
+	UPROPERTY(EditAnywhere, Category = Settings, meta = (ClampMin = "0.0"))
 	float DistanceThresholdForReset = 0;
 };
 
@@ -450,87 +453,41 @@ struct CONTROLRIGPHYSICS_API FRigPhysicsDynamics
 	}
 
 	// The total mass of this body.
-	UPROPERTY(BlueprintReadWrite, EditAnywhere, Category = Dynamics)
+	UPROPERTY(EditAnywhere, Category = Dynamics)
 	float Mass = 1.0f;
 
 	// Use this to specify the location of the centre of mass
-	UPROPERTY(BlueprintReadWrite, EditAnywhere, Category = Dynamics)
+	UPROPERTY(EditAnywhere, Category = Dynamics)
 	bool bOverrideCentreOfMass = false;
 
-	UPROPERTY(BlueprintReadWrite, EditAnywhere, Category = Dynamics, meta = (EditCondition = "bOverrideCentreOfMass"))
+	UPROPERTY(EditAnywhere, Category = Dynamics, meta = (EditCondition = "bOverrideCentreOfMass"))
 	FVector CentreOfMassOverride = FVector::ZeroVector;
 
 	// Use this to specify the moments of inertia
-	UPROPERTY(BlueprintReadWrite, EditAnywhere, Category = Dynamics)
+	UPROPERTY(EditAnywhere, Category = Dynamics)
 	bool bOverrideMomentsOfInertia = false;
 
-	UPROPERTY(BlueprintReadWrite, EditAnywhere, Category = Dynamics, meta = (EditCondition = "bOverrideMomentsOfInertia"))
+	UPROPERTY(EditAnywhere, Category = Dynamics, meta = (EditCondition = "bOverrideMomentsOfInertia"))
 	FVector MomentsOfInertiaOverride = FVector(1.0f, 1.0f, 1.0f);
 };
 
 // This represents the motor drive associated with a physics joint, that can drive the bodies
 // towards a target pose.
 USTRUCT(BlueprintType)
-struct CONTROLRIGPHYSICS_API FRigPhysicsDrive
+struct CONTROLRIGPHYSICS_API FRigPhysicsDriveData
 {
 	GENERATED_BODY()
 
-	friend FArchive& operator <<(FArchive& Ar, FRigPhysicsDrive& Data);
-
-	UPROPERTY(BlueprintReadWrite, EditAnywhere, Category = PhysicsDrive)
-	bool bEnable = true;
-
-	/** The strength used to drive linear motion */
-	UPROPERTY(EditAnywhere, BlueprintReadWrite, Category = PhysicsControl, meta = (ClampMin = "0.0"))
-	float LinearStrength = 0.0f;
-
-	/**
-	 * The amount of damping associated with the linear strength. A value of 1 Results in critically
-	 * damped motion where the control drives as quickly as possible to the target without overshooting.
-	 * Values > 1 result in more damped motion, and values below 1 result in faster, but more "wobbly" motion.
-	 */
-	UPROPERTY(EditAnywhere, BlueprintReadWrite, Category = PhysicsControl, meta = (ClampMin = "0.0"))
-	float LinearDampingRatio = 1.0f;
-
-	/**
-	 * The amount of additional linear damping. This is added to the damping that comes from LinearDampingRatio
-	 * and can be useful when you want damping even when LinearStrength is zero.
-	 */
-	UPROPERTY(EditAnywhere, BlueprintReadWrite, Category = PhysicsControl, meta = (ClampMin = "0.0"))
-	float LinearExtraDamping = 0.0f;
-
-	/** The maximum force used to drive the linear motion. Zero indicates no limit. */
-	UPROPERTY(EditAnywhere, BlueprintReadWrite, Category = PhysicsControl, meta = (ClampMin = "0.0"))
-	float MaxForce = 0.0;
-
-	/** The strength used to drive angular motion */
-	UPROPERTY(EditAnywhere, BlueprintReadWrite, Category = PhysicsControl, meta = (ClampMin = "0.0"))
-	float AngularStrength = 0.0f;
+	friend FArchive& operator<<(FArchive& Ar, FRigPhysicsDriveData& Data);
 
-	/**
-	 * The amount of damping associated with the angular strength. A value of 1 Results in critically
-	 * damped motion where the control drives as quickly as possible to the target without overshooting.
-	 * Values > 1 result in more damped motion, and values below 1 result in faster, but more "wobbly" motion.
-	 */
-	UPROPERTY(EditAnywhere, BlueprintReadWrite, Category = PhysicsControl, meta = (ClampMin = "0.0"))
-	float AngularDampingRatio = 1.0f;
-
-	/**
-	 * The amount of additional angular damping. This is added to the damping that comes from AngularDampingRatio
-	 * and can be useful when you want damping even when AngularStrength is zero.
-	 */
-	UPROPERTY(EditAnywhere, BlueprintReadWrite, Category = PhysicsControl, meta = (ClampMin = "0.0"))
-	float AngularExtraDamping = 0.0f;
+	UPROPERTY(EditAnywhere, Category = LinearDrive)
+	FLinearDriveConstraint LinearDriveConstraint;
 
-	/** The maximum torque used to drive the angular motion. Zero indicates no limit. */
-	UPROPERTY(EditAnywhere, BlueprintReadWrite, Category = PhysicsControl, meta = (ClampMin = "0.0"))
-	float MaxTorque = 0.0f;
+	UPROPERTY(EditAnywhere, Category = AngularDrive)
+	FAngularDriveConstraint AngularDriveConstraint;
 
-	/**
-	 * The amount of skeletal animation velocity to use in the targets
-	 * Danny TODO implement - this requires tracking/calculating the velocity
-	 */
-	UPROPERTY(EditAnywhere, BlueprintReadWrite, Category = PhysicsControl)
+	// The amount of skeletal animation velocity to use in the targets
+	UPROPERTY(EditAnywhere, Category = Drive)
 	float SkeletalAnimationVelocityMultiplier = 1.0f;
 };
 
@@ -549,71 +506,57 @@ struct CONTROLRIGPHYSICS_API FRigPhysicsDrive
  *   * There may be times when a body has additional Physics Joints.
  */
 USTRUCT(BlueprintType)
-struct CONTROLRIGPHYSICS_API FRigPhysicsJoint
+struct CONTROLRIGPHYSICS_API FRigPhysicsJointData
 {
 	GENERATED_BODY()
 
-	friend FArchive& operator<<(FArchive& Ar, FRigPhysicsJoint& Data);
-
-	UPROPERTY(BlueprintReadWrite, EditAnywhere, Category = PhysicsJoint)
-	bool bEnable = true;
-
-	// The body element that will be used as the parent for the physics joint (the child will always
-	// be whatever body this physics joint is owned by. If it is unset, then a suitable body will be
-	// searched for (by assuming this body is a component of a bone which is the child of a parent
-	// bone which has a body. If it is otherwise invalid, or the parent bone is not physical, then
-	// no physics joint will be created.
-	UPROPERTY(BlueprintReadOnly, EditAnywhere, Category = PhysicsJoint)
-	FRigComponentKey ParentBody;
+	friend FArchive& operator<<(FArchive& Ar, FRigPhysicsJointData& Data);
 
 	// The auto-calculated offset of the parent frame from the parent body places it at the 
 	// location of the child body relative to the parent body in the initial pose.
-	UPROPERTY(BlueprintReadWrite, EditAnywhere, Category = PhysicsJoint)
+	UPROPERTY(EditAnywhere, Category = Joint)
 	bool bAutoCalculateParentOffset = true;
 
 	// The parent frame offset applied after any auto-calculation
-	UPROPERTY(BlueprintReadWrite, EditAnywhere, Category = PhysicsJoint)
+	UPROPERTY(EditAnywhere, Category = Joint)
 	FTransform ExtraParentOffset;
 
 	// The auto-calculated offset of the child frame is for it to be co-located with the child body.
-	UPROPERTY(BlueprintReadWrite, EditAnywhere, Category = PhysicsJoint)
+	UPROPERTY(EditAnywhere, Category = Joint)
 	bool bAutoCalculateChildOffset = true;
 
 	// The child frame offset applied after any auto-calculation
-	UPROPERTY(BlueprintReadWrite, EditAnywhere, Category = PhysicsJoint)
+	UPROPERTY(EditAnywhere, Category = Joint)
 	FTransform ExtraChildOffset;
 
-	// 0 locks the movement. -ve makes the movement free
-	UPROPERTY(BlueprintReadWrite, EditAnywhere, Category = PhysicsJoint)
-	float LinearLimit = 0.0f;
+	UPROPERTY(EditAnywhere, Category = Joint)
+	FLinearConstraint LinearConstraint;
 
-	// Angular limit in degrees (twist, swing1, swing2). -ve indicates the movement is free
-	UPROPERTY(BlueprintReadWrite, EditAnywhere, Category = PhysicsJoint)
-	FVector AngularLimit = FVector(-1.0);
+	UPROPERTY(EditAnywhere, Category = Joint)
+	FConeConstraint ConeConstraint;
+
+	UPROPERTY(EditAnywhere, Category = Joint)
+	FTwistConstraint TwistConstraint;
 
 	// Disable collisions between the parent and child bodies
-	UPROPERTY(BlueprintReadWrite, EditAnywhere, Category = PhysicsJoint)
+	UPROPERTY(EditAnywhere, Category = Joint)
 	bool bDisableCollision = true;
 
 	// The amount of linear projection to reduce joint separation when the linear constraint is
 	// locked or has a hard limit. A value of one will apply full projection, but can introduce artefacts.
-	UPROPERTY(BlueprintReadWrite, EditAnywhere, Category = PhysicsJoint, meta = (ClampMin = 0, ClampMax = 1))
+	UPROPERTY(EditAnywhere, Category = Joint, meta = (ClampMin = 0, ClampMax = 1))
 	float LinearProjectionAmount = 0.5f;
 
 	// The amount of angular projection to reduce joint separation when the angular constraint is
 	// locked or has a hard limit. Note that projecting back to angular limits will often break the
 	// position projection, so it is normally better to disable this.
-	UPROPERTY(BlueprintReadWrite, EditAnywhere, Category = PhysicsJoint, meta = (ClampMin = 0, ClampMax = 1))
+	UPROPERTY(EditAnywhere, Category = Joint, meta = (ClampMin = 0, ClampMax = 1))
 	float AngularProjectionAmount = 0.0f;
 
 	// As this is reduced to 0, the parent body becomes less affected by the motion of the child
 	// body, making the joint behave in one direction.
-	UPROPERTY(BlueprintReadWrite, EditAnywhere, Category = PhysicsJoint, meta = (ClampMin = 0))
+	UPROPERTY(EditAnywhere, Category = Joint, meta = (ClampMin = 0))
 	float ParentInverseMassScale = 1.0f;
-
-	// Optional motor/drive associated with the physics joint
-	UPROPERTY(EditAnywhere, BlueprintReadWrite, Category = PhysicsControl, meta = (ShowOnlyInnerProperties))
-	FRigPhysicsDrive Drive;
 };
 
 /**
@@ -632,23 +575,23 @@ struct CONTROLRIGPHYSICS_API FRigPhysicsBodySolverSettings
 
 	// Note that setting the solver component, if known, has the benefit of avoiding the need to
 	// search for an automatic solver.
-	UPROPERTY(BlueprintReadWrite, EditAnywhere, Category = Solver)
+	UPROPERTY(EditAnywhere, Category = Solver)
 	FRigComponentKey PhysicsSolverComponentKey;
 
 	// If true (and the physics solver is not explicitly set), then this component will be added to
 	// any physics solver that exists above it in the hierarchy, if that solver allows automatically
 	// adding physics components.
-	UPROPERTY(BlueprintReadWrite, EditAnywhere, Category = Solver)
+	UPROPERTY(EditAnywhere, Category = Solver)
 	bool bUseAutomaticSolver = true;
 
 	// The bone that is used to initialize physics, as well as what to track when the body is set to be kinematic.
 	// Note that if this is unset, the it will default to the parent of our owner.
-	UPROPERTY(BlueprintReadWrite, EditAnywhere, Category = Solver)
+	UPROPERTY(EditAnywhere, Category = Solver)
 	FRigElementKey SourceBone;
 
 	// The bone that is written to following simulation. Note that if this is unset, the it will
 	// default to the parent of our owner.
-	UPROPERTY(BlueprintReadWrite, EditAnywhere, Category = Solver)
+	UPROPERTY(EditAnywhere, Category = Solver)
 	FRigElementKey TargetBone;
 
 	void OnRigHierarchyKeyChanged(const FRigHierarchyKey& InOldKey, const FRigHierarchyKey& InNewKey);
diff --git a/Engine/Plugins/Experimental/PhysicsControl/Source/ControlRigPhysics/Public/RigPhysicsExecution.h b/Engine/Plugins/Experimental/PhysicsControl/Source/ControlRigPhysics/Public/RigPhysicsExecution.h
--- a/Engine/Plugins/Experimental/PhysicsControl/Source/ControlRigPhysics/Public/RigPhysicsExecution.h
+++ b/Engine/Plugins/Experimental/PhysicsControl/Source/ControlRigPhysics/Public/RigPhysicsExecution.h
@@ -2,7 +2,7 @@
 
 #pragma once
 
-#include "RigPhysicsComponent.h"
+#include "RigPhysicsBodyComponent.h"
 #include "RigPhysicsSolverComponent.h"
 #include "RigPhysicsSimulation.h"
 #include "PhysicsControlData.h"
@@ -15,13 +15,9 @@
  * Base struct for all other mutable physics nodes
  */
 USTRUCT(meta = (Category = "RigPhysics", NodeColor = "1.0 0.6 0.3", Keywords = "Physics"))
-struct CONTROLRIGPHYSICS_API FRigUnit_PhysicsBaseMutable: public FRigUnitMutable
+struct CONTROLRIGPHYSICS_API FRigUnit_PhysicsBase: public FRigUnitMutable
 {
 	GENERATED_BODY()
-	
-#if WITH_EDITOR
-	virtual void OnUnitNodeCreated(FRigVMUnitNodeCreatedContext& InContext) const override;
-#endif
 };
 
 /**
@@ -29,7 +25,7 @@ struct CONTROLRIGPHYSICS_API FRigUnit_PhysicsBaseMutable: public FRigUnitMutable
  * Note: This node only runs as part of the construction event.
  */
 USTRUCT(meta=(DisplayName="Add Physics Solver Component", Keywords="Construction,Create,New,Simulation", Varying))
-struct CONTROLRIGPHYSICS_API FRigUnit_AddPhysicsSolver : public FRigUnit_PhysicsBaseMutable
+struct CONTROLRIGPHYSICS_API FRigUnit_AddPhysicsSolver : public FRigUnit_PhysicsBase
 {
 	GENERATED_BODY()
 
@@ -64,7 +60,7 @@ struct CONTROLRIGPHYSICS_API FRigUnit_AddPhysicsSolver : public FRigUnit_Physics
  * delays.
  */
 USTRUCT(meta = (DisplayName = "Instantiate physics"))
-struct CONTROLRIGPHYSICS_API FRigUnit_InstantiatePhysics : public FRigUnit_PhysicsBaseMutable
+struct CONTROLRIGPHYSICS_API FRigUnit_InstantiatePhysics : public FRigUnit_PhysicsBase
 {
 	GENERATED_BODY()
 
@@ -89,7 +85,7 @@ struct CONTROLRIGPHYSICS_API FRigUnit_InstantiatePhysics : public FRigUnit_Physi
  * Steps the specified physics solver
  */
 USTRUCT(meta = (DisplayName = "Step Physics Solver"))
-struct CONTROLRIGPHYSICS_API FRigUnit_StepPhysicsSolver : public FRigUnit_PhysicsBaseMutable
+struct CONTROLRIGPHYSICS_API FRigUnit_StepPhysicsSolver : public FRigUnit_PhysicsBase
 {
 	GENERATED_BODY()
 
@@ -121,20 +117,17 @@ struct CONTROLRIGPHYSICS_API FRigUnit_StepPhysicsSolver : public FRigUnit_Physic
 	FRigPhysicsVisualizationSettings VisualizationSettings;
 };
 
-/**
- * Adds a new physics body as a component on the owner element.
- * Note: This node only runs as part of the construction event.
- */
-USTRUCT(meta = (DisplayName = "Add Physics Body Component", Keywords = "Construction,Create,New,Joint,Body,Skeleton", Varying))
-struct CONTROLRIGPHYSICS_API FRigUnit_AddPhysicsBody : public FRigUnit_PhysicsBaseMutable
+// Adds a set of physics components including the body, joint and controls
+USTRUCT(meta = (DisplayName = "Add Physics Components", Keywords = "Construction,Create,New,Body,Skeleton", Varying))
+struct CONTROLRIGPHYSICS_API FRigUnit_AddPhysicsComponents : public FRigUnit_PhysicsBase
 {
 	GENERATED_BODY()
 
-	FRigUnit_AddPhysicsBody()
+	FRigUnit_AddPhysicsComponents()
 	{
 		Owner.Type = ERigElementType::Bone;
 		Solver.PhysicsSolverComponentKey.ElementKey.Type = ERigElementType::Bone;
-		Solver.PhysicsSolverComponentKey.Name = FRigPhysicsComponent::GetDefaultName();
+		Solver.PhysicsSolverComponentKey.Name = FRigPhysicsBodyComponent::GetDefaultName();
 	}
 
 	RIGVM_METHOD()
@@ -146,396 +139,119 @@ struct CONTROLRIGPHYSICS_API FRigUnit_AddPhysicsBody : public FRigUnit_PhysicsBa
 	UPROPERTY(meta = (Input, BoneName))
 	FRigElementKey Owner;
 
-	UPROPERTY(meta = (Output))
-	FRigComponentKey PhysicsBodyComponentKey;
-
-	/*
-	 * The solver to relate this new physics element to
-	 */
-	UPROPERTY(meta = (Input))
-	FRigPhysicsBodySolverSettings Solver;
-
-	// The dynamics properties of the new physics element	
-	UPROPERTY(meta = (Input))
-	FRigPhysicsDynamics Dynamics;
-
-	// The collision properties of the new physics element
 	UPROPERTY(meta = (Input))
-	FRigPhysicsCollision Collision;
+	bool bAddJoint = true;
 
-	// The physics joint (typically to its parent) properties of the new physics element
 	UPROPERTY(meta = (Input))
-	FRigPhysicsJoint Joint;
+	bool bAddSimSpaceControl = true;
 
-	// The runtime modifiable data
 	UPROPERTY(meta = (Input))
-	FPhysicsControlModifierData BodyData;
-};
-
-// Discards any existing collision data and replaces it with a box based on the joint positions.
-// Note that this must be called before the physics solver is instantiated/stepped.
-USTRUCT(meta = (DisplayName = "Automatically Calculate Collision", Varying))
-struct CONTROLRIGPHYSICS_API FRigUnit_HierarchyAutoCalculateCollision : public FRigUnit_PhysicsBaseMutable
-{
-	GENERATED_BODY()
-
-	FRigUnit_HierarchyAutoCalculateCollision()
-	{
-		PhysicsBodyComponentKey.ElementKey.Type = ERigElementType::Bone;
-		PhysicsBodyComponentKey.Name = FRigPhysicsComponent::GetDefaultName();
-	}
-
-	RIGVM_METHOD()
-	virtual void Execute() override;
+	bool bAddParentSpaceControl = true;
 
-	UPROPERTY(meta = (Input))
+	UPROPERTY(meta = (Output))
 	FRigComponentKey PhysicsBodyComponentKey;
 
-	// The minimum box extent, as a proportion of the maximum box extent.
-	UPROPERTY(meta = (Input))
-	float MinAspectRatio = 0.25f;
-};
-
-// Sets the mass for a physics component body
-USTRUCT(meta = (DisplayName = "Set Physics Dynamics Properties", Varying))
-struct CONTROLRIGPHYSICS_API FRigUnit_HierarchySetDynamics : public FRigUnit_PhysicsBaseMutable
-{
-	GENERATED_BODY()
+	UPROPERTY(meta = (Output))
+	FRigComponentKey PhysicsJointComponentKey;
 
-	FRigUnit_HierarchySetDynamics()
-	{
-		PhysicsBodyComponentKey.ElementKey.Type = ERigElementType::Bone;
-		PhysicsBodyComponentKey.Name = FRigPhysicsComponent::GetDefaultName();
-	}
+	UPROPERTY(meta = (Output))
+	FRigComponentKey SimSpaceControlComponentKey;
 
-	RIGVM_METHOD()
-	virtual void Execute() override;
+	UPROPERTY(meta = (Output))
+	FRigComponentKey ParentSpaceControlComponentKey;
 
+	// The solver to relate this new physics element to
 	UPROPERTY(meta = (Input))
-	FRigComponentKey PhysicsBodyComponentKey;
+	FRigPhysicsBodySolverSettings Solver;
 
+	// The dynamics properties of the new physics body	
 	UPROPERTY(meta = (Input))
 	FRigPhysicsDynamics Dynamics;
-};
-
-// Sets the collision for a physics component body
-USTRUCT(meta = (DisplayName = "Set Physics Collision Properties", Varying))
-struct CONTROLRIGPHYSICS_API FRigUnit_HierarchySetCollision : public FRigUnit_PhysicsBaseMutable
-{
-	GENERATED_BODY()
-
-	FRigUnit_HierarchySetCollision()
-	{
-		PhysicsBodyComponentKey.ElementKey.Type = ERigElementType::Bone;
-		PhysicsBodyComponentKey.Name = FRigPhysicsComponent::GetDefaultName();
-	}
-
-	RIGVM_METHOD()
-	virtual void Execute() override;
-
-	UPROPERTY(meta = (Input))
-	FRigComponentKey PhysicsBodyComponentKey;
 
+	// The collision properties of the new physics body
 	UPROPERTY(meta = (Input))
 	FRigPhysicsCollision Collision;
-};
-
-// Sets the joint for a physics component body
-USTRUCT(meta = (DisplayName = "Set Physics Joint Properties", Varying))
-struct CONTROLRIGPHYSICS_API FRigUnit_HierarchySetJoint : public FRigUnit_PhysicsBaseMutable
-{
-	GENERATED_BODY()
-
-	FRigUnit_HierarchySetJoint()
-	{
-		PhysicsBodyComponentKey.ElementKey.Type = ERigElementType::Bone;
-		PhysicsBodyComponentKey.Name = FRigPhysicsComponent::GetDefaultName();
-	}
-
-	RIGVM_METHOD()
-	virtual void Execute() override;
-
-	UPROPERTY(meta = (Input))
-	FRigComponentKey PhysicsBodyComponentKey;
-
-	UPROPERTY(meta = (Input))
-	FRigPhysicsJoint Joint;
-};
-
-// Sets the joint drive for a physics component body
-USTRUCT(meta = (DisplayName = "Set Physics Joint Drive Properties", Varying))
-struct CONTROLRIGPHYSICS_API FRigUnit_HierarchySetJointDrive : public FRigUnit_PhysicsBaseMutable
-{
-	GENERATED_BODY()
-
-	FRigUnit_HierarchySetJointDrive()
-	{
-		PhysicsBodyComponentKey.ElementKey.Type = ERigElementType::Bone;
-		PhysicsBodyComponentKey.Name = FRigPhysicsComponent::GetDefaultName();
-	}
-
-	RIGVM_METHOD()
-	virtual void Execute() override;
-
-	UPROPERTY(meta = (Input))
-	FRigComponentKey PhysicsBodyComponentKey;
-
-	UPROPERTY(meta = (Input))
-	FRigPhysicsDrive Drive;
-};
-
-/**
- * Sets what bone is used as a source transform for the physics body. This is used as a kinematic target, and when
- * initializing the simulation.
- */
-USTRUCT(meta = (DisplayName = "Set Source Bone", Varying))
-struct CONTROLRIGPHYSICS_API FRigUnit_HierarchySetPhysicsBodySourceBone : public FRigUnit_PhysicsBaseMutable
-{
-	GENERATED_BODY()
-
-	FRigUnit_HierarchySetPhysicsBodySourceBone()
-	{
-		PhysicsBodyComponentKey.ElementKey.Type = ERigElementType::Bone;
-		PhysicsBodyComponentKey.Name = FRigPhysicsComponent::GetDefaultName();
-		SourceBone.Type = ERigElementType::Bone;
-	}
-
-	RIGVM_METHOD()
-	virtual void Execute() override;
 
+	// The runtime modifiable data of the new physics body
 	UPROPERTY(meta = (Input))
-	FRigComponentKey PhysicsBodyComponentKey;
-
-	UPROPERTY(meta = (Input))
-	FRigElementKey SourceBone;
-};
-
-/**
- * Sets what bone is targeted by the simulation - i.e. where the simulation output is written to.
- */
-USTRUCT(meta = (DisplayName = "Set Target Bone", Varying))
-struct CONTROLRIGPHYSICS_API FRigUnit_HierarchySetPhysicsBodyTargetBone : public FRigUnit_PhysicsBaseMutable
-{
-	GENERATED_BODY()
-
-	FRigUnit_HierarchySetPhysicsBodyTargetBone()
-	{
-		PhysicsBodyComponentKey.ElementKey.Type = ERigElementType::Bone;
-		PhysicsBodyComponentKey.Name = FRigPhysicsComponent::GetDefaultName();
-		TargetBone.Type = ERigElementType::Bone;
-	}
-
-	RIGVM_METHOD()
-	virtual void Execute() override;
+	FPhysicsControlModifierData BodyData;
 
+	// The properties of the joint
 	UPROPERTY(meta = (Input))
-	FRigComponentKey PhysicsBodyComponentKey;
+	FRigPhysicsJointData JointData;
 
+	// Optional motor/drive associated with the physics joint
 	UPROPERTY(meta = (Input))
-	FRigElementKey TargetBone;
-};
-
-
-/**
- * Sets all the data on a body - but in a sparse way so you can decide which parameters get applied.
- * Danny TODO - Note that the sparse data does not get displayed correctly, so this is largely 
- * unusable - the flags that enable/disable all end up getting reset if the user attempts to change them.
- */
-USTRUCT(meta = (DisplayName = "Set Physics Body Data", Varying))
-struct CONTROLRIGPHYSICS_API FRigUnit_HierarchySetPhysicsBodySparseData : public FRigUnit_PhysicsBaseMutable
-{
-	GENERATED_BODY()
-
-	FRigUnit_HierarchySetPhysicsBodySparseData()
-	{
-		PhysicsBodyComponentKey.ElementKey.Type = ERigElementType::Bone;
-		PhysicsBodyComponentKey.Name = FRigPhysicsComponent::GetDefaultName();
-	}
-
-	RIGVM_METHOD()
-	virtual void Execute() override;
+	FRigPhysicsDriveData DriveData;
 
+	// Data for the simulation space control
 	UPROPERTY(meta = (Input))
-	FRigComponentKey PhysicsBodyComponentKey;
+	FPhysicsControlData SimSpaceControlData;
 
+	// Data for the parent space control
 	UPROPERTY(meta = (Input))
-	FPhysicsControlModifierSparseData Data;
+	FPhysicsControlData ParentSpaceControlData;
 };
 
-/**
- * Sets the kinematic target for a body - note that this won't actually make the body kinematic
- */
-USTRUCT(meta = (DisplayName = "Set Kinematic Target", Varying))
-struct CONTROLRIGPHYSICS_API FRigUnit_HierarchySetPhysicsBodyKinematicTarget : public FRigUnit_PhysicsBaseMutable
-{
-	GENERATED_BODY()
-
-	FRigUnit_HierarchySetPhysicsBodyKinematicTarget()
-	{
-		PhysicsBodyComponentKey.ElementKey.Type = ERigElementType::Bone;
-		PhysicsBodyComponentKey.Name = FRigPhysicsComponent::GetDefaultName();
-	}
-
-	RIGVM_METHOD()
-	virtual void Execute() override;
-
-	UPROPERTY(meta = (Input))
-	FRigComponentKey PhysicsBodyComponentKey;
 
-	UPROPERTY(meta = (Input))
-	ERigPhysicsKinematicTargetSpace KinematicTargetSpace = ERigPhysicsKinematicTargetSpace::OffsetInBoneSpace;
 
-	UPROPERTY(meta = (Input))
-	FTransform KinematicTarget;
-};
 
-/**
- * Sets the movement type for this body. 
- * Danny TODO explain the different between kinematic and static
- */
-USTRUCT(meta = (DisplayName = "Set Physics Body Movement Type", Varying))
-struct CONTROLRIGPHYSICS_API FRigUnit_HierarchySetPhysicsBodyMovementType : public FRigUnit_PhysicsBaseMutable
+// Creates multiple physics components based on the supplied physics asset.
+// Note that the resulting simulation bodies may not precisely match the physics asset.
+USTRUCT(meta = (DisplayName = "Instantiate From Physics Asset", Varying))
+struct CONTROLRIGPHYSICS_API FRigUnit_HierarchyInstantiateFromPhysicsAsset : public FRigUnit_PhysicsBase
 {
 	GENERATED_BODY()
 
-	FRigUnit_HierarchySetPhysicsBodyMovementType()
+	FRigUnit_HierarchyInstantiateFromPhysicsAsset()
 	{
-		PhysicsBodyComponentKey.ElementKey.Type = ERigElementType::Bone;
-		PhysicsBodyComponentKey.Name = FRigPhysicsComponent::GetDefaultName();
 	}
 
 	RIGVM_METHOD()
 	virtual void Execute() override;
 
 	UPROPERTY(meta = (Input))
-	FRigComponentKey PhysicsBodyComponentKey;
-
-	UPROPERTY(meta = (Input))
-	EPhysicsMovementType MovementType = EPhysicsMovementType::Simulated;
-};
-
-/**
- * Sets what collision is used for this body
- */
-USTRUCT(meta = (DisplayName = "Set Physics Body Collision Type", Varying))
-struct CONTROLRIGPHYSICS_API FRigUnit_HierarchySetPhysicsBodyCollisionType : public FRigUnit_PhysicsBaseMutable
-{
-	GENERATED_BODY()
-
-	FRigUnit_HierarchySetPhysicsBodyCollisionType()
-	{
-		PhysicsBodyComponentKey.ElementKey.Type = ERigElementType::Bone;
-		PhysicsBodyComponentKey.Name = FRigPhysicsComponent::GetDefaultName();
-	}
-
-	RIGVM_METHOD()
-	virtual void Execute() override;
+	FRigPhysicsBodySolverSettings Solver;
 
 	UPROPERTY(meta = (Input))
-	FRigComponentKey PhysicsBodyComponentKey;
+	TObjectPtr<UPhysicsAsset> PhysicsAsset;
 
+	// Name of the constraint profile to use. If empty (or invalid), the default profile will be used
 	UPROPERTY(meta = (Input))
-	TEnumAsByte<ECollisionEnabled::Type> CollisionType = ECollisionEnabled::QueryAndPhysics;
-};
-
-/**
- * Sets the multiplier on gravity that should be applied to the body.
- */
-USTRUCT(meta = (DisplayName = "Set Physics Body Gravity Multiplier", Varying))
-struct CONTROLRIGPHYSICS_API FRigUnit_HierarchySetPhysicsBodyGravityMultiplier : public FRigUnit_PhysicsBaseMutable
-{
-	GENERATED_BODY()
-
-	FRigUnit_HierarchySetPhysicsBodyGravityMultiplier()
-	{
-		PhysicsBodyComponentKey.ElementKey.Type = ERigElementType::Bone;
-		PhysicsBodyComponentKey.Name = FRigPhysicsComponent::GetDefaultName();
-	}
-
-	RIGVM_METHOD()
-	virtual void Execute() override;
+	FName ConstraintProfileName;
 
+	// If this is empty, then all bodies in the physics asset that match a bone in the hierarchy
+	// will be created. Otherwise only bodies that relate to the specified bones will be created.
 	UPROPERTY(meta = (Input))
-	FRigComponentKey PhysicsBodyComponentKey;
+	TArray<FRigElementKey> BonesToUse;
 
+	// Whether to enable the drives authored in the physics asset. Note that if you are creating parent space controls, you may not want 
 	UPROPERTY(meta = (Input))
-	float GravityMultiplier = 1.0f;
-};
-
-/**
- * Controls the amount that the simulation is blended back into the target bones. 
- * Danny TODO implement
- */
-USTRUCT(meta = (DisplayName = "Set Physics Body Physics Blend Weight", Varying))
-struct CONTROLRIGPHYSICS_API FRigUnit_HierarchySetPhysicsBodyPhysicsBlendWeight : public FRigUnit_PhysicsBaseMutable
-{
-	GENERATED_BODY()
-
-	FRigUnit_HierarchySetPhysicsBodyPhysicsBlendWeight()
-	{
-		PhysicsBodyComponentKey.ElementKey.Type = ERigElementType::Bone;
-		PhysicsBodyComponentKey.Name = FRigPhysicsComponent::GetDefaultName();
-	}
-
-	RIGVM_METHOD()
-	virtual void Execute() override;
+	bool bEnableDrives = true;
 
 	UPROPERTY(meta = (Input))
-	FRigComponentKey PhysicsBodyComponentKey;
+	bool bAddSimSpaceControl = false;
 
 	UPROPERTY(meta = (Input))
-	float PhysicsBlendWeight = 1.0f;
-};
-
-/**
- * Danny TODO I think this is redundant given the ability to specify the space when kinematic 
- * targets are set. Investigate propagating that into the rest of Physics Control.
- */
-USTRUCT(meta = (DisplayName = "Set Physics Body Use Skeletal Animation", Varying))
-struct CONTROLRIGPHYSICS_API FRigUnit_HierarchySetPhysicsBodyUseSkeletalAnimation : public FRigUnit_PhysicsBaseMutable
-{
-	GENERATED_BODY()
-
-	FRigUnit_HierarchySetPhysicsBodyUseSkeletalAnimation()
-	{
-		PhysicsBodyComponentKey.ElementKey.Type = ERigElementType::Bone;
-		PhysicsBodyComponentKey.Name = FRigPhysicsComponent::GetDefaultName();
-	}
-
-	RIGVM_METHOD()
-	virtual void Execute() override;
+	bool bAddParentSpaceControl = false;
 
+	// Data for the simulation space control
 	UPROPERTY(meta = (Input))
-	FRigComponentKey PhysicsBodyComponentKey;
+	FPhysicsControlData SimSpaceControlData;
 
+	// Data for the parent space control
 	UPROPERTY(meta = (Input))
-	bool bUseSkeletalAnimation = true;
-};
-
-/**
- * If true, then kinematic objects will be written back from simulation to the bones. This only 
- * necessary when either kinematic targets are being used, or when the target bone differs from the source bone.
- * Danny TODO - check the implementation, and also can we auto-calculate this at runtime?
- */
-USTRUCT(meta = (DisplayName = "Set Physics Body Update Kinematic From Simulation", Varying))
-struct CONTROLRIGPHYSICS_API FRigUnit_HierarchySetPhysicsBodyUpdateKinematicFromSimulation : public FRigUnit_PhysicsBaseMutable
-{
-	GENERATED_BODY()
+	FPhysicsControlData ParentSpaceControlData;
 
-	FRigUnit_HierarchySetPhysicsBodyUpdateKinematicFromSimulation()
-	{
-		PhysicsBodyComponentKey.ElementKey.Type = ERigElementType::Bone;
-		PhysicsBodyComponentKey.Name = FRigPhysicsComponent::GetDefaultName();
-	}
+	UPROPERTY(meta = (Output))
+	TArray<FRigComponentKey> PhysicsBodyComponentKeys;
 
-	RIGVM_METHOD()
-	virtual void Execute() override;
+	UPROPERTY(meta = (Output))
+	TArray<FRigComponentKey> PhysicsJointComponentKeys;
 
-	UPROPERTY(meta = (Input))
-	FRigComponentKey PhysicsBodyComponentKey;
+	UPROPERTY(meta = (Output))
+	TArray<FRigComponentKey> SimSpaceControlComponentKeys;
 
-	UPROPERTY(meta = (Input))
-	bool bUpdateKinematicFromSimulation = true;
+	UPROPERTY(meta = (Output))
+	TArray<FRigComponentKey> ParentSpaceControlComponentKeys;
 };
-
-
diff --git a/Engine/Plugins/Experimental/PhysicsControl/Source/ControlRigPhysics/Public/RigPhysicsJointComponent.h b/Engine/Plugins/Experimental/PhysicsControl/Source/ControlRigPhysics/Public/RigPhysicsJointComponent.h
new file mode 100644
--- /dev/null
+++ b/Engine/Plugins/Experimental/PhysicsControl/Source/ControlRigPhysics/Public/RigPhysicsJointComponent.h
@@ -0,0 +1,60 @@
+// Copyright Epic Games, Inc. All Rights Reserved.
+
+#pragma once
+
+#include "RigPhysicsData.h"
+#include "RigPhysicsBodyComponent.h"
+
+#include "PhysicsControlData.h"
+
+#include "Rigs/RigPhysics.h"
+#include "Rigs/RigHierarchyDefines.h"
+
+#include "RigPhysicsJointComponent.generated.h"
+
+USTRUCT(BlueprintType)
+struct CONTROLRIGPHYSICS_API FRigPhysicsJointComponent : public FRigBaseComponent
+{
+public:
+
+	GENERATED_BODY()
+	DECLARE_RIG_COMPONENT_METHODS(FRigPhysicsJointComponent)
+
+	FRigPhysicsJointComponent()
+	{
+		ParentBodyComponentKey.ElementKey.Type = ERigElementType::Bone;
+		ParentBodyComponentKey.Name = FRigPhysicsBodyComponent::GetDefaultName();
+		ChildBodyComponentKey.ElementKey.Type = ERigElementType::Bone;
+		ChildBodyComponentKey.Name = FRigPhysicsBodyComponent::GetDefaultName();
+	}
+
+	// The parent body of the joint. If unset, then the system will try to find a suitable body by
+	// looking for a parent that contains a body that is in the same solver as the child body.
+	UPROPERTY(BlueprintReadOnly, EditAnywhere, Category = Control)
+	FRigComponentKey ParentBodyComponentKey;
+
+	// The child body of the joint. If unset, then the system will try to find a suitable body
+	UPROPERTY(BlueprintReadOnly, EditAnywhere, Category = Control)
+	FRigComponentKey ChildBodyComponentKey;
+
+	// The properties of the joint
+	UPROPERTY(BlueprintReadOnly, EditAnywhere, Category = Physics, meta = (ShowOnlyInnerProperties))
+	FRigPhysicsJointData JointData;
+
+	// Optional motor/drive associated with the physics joint
+	UPROPERTY(EditAnywhere, BlueprintReadWrite, Category = Physics, meta = (ShowOnlyInnerProperties))
+	FRigPhysicsDriveData DriveData;
+
+	virtual void Save(FArchive& Ar) override;
+	virtual void Load(FArchive& Ar) override;
+
+	virtual bool CanBeAddedTo(
+		const FRigElementKey& InElementKey, const URigHierarchy* InHierarchy, FString* OutFailureReason) const override;
+
+	virtual void OnRigHierarchyKeyChanged(const FRigHierarchyKey& InOldKey, const FRigHierarchyKey& InNewKey) override;
+
+	virtual FName GetDefaultComponentName() const override { return GetDefaultName(); }
+
+	static FName GetDefaultName() { return TEXT("PhysicsJoint"); }
+};
+
diff --git a/Engine/Plugins/Experimental/PhysicsControl/Source/ControlRigPhysics/Public/RigPhysicsJointExecution.h b/Engine/Plugins/Experimental/PhysicsControl/Source/ControlRigPhysics/Public/RigPhysicsJointExecution.h
new file mode 100644
--- /dev/null
+++ b/Engine/Plugins/Experimental/PhysicsControl/Source/ControlRigPhysics/Public/RigPhysicsJointExecution.h
@@ -0,0 +1,110 @@
+// Copyright Epic Games, Inc. All Rights Reserved.
+
+#pragma once
+
+#include "RigPhysicsBodyComponent.h"
+#include "RigPhysicsJointComponent.h"
+#include "RigPhysicsSolverComponent.h"
+#include "RigPhysicsSimulation.h"
+
+#include "Units/Execution/RigUnit_DynamicHierarchy.h"
+
+#include "RigPhysicsJointExecution.generated.h"
+
+/**
+ * Base struct for all other mutable physics joint nodes
+ */
+USTRUCT(meta = (Category = "RigPhysics", NodeColor = "0.9 0.7 0.3", Keywords = "Physics"))
+struct CONTROLRIGPHYSICS_API FRigUnit_PhysicsJointBase: public FRigUnitMutable
+{
+	GENERATED_BODY()
+};
+
+/**
+ * Adds a new physics body as a component on the owner element.
+ * Note: This node only runs as part of the construction event.
+ */
+USTRUCT(meta = (DisplayName = "Add Physics Joint Component", Keywords = "Construction,Create,New,Body,Skeleton", Varying))
+struct CONTROLRIGPHYSICS_API FRigUnit_AddPhysicsJoint : public FRigUnit_PhysicsJointBase
+{
+	GENERATED_BODY()
+
+	FRigUnit_AddPhysicsJoint()
+	{
+		Owner.Type = ERigElementType::Bone;
+	}
+
+	RIGVM_METHOD()
+	virtual void Execute() override;
+
+	/**
+	 * The owner of the newly created component (must be set/valid)
+	 */
+	UPROPERTY(meta = (Input, BoneName))
+	FRigElementKey Owner;
+
+	UPROPERTY(meta = (Output))
+	FRigComponentKey PhysicsJointComponentKey;
+
+	// The parent body of the joint. If unset, then the system will try to find a suitable body by
+	// looking for a parent that contains a body that is in the same solver as the child body.
+	UPROPERTY(meta = (Input))
+	FRigComponentKey ParentBodyComponentKey;
+
+	// The child body of the joint. If unset, then the system will try to find a suitable body
+	UPROPERTY(meta = (Input))
+	FRigComponentKey ChildBodyComponentKey;
+
+	// The properties of the joint
+	UPROPERTY(meta = (Input))
+	FRigPhysicsJointData JointData;
+
+	// Optional motor/drive associated with the physics joint
+	UPROPERTY(meta = (Input))
+	FRigPhysicsDriveData DriveData;
+};
+
+// Sets the joint for a physics component body
+USTRUCT(meta = (DisplayName = "Set Physics Joint Properties", Varying))
+struct CONTROLRIGPHYSICS_API FRigUnit_HierarchySetJointData : public FRigUnit_PhysicsJointBase
+{
+	GENERATED_BODY()
+
+	FRigUnit_HierarchySetJointData()
+	{
+		PhysicsJointComponentKey.ElementKey.Type = ERigElementType::Bone;
+		PhysicsJointComponentKey.Name = FRigPhysicsJointComponent::GetDefaultName();
+	}
+
+	RIGVM_METHOD()
+	virtual void Execute() override;
+
+	UPROPERTY(meta = (Input))
+	FRigComponentKey PhysicsJointComponentKey;
+
+	UPROPERTY(meta = (Input))
+	FRigPhysicsJointData JointData;
+};
+
+// Sets the joint drive for a physics component body
+USTRUCT(meta = (DisplayName = "Set Physics Joint Drive Properties", Varying))
+struct CONTROLRIGPHYSICS_API FRigUnit_HierarchySetJointDriveData : public FRigUnit_PhysicsJointBase
+{
+	GENERATED_BODY()
+
+	FRigUnit_HierarchySetJointDriveData()
+	{
+		PhysicsJointComponentKey.ElementKey.Type = ERigElementType::Bone;
+		PhysicsJointComponentKey.Name = FRigPhysicsJointComponent::GetDefaultName();
+	}
+
+	RIGVM_METHOD()
+	virtual void Execute() override;
+
+	UPROPERTY(meta = (Input))
+	FRigComponentKey PhysicsJointComponentKey;
+
+	UPROPERTY(meta = (Input))
+	FRigPhysicsDriveData DriveData;
+};
+
diff --git a/Engine/Plugins/Experimental/PhysicsControl/Source/PhysicsControl/Public/PhysicsControlObjectVersion.h b/Engine/Plugins/Experimental/PhysicsControl/Source/PhysicsControl/Public/PhysicsControlObjectVersion.h
--- a/Engine/Plugins/Experimental/PhysicsControl/Source/PhysicsControl/Public/PhysicsControlObjectVersion.h
+++ b/Engine/Plugins/Experimental/PhysicsControl/Source/PhysicsControl/Public/PhysicsControlObjectVersion.h
@@ -22,6 +22,14 @@ struct PHYSICSCONTROL_API FPhysicsControlObjectVersion
 
 		ControlRigControlAddControlMultiplier,
 
+		ControlRigSeparateOutJointFromBody,
+
+		ControlRigSupportFullConstraintData,
+
+		ControlRigSupportFullDriveConstraintData,
+
+		ControlRigSupportNoCollisionBodies,
+
 		// -----<new versions can be added above this line>-------------------------------------------------
 		VersionPlusOne,
 		LatestVersion = VersionPlusOne - 1
