Add helpers for making drives and joints, and also support enable/disable joints when instantiating from physics asset. Fix use of collision enabled/disabled when instantiating from a physics asset.

diff --git a/Engine/Plugins/Experimental/PhysicsControl/Source/ControlRigPhysics/Private/RigPhysicsExecution.cpp b/Engine/Plugins/Experimental/PhysicsControl/Source/ControlRigPhysics/Private/RigPhysicsExecution.cpp
--- a/Engine/Plugins/Experimental/PhysicsControl/Source/ControlRigPhysics/Private/RigPhysicsExecution.cpp
+++ b/Engine/Plugins/Experimental/PhysicsControl/Source/ControlRigPhysics/Private/RigPhysicsExecution.cpp
@@ -227,6 +227,9 @@ static void SetUpJointData(
 	JointData.ExtraChildOffset = ConstraintInstance.GetRefFrame(EConstraintFrame::Frame1);
 	JointData.ExtraParentOffset = ConstraintInstance.GetRefFrame(EConstraintFrame::Frame2);
 
+	JointData.ExtraParentOffset.SetRotation(JointData.ExtraParentOffset.GetRotation() *
+		ConstraintInstance.AngularRotationOffset.Quaternion());
+
 	JointData.LinearConstraint = ConstraintProfileProperties.LinearLimit;
 	JointData.ConeConstraint = ConstraintProfileProperties.ConeLimit;
 	JointData.TwistConstraint = ConstraintProfileProperties.TwistLimit;
@@ -382,7 +385,7 @@ FRigUnit_HierarchyInstantiateFromPhysicsAsset_Execute()
 
 	for (const TPair<FRigidBodyIndexPair, bool>& PairPair : CollisionDisableTable)
 	{
-		if (PairPair.Value)
+		if (!PairPair.Value)
 		{
 			const FRigidBodyIndexPair& Pair = PairPair.Key;
 			int32 I1 = Pair.Indices[0];
@@ -414,62 +417,65 @@ FRigUnit_HierarchyInstantiateFromPhysicsAsset_Execute()
 	}
 
 	// Physics Joints
-	for (TObjectPtr<UPhysicsConstraintTemplate>& ConstraintTemplate : PhysicsAsset->ConstraintSetup)
+	if (bEnableJoints)
 	{
-		if (!ConstraintTemplate)
+		for (TObjectPtr<UPhysicsConstraintTemplate>& ConstraintTemplate : PhysicsAsset->ConstraintSetup)
 		{
-			continue;
-		}
+			if (!ConstraintTemplate)
+			{
+				continue;
+			}
 
-		// Note that physics assets are normally set up as child/parent for 1/2
-		// However, users can create their own constraints, and some assets will be the other way round.
-		FConstraintInstance& ConstraintInstance = ConstraintTemplate->DefaultInstance;
-		const FName ChildBoneName = ConstraintInstance.ConstraintBone1;
-		const FName ParentBoneName = ConstraintInstance.ConstraintBone2;
+			// Note that physics assets are normally set up as child/parent for 1/2
+			// However, users can create their own constraints, and some assets will be the other way round.
+			FConstraintInstance& ConstraintInstance = ConstraintTemplate->DefaultInstance;
+			const FName ChildBoneName = ConstraintInstance.ConstraintBone1;
+			const FName ParentBoneName = ConstraintInstance.ConstraintBone2;
 
-		const FRigComponentKey* ChildBodyComponentKey = BoneToBodyComponentKeyMap.Find(ChildBoneName);
-		const FRigComponentKey* ParentBodyComponentKey = BoneToBodyComponentKeyMap.Find(ParentBoneName);
+			const FRigComponentKey* ChildBodyComponentKey = BoneToBodyComponentKeyMap.Find(ChildBoneName);
+			const FRigComponentKey* ParentBodyComponentKey = BoneToBodyComponentKeyMap.Find(ParentBoneName);
 
-		if (!ChildBodyComponentKey || !ParentBodyComponentKey)
-		{
-			if (BonesToUse.IsEmpty())
+			if (!ChildBodyComponentKey || !ParentBodyComponentKey)
 			{
-				// If BonesToUse is not empty, it is hard to say whether this is an
-				// error/warning condition. However, if the physics asset is OK and BonesToUse
-				// is empty, then we shouldn't get here.
-				UE_CONTROLRIG_RIGUNIT_REPORT_WARNING(
-					TEXT("InstantiateFromPhysicsAsset - unable to make physics joint between %s and %s"),
-					*ChildBoneName.ToString(), *ParentBoneName.ToString());
+				if (BonesToUse.IsEmpty())
+				{
+					// If BonesToUse is not empty, it is hard to say whether this is an
+					// error/warning condition. However, if the physics asset is OK and BonesToUse
+					// is empty, then we shouldn't get here.
+					UE_CONTROLRIG_RIGUNIT_REPORT_WARNING(
+						TEXT("InstantiateFromPhysicsAsset - unable to make physics joint between %s and %s"),
+						*ChildBoneName.ToString(), *ParentBoneName.ToString());
+				}
+				continue;
 			}
-			continue;
-		}
 
-		const FConstraintProfileProperties& ConstraintProfileProperties = 
-			ConstraintTemplate->GetConstraintProfilePropertiesOrDefault(ConstraintProfileName);
+			const FConstraintProfileProperties& ConstraintProfileProperties = 
+				ConstraintTemplate->GetConstraintProfilePropertiesOrDefault(ConstraintProfileName);
 
-		FRigPhysicsJointData JointData;
-		SetUpJointData(JointData, ConstraintInstance, ConstraintProfileProperties);
-		FRigPhysicsDriveData DriveData;
-		if (bEnableDrives)
-		{
-			SetUpDriveData(DriveData, ConstraintProfileProperties);
-		}
+			FRigPhysicsJointData JointData;
+			SetUpJointData(JointData, ConstraintInstance, ConstraintProfileProperties);
+			FRigPhysicsDriveData DriveData;
+			if (bEnableDrives)
+			{
+				SetUpDriveData(DriveData, ConstraintProfileProperties);
+			}
 
-		// Add the component
-		FRigComponentKey PhysicsJointComponentKey = Controller->AddComponent(
-			FRigPhysicsJointComponent::StaticStruct(), TEXT("PhysicsJoint"), ChildBodyComponentKey->ElementKey);
-		if (PhysicsJointComponentKey.IsValid())
-		{
-			if (FRigPhysicsJointComponent* Component = Cast<FRigPhysicsJointComponent>(
-				ExecuteContext.Hierarchy->FindComponent(PhysicsJointComponentKey)))
+			// Add the component
+			FRigComponentKey PhysicsJointComponentKey = Controller->AddComponent(
+				FRigPhysicsJointComponent::StaticStruct(), TEXT("PhysicsJoint"), ChildBodyComponentKey->ElementKey);
+			if (PhysicsJointComponentKey.IsValid())
 			{
-				Component->JointData = JointData;
-				Component->DriveData = DriveData;
-				Component->ParentBodyComponentKey = *ParentBodyComponentKey;
-				Component->ChildBodyComponentKey = *ChildBodyComponentKey;
+				if (FRigPhysicsJointComponent* Component = Cast<FRigPhysicsJointComponent>(
+					ExecuteContext.Hierarchy->FindComponent(PhysicsJointComponentKey)))
+				{
+					Component->JointData = JointData;
+					Component->DriveData = DriveData;
+					Component->ParentBodyComponentKey = *ParentBodyComponentKey;
+					Component->ChildBodyComponentKey = *ChildBodyComponentKey;
 
-				PhysicsJointComponentKeys.Add(PhysicsJointComponentKey);
+					PhysicsJointComponentKeys.Add(PhysicsJointComponentKey);
+				}
 			}
-		}
-	} // Loop over constraint templates
+		} // Loop over constraint templates
+	}
 }
diff --git a/Engine/Plugins/Experimental/PhysicsControl/Source/ControlRigPhysics/Private/RigPhysicsJointExecution.cpp b/Engine/Plugins/Experimental/PhysicsControl/Source/ControlRigPhysics/Private/RigPhysicsJointExecution.cpp
--- a/Engine/Plugins/Experimental/PhysicsControl/Source/ControlRigPhysics/Private/RigPhysicsJointExecution.cpp
+++ b/Engine/Plugins/Experimental/PhysicsControl/Source/ControlRigPhysics/Private/RigPhysicsJointExecution.cpp
@@ -8,6 +8,10 @@
 #include "Units/Execution/RigUnit_PrepareForExecution.h"
 #include "ControlRigPhysicsModule.h"
 
+#include "PhysicsControlHelpers.h"
+
+#include "Chaos/ChaosConstraintSettings.h"
+
 #include "ControlRig.h"
 #if WITH_EDITOR
 #include "RigVMModel/RigVMController.h"
@@ -78,3 +82,198 @@ FRigUnit_HierarchySetJointDriveData_Execute()
 	}
 }
 
+//======================================================================================================================
+FRigUnit_MakeArticulationJointData_Execute()
+{
+	FVector SoftStiffness;
+	FVector SoftDamping;
+
+	UE::PhysicsControl::ConvertStrengthToSpringParams(
+		SoftStiffness, SoftDamping,
+		SoftStrength, SoftDampingRatio, FVector(0));
+
+	// Unfortunately, the physics engine will apply scalings to these values, so we need to counter that.
+	SoftStiffness /= Chaos::ConstraintSettings::SoftAngularStiffnessScale();
+	SoftDamping /= Chaos::ConstraintSettings::SoftAngularDampingScale();
+
+	// Twist
+	if (AngularLimit.X < 0)
+	{
+		JointData.TwistConstraint.TwistMotion = EAngularConstraintMotion::ACM_Free;
+	}
+	else if (AngularLimit.X == 0)
+	{
+		JointData.TwistConstraint.TwistMotion = EAngularConstraintMotion::ACM_Locked;
+	}
+	else
+	{
+		JointData.TwistConstraint.TwistMotion = EAngularConstraintMotion::ACM_Limited;
+		JointData.TwistConstraint.TwistLimitDegrees = AngularLimit.X;
+
+		JointData.TwistConstraint.Stiffness = SoftStiffness.X;
+		JointData.TwistConstraint.Damping = SoftDamping.X;
+
+		if (SoftStrength.X >= 0)
+		{
+			JointData.TwistConstraint.bSoftConstraint = true;
+		}
+		else
+		{
+			JointData.TwistConstraint.bSoftConstraint = false;
+		}
+	}
+
+	// Swing 1
+	if (AngularLimit.Y < 0)
+	{
+		JointData.ConeConstraint.Swing1Motion = EAngularConstraintMotion::ACM_Free;
+	}
+	else if (AngularLimit.Y == 0)
+	{
+		JointData.ConeConstraint.Swing1Motion = EAngularConstraintMotion::ACM_Locked;
+	}
+	else
+	{
+		JointData.ConeConstraint.Swing1Motion = EAngularConstraintMotion::ACM_Limited;
+		JointData.ConeConstraint.Swing1LimitDegrees = AngularLimit.Y;
+
+		JointData.ConeConstraint.Stiffness = SoftStiffness.Y;
+		JointData.ConeConstraint.Damping = SoftDamping.Y;
+
+		if (SoftStrength.Y >= 0)
+		{
+			JointData.ConeConstraint.bSoftConstraint = true;
+		}
+		else
+		{
+			JointData.ConeConstraint.bSoftConstraint = false;
+		}
+	}
+
+	// Swing2
+	if (AngularLimit.Z < 0)
+	{
+		JointData.ConeConstraint.Swing2Motion = EAngularConstraintMotion::ACM_Free;
+	}
+	else if (AngularLimit.Z == 0)
+	{
+		JointData.ConeConstraint.Swing2Motion = EAngularConstraintMotion::ACM_Locked;
+	}
+	else
+	{
+		JointData.ConeConstraint.Swing2Motion = EAngularConstraintMotion::ACM_Limited;
+		JointData.ConeConstraint.Swing2LimitDegrees = AngularLimit.Z;
+
+		JointData.ConeConstraint.Stiffness = SoftStiffness.Z;
+		JointData.ConeConstraint.Damping = SoftDamping.Z;
+
+		if (SoftStrength.Z >= 0)
+		{
+			JointData.ConeConstraint.bSoftConstraint = true;
+		}
+		else
+		{
+			JointData.ConeConstraint.bSoftConstraint = false;
+		}
+	}
+
+}
+
+//======================================================================================================================
+FRigUnit_MakeArticulationDriveData_Execute()
+{
+	// Convert to the constraint drive params
+	float AngularSpring;
+	float AngularDamping;
+
+	UE::PhysicsControl::ConvertStrengthToSpringParams(
+		AngularSpring, AngularDamping,
+		AngularStrength, AngularDampingRatio, AngularExtraDamping);
+
+	// Unfortunately, the physics engine will apply scalings to these values, so we need to counter that.
+	AngularSpring /= Chaos::ConstraintSettings::AngularDriveStiffnessScale();
+	AngularDamping /= Chaos::ConstraintSettings::AngularDriveDampingScale();
+
+	DriveData.SkeletalAnimationVelocityMultiplier = SkeletalAnimationVelocityMultiplier;
+
+	DriveData.LinearDriveConstraint.SetLinearPositionDrive(false, false, false);
+	DriveData.LinearDriveConstraint.SetLinearVelocityDrive(false, false, false);
+
+	if (bEnableAngularDrive)
+	{
+		DriveData.AngularDriveConstraint.SetOrientationDriveTwistAndSwing(true, true);
+		DriveData.AngularDriveConstraint.SetOrientationDriveSLERP(true);
+		DriveData.AngularDriveConstraint.SetAngularVelocityDriveTwistAndSwing(true, true);
+		DriveData.AngularDriveConstraint.SetAngularVelocityDriveSLERP(true);
+
+		DriveData.AngularDriveConstraint.SetDriveParams(AngularSpring, AngularDamping, 0.0f);
+		DriveData.AngularDriveConstraint.SetAngularDriveMode(AngularDriveMode);
+		DriveData.AngularDriveConstraint.SetAccelerationMode(true);
+	}
+	else
+	{
+		DriveData.AngularDriveConstraint.SetOrientationDriveTwistAndSwing(false, false);
+		DriveData.AngularDriveConstraint.SetOrientationDriveSLERP(false);
+		DriveData.AngularDriveConstraint.SetAngularVelocityDriveTwistAndSwing(false, false);
+		DriveData.AngularDriveConstraint.SetAngularVelocityDriveSLERP(false);
+	}
+
+}
+//======================================================================================================================
+FRigUnit_MakeDriveData_Execute()
+{
+	// Convert to the constraint drive params
+	float AngularSpring;
+	float AngularDamping;
+
+	float LinearSpring;
+	float LinearDamping;
+
+	UE::PhysicsControl::ConvertStrengthToSpringParams(
+		AngularSpring, AngularDamping,
+		AngularStrength, AngularDampingRatio, AngularExtraDamping);
+	UE::PhysicsControl::ConvertStrengthToSpringParams(
+		LinearSpring, LinearDamping,
+		LinearStrength, LinearDampingRatio, LinearExtraDamping);
+
+	// Unfortunately, the physics engine will apply scalings to these values, so we need to
+	// counter that.
+	LinearSpring /= Chaos::ConstraintSettings::LinearDriveStiffnessScale();
+	LinearDamping /= Chaos::ConstraintSettings::LinearDriveDampingScale();
+	AngularSpring /= Chaos::ConstraintSettings::AngularDriveStiffnessScale();
+	AngularDamping /= Chaos::ConstraintSettings::AngularDriveDampingScale();
+
+	DriveData.SkeletalAnimationVelocityMultiplier = SkeletalAnimationVelocityMultiplier;
+
+	if (bEnableLinearDrive)
+	{
+		DriveData.LinearDriveConstraint.SetLinearPositionDrive(true, true, true);
+		DriveData.LinearDriveConstraint.SetLinearVelocityDrive(true, true, true);
+		DriveData.LinearDriveConstraint.SetDriveParams(FVector(LinearSpring), FVector(LinearDamping), FVector(0));
+		DriveData.LinearDriveConstraint.SetAccelerationMode(true);
+	}
+	else
+	{
+		DriveData.LinearDriveConstraint.SetLinearPositionDrive(false, false, false);
+		DriveData.LinearDriveConstraint.SetLinearVelocityDrive(false, false, false);
+	}
+
+	if (bEnableAngularDrive)
+	{
+		DriveData.AngularDriveConstraint.SetOrientationDriveTwistAndSwing(true, true);
+		DriveData.AngularDriveConstraint.SetOrientationDriveSLERP(true);
+		DriveData.AngularDriveConstraint.SetAngularVelocityDriveTwistAndSwing(true, true);
+		DriveData.AngularDriveConstraint.SetAngularVelocityDriveSLERP(true);
+
+		DriveData.AngularDriveConstraint.SetDriveParams(AngularSpring, AngularDamping, 0.0f);
+		DriveData.AngularDriveConstraint.SetAngularDriveMode(AngularDriveMode);
+		DriveData.AngularDriveConstraint.SetAccelerationMode(true);
+	}
+	else
+	{
+		DriveData.AngularDriveConstraint.SetOrientationDriveTwistAndSwing(false, false);
+		DriveData.AngularDriveConstraint.SetOrientationDriveSLERP(false);
+		DriveData.AngularDriveConstraint.SetAngularVelocityDriveTwistAndSwing(false, false);
+		DriveData.AngularDriveConstraint.SetAngularVelocityDriveSLERP(false);
+	}
+}
\ No newline at end of file
diff --git a/Engine/Plugins/Experimental/PhysicsControl/Source/ControlRigPhysics/Public/RigPhysicsExecution.h b/Engine/Plugins/Experimental/PhysicsControl/Source/ControlRigPhysics/Public/RigPhysicsExecution.h
--- a/Engine/Plugins/Experimental/PhysicsControl/Source/ControlRigPhysics/Public/RigPhysicsExecution.h
+++ b/Engine/Plugins/Experimental/PhysicsControl/Source/ControlRigPhysics/Public/RigPhysicsExecution.h
@@ -193,9 +193,6 @@ struct CONTROLRIGPHYSICS_API FRigUnit_AddPhysicsComponents : public FRigUnit_Phy
 	FPhysicsControlData ParentSpaceControlData;
 };
 
-
-
-
 // Creates multiple physics components based on the supplied physics asset.
 // Note that the resulting simulation bodies may not precisely match the physics asset.
 USTRUCT(meta = (DisplayName = "Instantiate From Physics Asset", Varying))
@@ -225,7 +222,13 @@ struct CONTROLRIGPHYSICS_API FRigUnit_HierarchyInstantiateFromPhysicsAsset : pub
 	UPROPERTY(meta = (Input))
 	TArray<FRigElementKey> BonesToUse;
 
-	// Whether to enable the drives authored in the physics asset. Note that if you are creating parent space controls, you may not want 
+	// Whether to enable the joints authored in the physics asset. Note that you can't have drives
+	// without joints.
+	UPROPERTY(meta = (Input))
+	bool bEnableJoints = true;
+
+	// Whether to enable the drives authored in the physics asset. Note that if you are creating
+	// parent space controls, you may not want the drives
 	UPROPERTY(meta = (Input))
 	bool bEnableDrives = true;
 
diff --git a/Engine/Plugins/Experimental/PhysicsControl/Source/ControlRigPhysics/Public/RigPhysicsJointExecution.h b/Engine/Plugins/Experimental/PhysicsControl/Source/ControlRigPhysics/Public/RigPhysicsJointExecution.h
--- a/Engine/Plugins/Experimental/PhysicsControl/Source/ControlRigPhysics/Public/RigPhysicsJointExecution.h
+++ b/Engine/Plugins/Experimental/PhysicsControl/Source/ControlRigPhysics/Public/RigPhysicsJointExecution.h
@@ -11,11 +11,16 @@
 
 #include "RigPhysicsJointExecution.generated.h"
 
-/**
- * Base struct for all other mutable physics joint nodes
- */
+// Base struct for all mutable physics joint nodes
+USTRUCT(meta = (Category = "RigPhysics", NodeColor = "0.9 0.7 0.3", Keywords = "Physics"))
+struct CONTROLRIGPHYSICS_API FRigUnit_PhysicsJointBaseMutable: public FRigUnitMutable
+{
+	GENERATED_BODY()
+};
+
+// Base struct for all non-mutable physics joint nodes
 USTRUCT(meta = (Category = "RigPhysics", NodeColor = "0.9 0.7 0.3", Keywords = "Physics"))
-struct CONTROLRIGPHYSICS_API FRigUnit_PhysicsJointBase: public FRigUnitMutable
+struct CONTROLRIGPHYSICS_API FRigUnit_PhysicsJointBase : public FRigUnit
 {
 	GENERATED_BODY()
 };
@@ -25,7 +30,7 @@ struct CONTROLRIGPHYSICS_API FRigUnit_PhysicsJointBase: public FRigUnitMutable
  * Note: This node only runs as part of the construction event.
  */
 USTRUCT(meta = (DisplayName = "Add Physics Joint Component", Keywords = "Construction,Create,New,Body,Skeleton", Varying))
-struct CONTROLRIGPHYSICS_API FRigUnit_AddPhysicsJoint : public FRigUnit_PhysicsJointBase
+struct CONTROLRIGPHYSICS_API FRigUnit_AddPhysicsJoint : public FRigUnit_PhysicsJointBaseMutable
 {
 	GENERATED_BODY()
 
@@ -66,7 +71,7 @@ struct CONTROLRIGPHYSICS_API FRigUnit_AddPhysicsJoint : public FRigUnit_PhysicsJ
 
 // Sets the joint for a physics component body
 USTRUCT(meta = (DisplayName = "Set Physics Joint Properties", Varying))
-struct CONTROLRIGPHYSICS_API FRigUnit_HierarchySetJointData : public FRigUnit_PhysicsJointBase
+struct CONTROLRIGPHYSICS_API FRigUnit_HierarchySetJointData : public FRigUnit_PhysicsJointBaseMutable
 {
 	GENERATED_BODY()
 
@@ -88,7 +93,7 @@ struct CONTROLRIGPHYSICS_API FRigUnit_HierarchySetJointData : public FRigUnit_Ph
 
 // Sets the joint drive for a physics component body
 USTRUCT(meta = (DisplayName = "Set Physics Joint Drive Properties", Varying))
-struct CONTROLRIGPHYSICS_API FRigUnit_HierarchySetJointDriveData : public FRigUnit_PhysicsJointBase
+struct CONTROLRIGPHYSICS_API FRigUnit_HierarchySetJointDriveData : public FRigUnit_PhysicsJointBaseMutable
 {
 	GENERATED_BODY()
 
@@ -108,3 +113,143 @@ struct CONTROLRIGPHYSICS_API FRigUnit_HierarchySetJointDriveData : public FRigUn
 	FRigPhysicsDriveData DriveData;
 };
 
+// Helper to simplify creation of joint data
+USTRUCT(meta = (DisplayName = "Make Joint Data", Keywords = "Construction,Create,New,Body,Skeleton", Varying))
+struct CONTROLRIGPHYSICS_API FRigUnit_MakeArticulationJointData : public FRigUnit_PhysicsJointBase
+{
+	GENERATED_BODY()
+
+	RIGVM_METHOD()
+	virtual void Execute() override;
+
+	// Degrees, twist, swing1, swing2
+	// -ve indicates the limit range is free
+	UPROPERTY(meta = (Input, ClampMin = "-1.0"))
+	FVector AngularLimit = FVector(-1, -1, -1);
+
+	// If limited, then this will be used to control the softness
+	// -ve indicates the limit is hard
+	// A value of 1 is reasonably soft
+	UPROPERTY(meta = (Input, ClampMin = "-1.0"))
+	FVector SoftStrength = FVector(-1, -1, -1);
+
+	UPROPERTY(meta = (Input, ClampMin = "0.0"))
+	FVector SoftDampingRatio = FVector(1, 1, 1);
+
+	UPROPERTY(meta = (Output))
+	FRigPhysicsJointData JointData;
+};
+
+// Helper to simplify creation of articulation drive data
+USTRUCT(meta = (DisplayName = "Make Articulation Drive Data", Keywords = "Construction,Create,New,Body,Skeleton", Varying))
+struct CONTROLRIGPHYSICS_API FRigUnit_MakeArticulationDriveData : public FRigUnit_PhysicsJointBase
+{
+	GENERATED_BODY()
+
+	RIGVM_METHOD()
+	virtual void Execute() override;
+
+	UPROPERTY(meta = (Input))
+	// Whether to enable the angular drive
+	bool bEnableAngularDrive = true;
+
+	// The type of drive. Note that SLERP drives don't work if any axis is locked
+	UPROPERTY(meta = (Input))
+	TEnumAsByte<enum EAngularDriveMode::Type> AngularDriveMode = EAngularDriveMode::SLERP;
+
+	/** The strength used to drive angular motion */
+	UPROPERTY(meta = (Input, ClampMin = "0.0"))
+	float AngularStrength = 10.0f;
+
+	/**
+	 * The amount of damping associated with the angular strength. A value of 1 Results in critically
+	 * damped motion where the control drives as quickly as possible to the target without overshooting.
+	 * Values > 1 result in more damped motion, and values below 1 result in faster, but more "wobbly" motion.
+	 */
+	UPROPERTY(meta = (Input, ClampMin = "0.0"))
+	float AngularDampingRatio = 1.0f;
+
+	/**
+	 * The amount of additional angular damping. This is added to the damping that comes from AngularDampingRatio
+	 * and can be useful when you want damping even when AngularStrength is zero.
+	 */
+	UPROPERTY(meta = (Input, ClampMin = "0.0"))
+	float AngularExtraDamping = 0.0f;
+
+	// The amount of skeletal animation velocity to use in the targets
+	UPROPERTY(meta = (Input, ClampMin = "0.0"))
+	float SkeletalAnimationVelocityMultiplier = 1.0f;
+
+	UPROPERTY(meta = (Output))
+	FRigPhysicsDriveData DriveData;
+};
+
+
+// Helper to simplify creation of drive data
+USTRUCT(meta = (DisplayName = "Make Drive Data", Keywords = "Construction,Create,New,Body,Skeleton", Varying))
+struct CONTROLRIGPHYSICS_API FRigUnit_MakeDriveData : public FRigUnit_PhysicsJointBase
+{
+	GENERATED_BODY()
+
+	RIGVM_METHOD()
+	virtual void Execute() override;
+
+	// Whether to enable the linear drive (not normally used for character joints)
+	UPROPERTY(meta = (Input))
+	bool bEnableLinearDrive = false;
+
+	/** The strength used to drive linear motion */
+	UPROPERTY(meta = (Input, ClampMin = "0.0"))
+	float LinearStrength = 10.0f;
+
+	/**
+	 * The amount of damping associated with the linear strength. A value of 1 Results in critically
+	 * damped motion where the control drives as quickly as possible to the target without overshooting.
+	 * Values > 1 result in more damped motion, and values below 1 result in faster, but more "wobbly" motion.
+	 */
+	UPROPERTY(meta = (Input, ClampMin = "0.0"))
+	float LinearDampingRatio = 1.0f;
+
+	/**
+	 * The amount of additional linear damping. This is added to the damping that comes from LinearDampingRatio
+	 * and can be useful when you want damping even when LinearStrength is zero.
+	 */
+	UPROPERTY(meta = (Input, ClampMin = "0.0"))
+	float LinearExtraDamping = 0.0f;
+
+	UPROPERTY(meta = (Input))
+	// Whether to enable the angular drive
+	bool bEnableAngularDrive = true;
+
+	// The type of drive. Note that SLERP drives don't work if any axis is locked
+	UPROPERTY(meta = (Input))
+	TEnumAsByte<enum EAngularDriveMode::Type> AngularDriveMode = EAngularDriveMode::SLERP;
+
+	/** The strength used to drive angular motion */
+	UPROPERTY(meta = (Input, ClampMin = "0.0"))
+	float AngularStrength = 10.0f;
+
+	/**
+	 * The amount of damping associated with the angular strength. A value of 1 Results in critically
+	 * damped motion where the control drives as quickly as possible to the target without overshooting.
+	 * Values > 1 result in more damped motion, and values below 1 result in faster, but more "wobbly" motion.
+	 */
+	UPROPERTY(meta = (Input, ClampMin = "0.0"))
+	float AngularDampingRatio = 1.0f;
+
+	/**
+	 * The amount of additional angular damping. This is added to the damping that comes from AngularDampingRatio
+	 * and can be useful when you want damping even when AngularStrength is zero.
+	 */
+	UPROPERTY(meta = (Input, ClampMin = "0.0"))
+	float AngularExtraDamping = 0.0f;
+
+	// The amount of skeletal animation velocity to use in the targets
+	UPROPERTY(meta = (Input, ClampMin = "0.0"))
+	float SkeletalAnimationVelocityMultiplier = 1.0f;
+
+	UPROPERTY(meta = (Output))
+	FRigPhysicsDriveData DriveData;
+};
+
+
diff --git a/Engine/Source/Runtime/Engine/Classes/PhysicsEngine/ConstraintDrives.h b/Engine/Source/Runtime/Engine/Classes/PhysicsEngine/ConstraintDrives.h
--- a/Engine/Source/Runtime/Engine/Classes/PhysicsEngine/ConstraintDrives.h
+++ b/Engine/Source/Runtime/Engine/Classes/PhysicsEngine/ConstraintDrives.h
@@ -106,8 +106,6 @@ struct FLinearDriveConstraint
 		return XDrive.bEnableVelocityDrive || YDrive.bEnableVelocityDrive || ZDrive.bEnableVelocityDrive;
 	}
 
-private:
-	friend struct FConstraintInstance;
 	//These functions may leave the struct in an invalid state unless calling UpdatePhysX* functions.
 	//They are only meant as helpers for FConstraintInstance
 	ENGINE_API void SetLinearPositionDrive(bool bEnableXDrive, bool bEnableYDrive, bool bEnableZDrive);
@@ -185,8 +183,6 @@ struct FAngularDriveConstraint
 	/** Updates physx drive with properties from unreal */
 	//void UpdatePhysXAngularDrive_AssumesLocked(physx::PxD6Joint* Joint) const;
 
-private:
-	friend struct FConstraintInstance;
 	//These functions may leave the struct in an invalid state unless calling UpdatePhysX* functions.
 	//They are only meant as helpers for FConstraintInstance
 	ENGINE_API void SetOrientationDriveTwistAndSwing(bool InEnableTwistDrive, bool InEnableSwingDrive);
