Control rig physics - tidy up and expose some functions for debugging

diff --git a/Engine/Plugins/Experimental/PhysicsControl/Source/ControlRigPhysics/Private/ControlRigPhysicsModule.cpp b/Engine/Plugins/Experimental/PhysicsControl/Source/ControlRigPhysics/Private/ControlRigPhysicsModule.cpp
--- a/Engine/Plugins/Experimental/PhysicsControl/Source/ControlRigPhysics/Private/ControlRigPhysicsModule.cpp
+++ b/Engine/Plugins/Experimental/PhysicsControl/Source/ControlRigPhysics/Private/ControlRigPhysicsModule.cpp
@@ -9,10 +9,6 @@
 
 DEFINE_LOG_CATEGORY(LogControlRigPhysics);
 
-TAutoConsoleVariable<bool> CVarControlRigPhysicsAutomaticallySetSolver(
-	TEXT("ControlRig.Physics.AutomaticallySetSolver"), false,
-	TEXT("if this is set to true component and nodes will pick the solver automatically."));
-
 //======================================================================================================================
 void FControlRigPhysicsModule::StartupModule()
 {
diff --git a/Engine/Plugins/Experimental/PhysicsControl/Source/ControlRigPhysics/Private/RigPhysicsBodyComponent.cpp b/Engine/Plugins/Experimental/PhysicsControl/Source/ControlRigPhysics/Private/RigPhysicsBodyComponent.cpp
--- a/Engine/Plugins/Experimental/PhysicsControl/Source/ControlRigPhysics/Private/RigPhysicsBodyComponent.cpp
+++ b/Engine/Plugins/Experimental/PhysicsControl/Source/ControlRigPhysics/Private/RigPhysicsBodyComponent.cpp
@@ -115,25 +115,6 @@ void FRigPhysicsBodyComponent::OnAddedToHierarchy(URigHierarchy* InHierarchy, UR
 	{
 		AutoCalculateCollision(InHierarchy);
 	}
-
-	if(!CVarControlRigPhysicsAutomaticallySetSolver.GetValueOnAnyThread())
-	{
-		return;
-	}
-	
-	// add a solver if there isn't any
-	TArray<const FRigPhysicsSolverComponent*> SolverComponents = InHierarchy->GetComponents<FRigPhysicsSolverComponent>();
-	if(SolverComponents.IsEmpty())
-	{
-		BodySolverSettings.PhysicsSolverComponentKey = InController->AddTopLevelComponent(
-			FRigPhysicsSolverComponent::StaticStruct(), NAME_None);
-	}
-	else
-	{
-		BodySolverSettings.PhysicsSolverComponentKey = SolverComponents[0]->GetKey();
-	}
-	// let the system know we changed something
-	InHierarchy->Notify(ERigHierarchyNotification::ComponentContentChanged, this);
 }
 
 //======================================================================================================================
diff --git a/Engine/Plugins/Experimental/PhysicsControl/Source/ControlRigPhysics/Private/RigPhysicsExecution.cpp b/Engine/Plugins/Experimental/PhysicsControl/Source/ControlRigPhysics/Private/RigPhysicsExecution.cpp
--- a/Engine/Plugins/Experimental/PhysicsControl/Source/ControlRigPhysics/Private/RigPhysicsExecution.cpp
+++ b/Engine/Plugins/Experimental/PhysicsControl/Source/ControlRigPhysics/Private/RigPhysicsExecution.cpp
@@ -85,7 +85,8 @@ FRigUnit_StepPhysicsSolver_Execute()
 
 			if (PhysicsSimulation)
 			{
-				PhysicsSimulation->StepSimulation(ExecuteContext, SolverComponent, DeltaTimeOverride, Alpha);
+				PhysicsSimulation->StepSimulation(
+					ExecuteContext, SolverComponent, DeltaTimeOverride, SimulationSpaceDeltaTimeOverride, Alpha);
 				if (VisualizationSettings.bEnableVisualization)
 				{
 					PhysicsSimulation->Draw(
@@ -269,8 +270,7 @@ FRigUnit_HierarchyInstantiateFromPhysicsAsset_Execute()
 		return;
 	}
 
-	FRigHierarchyControllerInstructionBracket InstructionBracket(
-		Controller, ExecuteContext.GetInstructionIndex());
+	FRigHierarchyControllerInstructionBracket InstructionBracket(Controller, ExecuteContext.GetInstructionIndex());
 
 	// When we create physics joints we need to be able to refer back to the bodies that have been created
 	TMap<FName, FRigComponentKey> BoneToBodyComponentKeyMap;
@@ -479,3 +479,32 @@ FRigUnit_HierarchyInstantiateFromPhysicsAsset_Execute()
 		} // Loop over constraint templates
 	}
 }
+
+//======================================================================================================================
+FRigUnit_GetPhysicsSolverSpaceData_Execute()
+{
+	DECLARE_SCOPE_HIERARCHICAL_COUNTER_RIGUNIT();
+	if (!ExecuteContext.Hierarchy)
+	{
+		return;
+	}
+	if (UControlRig* ControlRig = Cast<UControlRig>(ExecuteContext.Hierarchy->GetOuter()))
+	{
+		if (FRigPhysicsSolverComponent* SolverComponent = Cast<FRigPhysicsSolverComponent>(
+			ExecuteContext.Hierarchy->FindComponent(PhysicsSolverComponentKey)))
+		{
+			if (FRigPhysicsSimulation* PhysicsSimulation = Cast<FRigPhysicsSimulation>(
+				ControlRig->GetPhysicsSimulation(PhysicsSolverComponentKey)))
+			{
+				const FRigPhysicsSimulation::FSimulationSpaceData& SimulationSpaceData = 
+					PhysicsSimulation->GetSimulationSpaceData();
+
+				LinearVelocity = SimulationSpaceData.LinearVelocity;
+				AngularVelocity = SimulationSpaceData.AngularVelocity;
+				LinearAcceleration = SimulationSpaceData.LinearAcceleration;
+				AngularAcceleration = SimulationSpaceData.AngularAcceleration;
+				Gravity = SimulationSpaceData.Gravity;
+			}
+		}
+	}
+}
\ No newline at end of file
diff --git a/Engine/Plugins/Experimental/PhysicsControl/Source/ControlRigPhysics/Private/RigPhysicsSimulation.cpp b/Engine/Plugins/Experimental/PhysicsControl/Source/ControlRigPhysics/Private/RigPhysicsSimulation.cpp
--- a/Engine/Plugins/Experimental/PhysicsControl/Source/ControlRigPhysics/Private/RigPhysicsSimulation.cpp
+++ b/Engine/Plugins/Experimental/PhysicsControl/Source/ControlRigPhysics/Private/RigPhysicsSimulation.cpp
@@ -768,11 +768,10 @@ void FRigPhysicsSimulation::InstantiateControls(
 
 //======================================================================================================================
 void FRigPhysicsSimulation::UpdateBodyPrePhysics(
-	const FRigVMExecuteContext& ExecuteContext,
+	const FRigVMExecuteContext&       ExecuteContext,
 	const FRigPhysicsSolverComponent* SolverComponent,
-	const FRigBodyRecord& Record, 
-	const FRigPhysicsBodyComponent* PhysicsComponent,
-	const FSimulationSpaceData& SimulationSpaceData)
+	const FRigBodyRecord&             Record, 
+	const FRigPhysicsBodyComponent*   PhysicsComponent)
 {
 	URigHierarchy* Hierarchy = OwningControlRig->GetHierarchy();
 	if (!Record.ActorHandle || !Hierarchy)
@@ -988,8 +987,7 @@ static UE::PhysicsControl::FPosQuat CalculateTargetTM(
 void FRigPhysicsSimulation::UpdatePrePhysics(
 	const FRigVMExecuteContext&       ExecuteContext, 
 	const FRigPhysicsSolverComponent* SolverComponent, 
-	const float                       DeltaTime, 
-	const FSimulationSpaceData&       SimulationSpaceData)
+	const float                       DeltaTime)
 {
 	DECLARE_SCOPE_HIERARCHICAL_COUNTER_FUNC()
 	QUICK_SCOPE_CYCLE_COUNTER(STAT_RigPhysics_UpdatePrePhysics);
@@ -1011,7 +1009,7 @@ void FRigPhysicsSimulation::UpdatePrePhysics(
 		if (const FRigPhysicsBodyComponent* PhysicsComponent = Cast<FRigPhysicsBodyComponent>(
 			Hierarchy->FindComponent(ComponentKey)))
 		{
-			UpdateBodyPrePhysics(ExecuteContext, SolverComponent, Record, PhysicsComponent, SimulationSpaceData);
+			UpdateBodyPrePhysics(ExecuteContext, SolverComponent, Record, PhysicsComponent);
 		}
 	}
 
@@ -1254,6 +1252,7 @@ void FRigPhysicsSimulation::StepSimulation(
 	const FRigVMExecuteContext&       ExecuteContext, 
 	const FRigPhysicsSolverComponent* SolverComponent, 
 	const float                       DeltaTimeOverride, 
+	const float                       SimulationSpaceDeltaTimeOverride,
 	const float                       Alpha)
 {
 	DECLARE_SCOPE_HIERARCHICAL_COUNTER_FUNC()
@@ -1269,9 +1268,14 @@ void FRigPhysicsSimulation::StepSimulation(
 		PhysicsDeltaTime = 0.0f;
 	}
 
+	float PhysicsSimulationSpaceDeltaTime = PhysicsDeltaTime;
+	if (SimulationSpaceDeltaTimeOverride > 0.0f)
+	{
+		PhysicsSimulationSpaceDeltaTime = SimulationSpaceDeltaTimeOverride;
+	}
+
 	// We need to know about the simulation space etc before we can instantiate anything into the right place
-	FSimulationSpaceData SimulationSpaceData = UpdateSimulationSpaceStateAndCalculateData(
-		SolverComponent, PhysicsDeltaTime);
+	SimulationSpaceData = UpdateSimulationSpaceStateAndCalculateData(SolverComponent, PhysicsSimulationSpaceDeltaTime);
 
 	// We instantiate when we do the first simulation - this makes sure any changes applied by
 	// the user have been made. It also means there is no overhead if physics is never stepped.
@@ -1323,7 +1327,7 @@ void FRigPhysicsSimulation::StepSimulation(
 	// * We can't actually simulate with dt = 0
 	if (PhysicsDeltaTime > 0.0f)
 	{
-		UpdatePrePhysics(ExecuteContext, SolverComponent, PhysicsDeltaTime, SimulationSpaceData);
+		UpdatePrePhysics(ExecuteContext, SolverComponent, PhysicsDeltaTime);
 
 		Simulation->Simulate(
 			PhysicsDeltaTime, SolverSettings.MaxDeltaTime, SolverSettings.MaxTimeSteps,
diff --git a/Engine/Plugins/Experimental/PhysicsControl/Source/ControlRigPhysics/Private/RigPhysicsSimulation.h b/Engine/Plugins/Experimental/PhysicsControl/Source/ControlRigPhysics/Private/RigPhysicsSimulation.h
--- a/Engine/Plugins/Experimental/PhysicsControl/Source/ControlRigPhysics/Private/RigPhysicsSimulation.h
+++ b/Engine/Plugins/Experimental/PhysicsControl/Source/ControlRigPhysics/Private/RigPhysicsSimulation.h
@@ -141,6 +141,7 @@ public:
 		const FRigVMExecuteContext&       ExecuteContext, 
 		const FRigPhysicsSolverComponent* SolverComponent, 
 		const float                       DeltaTimeOverride, 
+		const float                       SimulationSpaceDeltaTimeOverride,
 		const float                       Alpha);
 
 	// Draws shapes etc AND (potentially) enables the low-level Chaos debug draw
@@ -150,6 +151,20 @@ public:
 		const FRigPhysicsVisualizationSettings& VisualizationSettings,
 		const UWorld*                           DebugWorld) const;
 
+	// Represents the properties of the simulation space - calculated near the beginning of the update.
+	// Note that all these are specified in the simulation space itself
+	struct FSimulationSpaceData
+	{
+		FVector LinearVelocity;
+		FVector AngularVelocity;
+		FVector LinearAcceleration;
+		FVector AngularAcceleration;
+		FVector Gravity;
+	};
+
+	// Returns the simulation space data, as calculated at the start of the last step
+	const FSimulationSpaceData& GetSimulationSpaceData() const { return SimulationSpaceData; }
+
 private:
 	// Used by the world-space to simulation-space motion transfer system in Component- or
 	// Bone-Space sims, and preserved between updates.
@@ -172,17 +187,6 @@ private:
 		float      PrevDt = 1.0f;
 	};
 
-	// Represents the properties of the simulation space - calculated near the beginning of the update.
-	// Note that all these are specified in the simulation space itself
-	struct FSimulationSpaceData
-	{
-		FVector LinearVelocity;
-		FVector AngularVelocity;
-		FVector LinearAcceleration;
-		FVector AngularAcceleration;
-		FVector Gravity;
-	};
-
 	// Creates the low level simulation
 	void InitialiseSimulation(const FRigPhysicsSolverComponent* SolverComponent);
 
@@ -234,15 +238,13 @@ private:
 	void UpdatePrePhysics(
 		const FRigVMExecuteContext&       ExecuteContext, 
 		const FRigPhysicsSolverComponent* SolverComponent, 
-		const float                       DeltaTime, 
-		const FSimulationSpaceData&       SimulationSpaceData);
+		const float                       DeltaTime);
 
 	void UpdateBodyPrePhysics(
 		const FRigVMExecuteContext&       ExecuteContext,
 		const FRigPhysicsSolverComponent* SolverComponent,
 		const FRigBodyRecord&             Record, 
-		const FRigPhysicsBodyComponent*       PhysicsComponent,
-		const FSimulationSpaceData&       SimulationSpaceData);
+		const FRigPhysicsBodyComponent*   PhysicsComponent);
 
 	void UpdateJointPrePhysics(
 		FRigJointRecord&                 JointRecord,
@@ -320,6 +322,10 @@ private:
 
 	FSimulationSpaceState SimulationSpaceState;
 
+	// Retain the data - we don't actually need to but (a) it makes it available for debugging and
+	// (b) it avoids passing it through the functions.
+	FSimulationSpaceData SimulationSpaceData;
+
 	// This is incremented at the end of each simulation step, used to identify when previously
 	// calculated values are valid
 	// Note that the universe will roll over before a uint64 does
diff --git a/Engine/Plugins/Experimental/PhysicsControl/Source/ControlRigPhysics/Private/RigPhysicsSimulationSpace.cpp b/Engine/Plugins/Experimental/PhysicsControl/Source/ControlRigPhysics/Private/RigPhysicsSimulationSpace.cpp
--- a/Engine/Plugins/Experimental/PhysicsControl/Source/ControlRigPhysics/Private/RigPhysicsSimulationSpace.cpp
+++ b/Engine/Plugins/Experimental/PhysicsControl/Source/ControlRigPhysics/Private/RigPhysicsSimulationSpace.cpp
@@ -178,7 +178,7 @@ FRigPhysicsSimulation::FSimulationSpaceData FRigPhysicsSimulation::UpdateSimulat
 	SimulationSpaceState.SimulationSpaceTM = GetSpaceTransform(
 		SolverSettings.SimulationSpace, SimulationSpaceState.ComponentTM, SimulationSpaceState.BoneRelComponentTM);
 
-	FSimulationSpaceData SimulationSpaceData;
+	SimulationSpaceData = FSimulationSpaceData();
 	SimulationSpaceData.Gravity = ::ConvertWorldVectorToSimSpaceNoScale(
 		SolverSettings.SimulationSpace, SolverSettings.Gravity, 
 		SimulationSpaceState.ComponentTM, SimulationSpaceState.BoneRelComponentTM);
diff --git a/Engine/Plugins/Experimental/PhysicsControl/Source/ControlRigPhysics/Public/ControlRigPhysicsModule.h b/Engine/Plugins/Experimental/PhysicsControl/Source/ControlRigPhysics/Public/ControlRigPhysicsModule.h
--- a/Engine/Plugins/Experimental/PhysicsControl/Source/ControlRigPhysics/Public/ControlRigPhysicsModule.h
+++ b/Engine/Plugins/Experimental/PhysicsControl/Source/ControlRigPhysics/Public/ControlRigPhysicsModule.h
@@ -21,4 +21,3 @@ public:
 	virtual void ShutdownModule() override;
 };
 
-extern CONTROLRIGPHYSICS_API TAutoConsoleVariable<bool> CVarControlRigPhysicsAutomaticallySetSolver;
diff --git a/Engine/Plugins/Experimental/PhysicsControl/Source/ControlRigPhysics/Public/RigPhysicsBodyExecution.h b/Engine/Plugins/Experimental/PhysicsControl/Source/ControlRigPhysics/Public/RigPhysicsBodyExecution.h
--- a/Engine/Plugins/Experimental/PhysicsControl/Source/ControlRigPhysics/Public/RigPhysicsBodyExecution.h
+++ b/Engine/Plugins/Experimental/PhysicsControl/Source/ControlRigPhysics/Public/RigPhysicsBodyExecution.h
@@ -2,6 +2,7 @@
 
 #pragma once
 
+#include "RigPhysicsExecution.h"
 #include "RigPhysicsBodyComponent.h"
 #include "RigPhysicsSolverComponent.h"
 #include "RigPhysicsSimulation.h"
@@ -11,21 +12,12 @@
 
 #include "RigPhysicsBodyExecution.generated.h"
 
-/**
- * Base struct for all other mutable physics nodes
- */
-USTRUCT(meta = (Category = "RigPhysics", NodeColor = "1.0 0.6 0.3", Keywords = "Physics"))
-struct CONTROLRIGPHYSICS_API FRigUnit_PhysicsBodyBase: public FRigUnitMutable
-{
-	GENERATED_BODY()
-};
-
 /**
  * Adds a new physics body as a component on the owner element.
  * Note: This node only runs as part of the construction event.
  */
 USTRUCT(meta = (DisplayName = "Add Physics Body Component", Keywords = "Construction,Create,New,Body,Skeleton", Varying))
-struct CONTROLRIGPHYSICS_API FRigUnit_AddPhysicsBody : public FRigUnit_PhysicsBodyBase
+struct CONTROLRIGPHYSICS_API FRigUnit_AddPhysicsBody : public FRigUnit_PhysicsBaseMutable
 {
 	GENERATED_BODY()
 
@@ -70,7 +62,7 @@ struct CONTROLRIGPHYSICS_API FRigUnit_AddPhysicsBody : public FRigUnit_PhysicsBo
 // Discards any existing collision data and replaces it with a box based on the joint positions.
 // Note that this must be called before the physics solver is instantiated/stepped.
 USTRUCT(meta = (DisplayName = "Automatically Calculate Collision", Varying))
-struct CONTROLRIGPHYSICS_API FRigUnit_HierarchyAutoCalculateCollision : public FRigUnit_PhysicsBodyBase
+struct CONTROLRIGPHYSICS_API FRigUnit_HierarchyAutoCalculateCollision : public FRigUnit_PhysicsBaseMutable
 {
 	GENERATED_BODY()
 
@@ -93,7 +85,7 @@ struct CONTROLRIGPHYSICS_API FRigUnit_HierarchyAutoCalculateCollision : public F
 
 // Sets the mass etc for a physics component body
 USTRUCT(meta = (DisplayName = "Set Physics Body Dynamics Properties", Varying))
-struct CONTROLRIGPHYSICS_API FRigUnit_HierarchySetDynamics : public FRigUnit_PhysicsBodyBase
+struct CONTROLRIGPHYSICS_API FRigUnit_HierarchySetDynamics : public FRigUnit_PhysicsBaseMutable
 {
 	GENERATED_BODY()
 
@@ -115,7 +107,7 @@ struct CONTROLRIGPHYSICS_API FRigUnit_HierarchySetDynamics : public FRigUnit_Phy
 
 // Sets the collision for a physics component body
 USTRUCT(meta = (DisplayName = "Set Physics Body Collision Properties", Varying))
-struct CONTROLRIGPHYSICS_API FRigUnit_HierarchySetCollision : public FRigUnit_PhysicsBodyBase
+struct CONTROLRIGPHYSICS_API FRigUnit_HierarchySetCollision : public FRigUnit_PhysicsBaseMutable
 {
 	GENERATED_BODY()
 
@@ -138,7 +130,7 @@ struct CONTROLRIGPHYSICS_API FRigUnit_HierarchySetCollision : public FRigUnit_Ph
 
 // Disables collision between two bodies
 USTRUCT(meta = (DisplayName = "Disable Collision Between", Varying))
-struct CONTROLRIGPHYSICS_API FRigUnit_HierarchyDisableCollisionBetween : public FRigUnit_PhysicsBodyBase
+struct CONTROLRIGPHYSICS_API FRigUnit_HierarchyDisableCollisionBetween : public FRigUnit_PhysicsBaseMutable
 {
 	GENERATED_BODY()
 
@@ -165,7 +157,7 @@ struct CONTROLRIGPHYSICS_API FRigUnit_HierarchyDisableCollisionBetween : public
  * initializing the simulation.
  */
 USTRUCT(meta = (DisplayName = "Set Physics Body Source Bone", Varying))
-struct CONTROLRIGPHYSICS_API FRigUnit_HierarchySetPhysicsBodySourceBone : public FRigUnit_PhysicsBodyBase
+struct CONTROLRIGPHYSICS_API FRigUnit_HierarchySetPhysicsBodySourceBone : public FRigUnit_PhysicsBaseMutable
 {
 	GENERATED_BODY()
 
@@ -190,7 +182,7 @@ struct CONTROLRIGPHYSICS_API FRigUnit_HierarchySetPhysicsBodySourceBone : public
  * Sets what bone is targeted by the simulation - i.e. where the simulation output is written to.
  */
 USTRUCT(meta = (DisplayName = "Set Physics Body Target Bone", Varying))
-struct CONTROLRIGPHYSICS_API FRigUnit_HierarchySetPhysicsBodyTargetBone : public FRigUnit_PhysicsBodyBase
+struct CONTROLRIGPHYSICS_API FRigUnit_HierarchySetPhysicsBodyTargetBone : public FRigUnit_PhysicsBaseMutable
 {
 	GENERATED_BODY()
 
@@ -218,7 +210,7 @@ struct CONTROLRIGPHYSICS_API FRigUnit_HierarchySetPhysicsBodyTargetBone : public
  * unusable - the flags that enable/disable all end up getting reset if the user attempts to change them.
  */
 USTRUCT(meta = (DisplayName = "Set Physics Body Data", Varying))
-struct CONTROLRIGPHYSICS_API FRigUnit_HierarchySetPhysicsBodySparseData : public FRigUnit_PhysicsBodyBase
+struct CONTROLRIGPHYSICS_API FRigUnit_HierarchySetPhysicsBodySparseData : public FRigUnit_PhysicsBaseMutable
 {
 	GENERATED_BODY()
 
@@ -242,7 +234,7 @@ struct CONTROLRIGPHYSICS_API FRigUnit_HierarchySetPhysicsBodySparseData : public
  * Sets the kinematic target for a body - note that this won't actually make the body kinematic
  */
 USTRUCT(meta = (DisplayName = "Set Physics Body Kinematic Target", Varying))
-struct CONTROLRIGPHYSICS_API FRigUnit_HierarchySetPhysicsBodyKinematicTarget : public FRigUnit_PhysicsBodyBase
+struct CONTROLRIGPHYSICS_API FRigUnit_HierarchySetPhysicsBodyKinematicTarget : public FRigUnit_PhysicsBaseMutable
 {
 	GENERATED_BODY()
 
@@ -270,7 +262,7 @@ struct CONTROLRIGPHYSICS_API FRigUnit_HierarchySetPhysicsBodyKinematicTarget : p
  * Danny TODO explain the different between kinematic and static
  */
 USTRUCT(meta = (DisplayName = "Set Physics Body Movement Type", Varying))
-struct CONTROLRIGPHYSICS_API FRigUnit_HierarchySetPhysicsBodyMovementType : public FRigUnit_PhysicsBodyBase
+struct CONTROLRIGPHYSICS_API FRigUnit_HierarchySetPhysicsBodyMovementType : public FRigUnit_PhysicsBaseMutable
 {
 	GENERATED_BODY()
 
@@ -294,7 +286,7 @@ struct CONTROLRIGPHYSICS_API FRigUnit_HierarchySetPhysicsBodyMovementType : publ
  * Sets what collision is used for this body
  */
 USTRUCT(meta = (DisplayName = "Set Physics Body Collision Type", Varying))
-struct CONTROLRIGPHYSICS_API FRigUnit_HierarchySetPhysicsBodyCollisionType : public FRigUnit_PhysicsBodyBase
+struct CONTROLRIGPHYSICS_API FRigUnit_HierarchySetPhysicsBodyCollisionType : public FRigUnit_PhysicsBaseMutable
 {
 	GENERATED_BODY()
 
@@ -318,7 +310,7 @@ struct CONTROLRIGPHYSICS_API FRigUnit_HierarchySetPhysicsBodyCollisionType : pub
  * Sets the multiplier on gravity that should be applied to the body.
  */
 USTRUCT(meta = (DisplayName = "Set Physics Body Gravity Multiplier", Varying))
-struct CONTROLRIGPHYSICS_API FRigUnit_HierarchySetPhysicsBodyGravityMultiplier : public FRigUnit_PhysicsBodyBase
+struct CONTROLRIGPHYSICS_API FRigUnit_HierarchySetPhysicsBodyGravityMultiplier : public FRigUnit_PhysicsBaseMutable
 {
 	GENERATED_BODY()
 
@@ -343,7 +335,7 @@ struct CONTROLRIGPHYSICS_API FRigUnit_HierarchySetPhysicsBodyGravityMultiplier :
  * Danny TODO implement
  */
 USTRUCT(meta = (DisplayName = "Set Physics Body Physics Blend Weight", Varying))
-struct CONTROLRIGPHYSICS_API FRigUnit_HierarchySetPhysicsBodyPhysicsBlendWeight : public FRigUnit_PhysicsBodyBase
+struct CONTROLRIGPHYSICS_API FRigUnit_HierarchySetPhysicsBodyPhysicsBlendWeight : public FRigUnit_PhysicsBaseMutable
 {
 	GENERATED_BODY()
 
@@ -368,7 +360,7 @@ struct CONTROLRIGPHYSICS_API FRigUnit_HierarchySetPhysicsBodyPhysicsBlendWeight
  * targets are set. Investigate propagating that into the rest of Physics Control.
  */
 USTRUCT(meta = (DisplayName = "Set Physics Body Use Skeletal Animation", Varying))
-struct CONTROLRIGPHYSICS_API FRigUnit_HierarchySetPhysicsBodyUseSkeletalAnimation : public FRigUnit_PhysicsBodyBase
+struct CONTROLRIGPHYSICS_API FRigUnit_HierarchySetPhysicsBodyUseSkeletalAnimation : public FRigUnit_PhysicsBaseMutable
 {
 	GENERATED_BODY()
 
@@ -394,7 +386,7 @@ struct CONTROLRIGPHYSICS_API FRigUnit_HierarchySetPhysicsBodyUseSkeletalAnimatio
  * Danny TODO - check the implementation, and also can we auto-calculate this at runtime?
  */
 USTRUCT(meta = (DisplayName = "Set Physics Body Update Kinematic From Simulation", Varying))
-struct CONTROLRIGPHYSICS_API FRigUnit_HierarchySetPhysicsBodyUpdateKinematicFromSimulation : public FRigUnit_PhysicsBodyBase
+struct CONTROLRIGPHYSICS_API FRigUnit_HierarchySetPhysicsBodyUpdateKinematicFromSimulation : public FRigUnit_PhysicsBaseMutable
 {
 	GENERATED_BODY()
 
diff --git a/Engine/Plugins/Experimental/PhysicsControl/Source/ControlRigPhysics/Public/RigPhysicsControlExecution.h b/Engine/Plugins/Experimental/PhysicsControl/Source/ControlRigPhysics/Public/RigPhysicsControlExecution.h
--- a/Engine/Plugins/Experimental/PhysicsControl/Source/ControlRigPhysics/Public/RigPhysicsControlExecution.h
+++ b/Engine/Plugins/Experimental/PhysicsControl/Source/ControlRigPhysics/Public/RigPhysicsControlExecution.h
@@ -2,6 +2,7 @@
 
 #pragma once
 
+#include "RigPhysicsExecution.h"
 #include "RigPhysicsControlComponent.h"
 #include "PhysicsControlData.h"
 
@@ -9,21 +10,12 @@
 
 #include "RigPhysicsControlExecution.generated.h"
 
-/**
- * Base struct for all other mutable physics control nodes
- */
-USTRUCT(meta = (Category = "RigPhysics", NodeColor = "0.8 0.8 0.3", Keywords = "Physics, Control"))
-struct CONTROLRIGPHYSICS_API FRigUnit_PhysicsControlBase : public FRigUnitMutable
-{
-	GENERATED_BODY()
-};
-
 /**
  * Adds a new physics control as a component on the owner element.
  * Note: This node only runs as part of the construction event.
  */
 USTRUCT(meta = (Category = "RigPhysics", NodeColor = "1.0 0.6 0.3", DisplayName = "Add Physics Control Component", Keywords = "Construction,Create,New,Joint,Control,Physics", Varying))
-struct CONTROLRIGPHYSICS_API FRigUnit_AddPhysicsControl : public FRigUnit_PhysicsControlBase
+struct CONTROLRIGPHYSICS_API FRigUnit_AddPhysicsControl : public FRigUnit_PhysicsBaseMutable
 {
 	GENERATED_BODY()
 
@@ -66,7 +58,7 @@ struct CONTROLRIGPHYSICS_API FRigUnit_AddPhysicsControl : public FRigUnit_Physic
 
 // Sets the control data for a physics control
 USTRUCT(meta = (DisplayName = "Set Physics Control Data", Varying))
-struct CONTROLRIGPHYSICS_API FRigUnit_HierarchySetControlData : public FRigUnit_PhysicsControlBase
+struct CONTROLRIGPHYSICS_API FRigUnit_HierarchySetControlData : public FRigUnit_PhysicsBaseMutable
 {
 	GENERATED_BODY()
 
@@ -88,7 +80,7 @@ struct CONTROLRIGPHYSICS_API FRigUnit_HierarchySetControlData : public FRigUnit_
 
 // Sets the target for a physics control
 USTRUCT(meta = (DisplayName = "Set Physics Control Target", Varying))
-struct CONTROLRIGPHYSICS_API FRigUnit_HierarchySetControlTarget : public FRigUnit_PhysicsControlBase
+struct CONTROLRIGPHYSICS_API FRigUnit_HierarchySetControlTarget : public FRigUnit_PhysicsBaseMutable
 {
 	GENERATED_BODY()
 
@@ -110,7 +102,7 @@ struct CONTROLRIGPHYSICS_API FRigUnit_HierarchySetControlTarget : public FRigUni
 
 // Sets the multipliers for a physics control
 USTRUCT(meta = (DisplayName = "Set Physics Control Multiplier", Varying))
-struct CONTROLRIGPHYSICS_API FRigUnit_HierarchySetControlMultiplier : public FRigUnit_PhysicsControlBase
+struct CONTROLRIGPHYSICS_API FRigUnit_HierarchySetControlMultiplier : public FRigUnit_PhysicsBaseMutable
 {
 	GENERATED_BODY()
 
diff --git a/Engine/Plugins/Experimental/PhysicsControl/Source/ControlRigPhysics/Public/RigPhysicsExecution.h b/Engine/Plugins/Experimental/PhysicsControl/Source/ControlRigPhysics/Public/RigPhysicsExecution.h
--- a/Engine/Plugins/Experimental/PhysicsControl/Source/ControlRigPhysics/Public/RigPhysicsExecution.h
+++ b/Engine/Plugins/Experimental/PhysicsControl/Source/ControlRigPhysics/Public/RigPhysicsExecution.h
@@ -15,7 +15,16 @@
  * Base struct for all other mutable physics nodes
  */
 USTRUCT(meta = (Category = "RigPhysics", NodeColor = "1.0 0.6 0.3", Keywords = "Physics"))
-struct CONTROLRIGPHYSICS_API FRigUnit_PhysicsBase: public FRigUnitMutable
+struct CONTROLRIGPHYSICS_API FRigUnit_PhysicsBaseMutable: public FRigUnitMutable
+{
+	GENERATED_BODY()
+};
+
+/**
+ * Base struct for all other non-mutable physics nodes
+ */
+USTRUCT(meta = (Category = "RigPhysics", NodeColor = "1.0 0.6 0.3", Keywords = "Physics"))
+struct CONTROLRIGPHYSICS_API FRigUnit_PhysicsBase : public FRigUnit
 {
 	GENERATED_BODY()
 };
@@ -25,7 +34,7 @@ struct CONTROLRIGPHYSICS_API FRigUnit_PhysicsBase: public FRigUnitMutable
  * Note: This node only runs as part of the construction event.
  */
 USTRUCT(meta=(DisplayName="Add Physics Solver Component", Keywords="Construction,Create,New,Simulation", Varying))
-struct CONTROLRIGPHYSICS_API FRigUnit_AddPhysicsSolver : public FRigUnit_PhysicsBase
+struct CONTROLRIGPHYSICS_API FRigUnit_AddPhysicsSolver : public FRigUnit_PhysicsBaseMutable
 {
 	GENERATED_BODY()
 
@@ -60,7 +69,7 @@ struct CONTROLRIGPHYSICS_API FRigUnit_AddPhysicsSolver : public FRigUnit_Physics
  * delays.
  */
 USTRUCT(meta = (DisplayName = "Instantiate physics"))
-struct CONTROLRIGPHYSICS_API FRigUnit_InstantiatePhysics : public FRigUnit_PhysicsBase
+struct CONTROLRIGPHYSICS_API FRigUnit_InstantiatePhysics : public FRigUnit_PhysicsBaseMutable
 {
 	GENERATED_BODY()
 
@@ -85,7 +94,7 @@ struct CONTROLRIGPHYSICS_API FRigUnit_InstantiatePhysics : public FRigUnit_Physi
  * Steps the specified physics solver
  */
 USTRUCT(meta = (DisplayName = "Step Physics Solver"))
-struct CONTROLRIGPHYSICS_API FRigUnit_StepPhysicsSolver : public FRigUnit_PhysicsBase
+struct CONTROLRIGPHYSICS_API FRigUnit_StepPhysicsSolver : public FRigUnit_PhysicsBaseMutable
 {
 	GENERATED_BODY()
 
@@ -108,6 +117,13 @@ struct CONTROLRIGPHYSICS_API FRigUnit_StepPhysicsSolver : public FRigUnit_Physic
 	UPROPERTY(meta = (Input))
 	float DeltaTimeOverride = 0.0f;
 
+	// If this is zero, then the simulation delta time will be used for evaluating movement of the
+	// simulation space. If this is positive then it will override. This may be needed if the
+	// component movement is being done in parallel, in which case you might need to pass in the
+	// previous time delta here.
+	UPROPERTY(meta = (Input, ClampMin = "0.0"))
+	float SimulationSpaceDeltaTimeOverride = 0.0f;
+
 	// How much of the simulation is combined with the input bone. This currently happens in
 	// component space. Note that the simulation will continue to run, even if alpha = 0
 	UPROPERTY(meta = (Input, ClampMin = "0.0", ClampMax = "1.0"))
@@ -119,7 +135,7 @@ struct CONTROLRIGPHYSICS_API FRigUnit_StepPhysicsSolver : public FRigUnit_Physic
 
 // Adds a set of physics components including the body, joint and controls
 USTRUCT(meta = (DisplayName = "Add Physics Components", Keywords = "Construction,Create,New,Body,Skeleton", Varying))
-struct CONTROLRIGPHYSICS_API FRigUnit_AddPhysicsComponents : public FRigUnit_PhysicsBase
+struct CONTROLRIGPHYSICS_API FRigUnit_AddPhysicsComponents : public FRigUnit_PhysicsBaseMutable
 {
 	GENERATED_BODY()
 
@@ -196,7 +212,7 @@ struct CONTROLRIGPHYSICS_API FRigUnit_AddPhysicsComponents : public FRigUnit_Phy
 // Creates multiple physics components based on the supplied physics asset.
 // Note that the resulting simulation bodies may not precisely match the physics asset.
 USTRUCT(meta = (DisplayName = "Instantiate From Physics Asset", Varying))
-struct CONTROLRIGPHYSICS_API FRigUnit_HierarchyInstantiateFromPhysicsAsset : public FRigUnit_PhysicsBase
+struct CONTROLRIGPHYSICS_API FRigUnit_HierarchyInstantiateFromPhysicsAsset : public FRigUnit_PhysicsBaseMutable
 {
 	GENERATED_BODY()
 
@@ -258,3 +274,42 @@ struct CONTROLRIGPHYSICS_API FRigUnit_HierarchyInstantiateFromPhysicsAsset : pub
 	UPROPERTY(meta = (Output))
 	TArray<FRigComponentKey> ParentSpaceControlComponentKeys;
 };
+
+
+// Retrieves the simulation space data. Note that this will have been generated during the
+// simulation step, so the values returned will relate to the previous update if the solver has not
+// yet been stepped.
+USTRUCT(meta = (DisplayName = "Get Physics Solver Space Data"))
+struct CONTROLRIGPHYSICS_API FRigUnit_GetPhysicsSolverSpaceData : public FRigUnit_PhysicsBase
+{
+	GENERATED_BODY()
+
+	FRigUnit_GetPhysicsSolverSpaceData()
+	{
+		PhysicsSolverComponentKey.ElementKey.Type = ERigElementType::Bone;
+		PhysicsSolverComponentKey.Name = FRigPhysicsSolverComponent::GetDefaultName();
+	}
+
+	RIGVM_METHOD()
+	virtual void Execute() override;
+
+	// The solver to relate this new physics element to
+	UPROPERTY(meta = (Input))
+	FRigComponentKey PhysicsSolverComponentKey;
+
+	UPROPERTY(meta = (Output))
+	FVector LinearVelocity;
+
+	UPROPERTY(meta = (Output))
+	FVector AngularVelocity;
+
+	UPROPERTY(meta = (Output))
+	FVector LinearAcceleration;
+
+	UPROPERTY(meta = (Output))
+	FVector AngularAcceleration;
+
+	UPROPERTY(meta = (Output))
+	FVector Gravity;
+};
+
diff --git a/Engine/Plugins/Experimental/PhysicsControl/Source/ControlRigPhysics/Public/RigPhysicsJointExecution.h b/Engine/Plugins/Experimental/PhysicsControl/Source/ControlRigPhysics/Public/RigPhysicsJointExecution.h
--- a/Engine/Plugins/Experimental/PhysicsControl/Source/ControlRigPhysics/Public/RigPhysicsJointExecution.h
+++ b/Engine/Plugins/Experimental/PhysicsControl/Source/ControlRigPhysics/Public/RigPhysicsJointExecution.h
@@ -2,6 +2,7 @@
 
 #pragma once
 
+#include "RigPhysicsExecution.h"
 #include "RigPhysicsBodyComponent.h"
 #include "RigPhysicsJointComponent.h"
 #include "RigPhysicsSolverComponent.h"
@@ -11,26 +12,12 @@
 
 #include "RigPhysicsJointExecution.generated.h"
 
-// Base struct for all mutable physics joint nodes
-USTRUCT(meta = (Category = "RigPhysics", NodeColor = "0.9 0.7 0.3", Keywords = "Physics"))
-struct CONTROLRIGPHYSICS_API FRigUnit_PhysicsJointBaseMutable: public FRigUnitMutable
-{
-	GENERATED_BODY()
-};
-
-// Base struct for all non-mutable physics joint nodes
-USTRUCT(meta = (Category = "RigPhysics", NodeColor = "0.9 0.7 0.3", Keywords = "Physics"))
-struct CONTROLRIGPHYSICS_API FRigUnit_PhysicsJointBase : public FRigUnit
-{
-	GENERATED_BODY()
-};
-
 /**
  * Adds a new physics body as a component on the owner element.
  * Note: This node only runs as part of the construction event.
  */
 USTRUCT(meta = (DisplayName = "Add Physics Joint Component", Keywords = "Construction,Create,New,Body,Skeleton", Varying))
-struct CONTROLRIGPHYSICS_API FRigUnit_AddPhysicsJoint : public FRigUnit_PhysicsJointBaseMutable
+struct CONTROLRIGPHYSICS_API FRigUnit_AddPhysicsJoint : public FRigUnit_PhysicsBaseMutable
 {
 	GENERATED_BODY()
 
@@ -71,7 +58,7 @@ struct CONTROLRIGPHYSICS_API FRigUnit_AddPhysicsJoint : public FRigUnit_PhysicsJ
 
 // Sets the joint for a physics component body
 USTRUCT(meta = (DisplayName = "Set Physics Joint Properties", Varying))
-struct CONTROLRIGPHYSICS_API FRigUnit_HierarchySetJointData : public FRigUnit_PhysicsJointBaseMutable
+struct CONTROLRIGPHYSICS_API FRigUnit_HierarchySetJointData : public FRigUnit_PhysicsBaseMutable
 {
 	GENERATED_BODY()
 
@@ -93,7 +80,7 @@ struct CONTROLRIGPHYSICS_API FRigUnit_HierarchySetJointData : public FRigUnit_Ph
 
 // Sets the joint drive for a physics component body
 USTRUCT(meta = (DisplayName = "Set Physics Joint Drive Properties", Varying))
-struct CONTROLRIGPHYSICS_API FRigUnit_HierarchySetJointDriveData : public FRigUnit_PhysicsJointBaseMutable
+struct CONTROLRIGPHYSICS_API FRigUnit_HierarchySetJointDriveData : public FRigUnit_PhysicsBaseMutable
 {
 	GENERATED_BODY()
 
@@ -115,7 +102,7 @@ struct CONTROLRIGPHYSICS_API FRigUnit_HierarchySetJointDriveData : public FRigUn
 
 // Helper to simplify creation of joint data
 USTRUCT(meta = (DisplayName = "Make Joint Data", Keywords = "Construction,Create,New,Body,Skeleton", Varying))
-struct CONTROLRIGPHYSICS_API FRigUnit_MakeArticulationJointData : public FRigUnit_PhysicsJointBase
+struct CONTROLRIGPHYSICS_API FRigUnit_MakeArticulationJointData : public FRigUnit_PhysicsBase
 {
 	GENERATED_BODY()
 
@@ -142,7 +129,7 @@ struct CONTROLRIGPHYSICS_API FRigUnit_MakeArticulationJointData : public FRigUni
 
 // Helper to simplify creation of articulation drive data
 USTRUCT(meta = (DisplayName = "Make Articulation Drive Data", Keywords = "Construction,Create,New,Body,Skeleton", Varying))
-struct CONTROLRIGPHYSICS_API FRigUnit_MakeArticulationDriveData : public FRigUnit_PhysicsJointBase
+struct CONTROLRIGPHYSICS_API FRigUnit_MakeArticulationDriveData : public FRigUnit_PhysicsBase
 {
 	GENERATED_BODY()
 
@@ -187,7 +174,7 @@ struct CONTROLRIGPHYSICS_API FRigUnit_MakeArticulationDriveData : public FRigUni
 
 // Helper to simplify creation of drive data
 USTRUCT(meta = (DisplayName = "Make Drive Data", Keywords = "Construction,Create,New,Body,Skeleton", Varying))
-struct CONTROLRIGPHYSICS_API FRigUnit_MakeDriveData : public FRigUnit_PhysicsJointBase
+struct CONTROLRIGPHYSICS_API FRigUnit_MakeDriveData : public FRigUnit_PhysicsBase
 {
 	GENERATED_BODY()
 
