Control Rig Physics - debug support and update keywords etc

diff --git a/Engine/Plugins/Experimental/PhysicsControl/Source/ControlRigPhysics/Private/RigPhysicsSimulation.cpp b/Engine/Plugins/Experimental/PhysicsControl/Source/ControlRigPhysics/Private/RigPhysicsSimulation.cpp
--- a/Engine/Plugins/Experimental/PhysicsControl/Source/ControlRigPhysics/Private/RigPhysicsSimulation.cpp
+++ b/Engine/Plugins/Experimental/PhysicsControl/Source/ControlRigPhysics/Private/RigPhysicsSimulation.cpp
@@ -479,6 +479,12 @@ ImmediatePhysics::FActorHandle* FRigPhysicsSimulation::CreateBody(
 	check(ActorHandle);
 
 	ActorHandle->SetName(BodyName);
+#if CHAOS_DEBUG_NAME
+	if (Chaos::FGeometryParticleHandle* ParticleHandle = ActorHandle->GetParticle())
+	{
+		ParticleHandle->SetDebugName(MakeShared<FString>(BodyName.ToString()));
+	}
+#endif
 
 	Simulation->AddToCollidingPairs(ActorHandle);
 
diff --git a/Engine/Plugins/Experimental/PhysicsControl/Source/ControlRigPhysics/Public/RigPhysicsBodyExecution.h b/Engine/Plugins/Experimental/PhysicsControl/Source/ControlRigPhysics/Public/RigPhysicsBodyExecution.h
--- a/Engine/Plugins/Experimental/PhysicsControl/Source/ControlRigPhysics/Public/RigPhysicsBodyExecution.h
+++ b/Engine/Plugins/Experimental/PhysicsControl/Source/ControlRigPhysics/Public/RigPhysicsBodyExecution.h
@@ -258,10 +258,10 @@ struct CONTROLRIGPHYSICS_API FRigUnit_HierarchySetPhysicsBodyKinematicTarget : p
 };
 
 /**
- * Sets the movement type for this body. 
+ * Sets the movement mode for this body. 
  * Danny TODO explain the different between kinematic and static
  */
-USTRUCT(meta = (DisplayName = "Set Physics Body Movement Type", Varying))
+USTRUCT(meta = (DisplayName = "Set Physics Body Movement Mode", Varying))
 struct CONTROLRIGPHYSICS_API FRigUnit_HierarchySetPhysicsBodyMovementType : public FRigUnit_PhysicsBaseMutable
 {
 	GENERATED_BODY()
@@ -283,9 +283,9 @@ struct CONTROLRIGPHYSICS_API FRigUnit_HierarchySetPhysicsBodyMovementType : publ
 };
 
 /**
- * Sets what collision is used for this body
+ * Sets what collision mode is used for this body
  */
-USTRUCT(meta = (DisplayName = "Set Physics Body Collision Type", Varying))
+USTRUCT(meta = (DisplayName = "Set Physics Body Collision Mode", Varying))
 struct CONTROLRIGPHYSICS_API FRigUnit_HierarchySetPhysicsBodyCollisionType : public FRigUnit_PhysicsBaseMutable
 {
 	GENERATED_BODY()
diff --git a/Engine/Plugins/Experimental/PhysicsControl/Source/ControlRigPhysics/Public/RigPhysicsExecution.h b/Engine/Plugins/Experimental/PhysicsControl/Source/ControlRigPhysics/Public/RigPhysicsExecution.h
--- a/Engine/Plugins/Experimental/PhysicsControl/Source/ControlRigPhysics/Public/RigPhysicsExecution.h
+++ b/Engine/Plugins/Experimental/PhysicsControl/Source/ControlRigPhysics/Public/RigPhysicsExecution.h
@@ -93,7 +93,7 @@ struct CONTROLRIGPHYSICS_API FRigUnit_InstantiatePhysics : public FRigUnit_Physi
 /**
  * Steps the specified physics solver
  */
-USTRUCT(meta = (DisplayName = "Step Physics Solver"))
+USTRUCT(meta = (DisplayName = "Step Physics Solver", Keywords = "Simulate"))
 struct CONTROLRIGPHYSICS_API FRigUnit_StepPhysicsSolver : public FRigUnit_PhysicsBaseMutable
 {
 	GENERATED_BODY()
@@ -134,7 +134,7 @@ struct CONTROLRIGPHYSICS_API FRigUnit_StepPhysicsSolver : public FRigUnit_Physic
 };
 
 // Adds a set of physics components including the body, joint and controls
-USTRUCT(meta = (DisplayName = "Add Physics Components", Keywords = "Construction,Create,New,Body,Skeleton", Varying))
+USTRUCT(meta = (DisplayName = "Add Physics Components", Keywords = "Construction,Create,New", Varying))
 struct CONTROLRIGPHYSICS_API FRigUnit_AddPhysicsComponents : public FRigUnit_PhysicsBaseMutable
 {
 	GENERATED_BODY()
@@ -211,7 +211,7 @@ struct CONTROLRIGPHYSICS_API FRigUnit_AddPhysicsComponents : public FRigUnit_Phy
 
 // Creates multiple physics components based on the supplied physics asset.
 // Note that the resulting simulation bodies may not precisely match the physics asset.
-USTRUCT(meta = (DisplayName = "Instantiate From Physics Asset", Varying))
+USTRUCT(meta = (DisplayName = "Instantiate From Physics Asset", Keywords = "Construction,Create,New", Varying))
 struct CONTROLRIGPHYSICS_API FRigUnit_HierarchyInstantiateFromPhysicsAsset : public FRigUnit_PhysicsBaseMutable
 {
 	GENERATED_BODY()
@@ -279,7 +279,7 @@ struct CONTROLRIGPHYSICS_API FRigUnit_HierarchyInstantiateFromPhysicsAsset : pub
 // Retrieves the simulation space data. Note that this will have been generated during the
 // simulation step, so the values returned will relate to the previous update if the solver has not
 // yet been stepped.
-USTRUCT(meta = (DisplayName = "Get Physics Solver Space Data"))
+USTRUCT(meta = (DisplayName = "Get Physics Solver Space Data", Keywords = "Debug"))
 struct CONTROLRIGPHYSICS_API FRigUnit_GetPhysicsSolverSpaceData : public FRigUnit_PhysicsBase
 {
 	GENERATED_BODY()
