Control rig physics debug support

diff --git a/Engine/Plugins/Experimental/PhysicsControl/Source/ControlRigPhysics/Private/RigPhysicsExecution.cpp b/Engine/Plugins/Experimental/PhysicsControl/Source/ControlRigPhysics/Private/RigPhysicsExecution.cpp
--- a/Engine/Plugins/Experimental/PhysicsControl/Source/ControlRigPhysics/Private/RigPhysicsExecution.cpp
+++ b/Engine/Plugins/Experimental/PhysicsControl/Source/ControlRigPhysics/Private/RigPhysicsExecution.cpp
@@ -24,6 +24,10 @@
 
 #include UE_INLINE_GENERATED_CPP_BY_NAME(RigPhysicsExecution)
 
+TAutoConsoleVariable<int> CVarControlRigPhysicsVisualizationOverride(
+	TEXT("ControlRig.Physics.VisualizationOveride"), -1,
+	TEXT("-1 uses the Step Physics Solver setting, 0 forces visualization to be disabled, 1 forces it to be enabled."));
+
 //======================================================================================================================
 FRigUnit_AddPhysicsSolver_Execute()
 {
@@ -87,7 +91,8 @@ FRigUnit_StepPhysicsSolver_Execute()
 			{
 				PhysicsSimulation->StepSimulation(
 					ExecuteContext, SolverComponent, DeltaTimeOverride, SimulationSpaceDeltaTimeOverride, Alpha);
-				if (VisualizationSettings.bEnableVisualization)
+				int32 VisualizationOverride = CVarControlRigPhysicsVisualizationOverride.GetValueOnAnyThread();
+				if (VisualizationOverride == 1 || (VisualizationOverride < 0 && VisualizationSettings.bEnableVisualization))
 				{
 					PhysicsSimulation->Draw(
 						ExecuteContext.GetDrawInterface(), SolverComponent->SolverSettings,
diff --git a/Engine/Plugins/Experimental/PhysicsControl/Source/ControlRigPhysics/Private/RigPhysicsSimulationVisualization.cpp b/Engine/Plugins/Experimental/PhysicsControl/Source/ControlRigPhysics/Private/RigPhysicsSimulationVisualization.cpp
--- a/Engine/Plugins/Experimental/PhysicsControl/Source/ControlRigPhysics/Private/RigPhysicsSimulationVisualization.cpp
+++ b/Engine/Plugins/Experimental/PhysicsControl/Source/ControlRigPhysics/Private/RigPhysicsSimulationVisualization.cpp
@@ -164,7 +164,7 @@ static void DrawJoint(
 	bool bChildIsKinematic = ChildActor ? ChildActor->GetIsKinematic() : true;
 	bool bParentIsKinematic = ParentActor ? ParentActor->GetIsKinematic() : true;
 
-	float Size = 5.0f;
+	float Size = 5.0f * VisualizationSettings.ShapeSize;
 
 	const Chaos::FPBDJointSettings& JointSettings = ConstraintHandle->GetSettings();
 
diff --git a/Engine/Plugins/Experimental/PhysicsControl/Source/ControlRigPhysics/Public/RigPhysicsData.h b/Engine/Plugins/Experimental/PhysicsControl/Source/ControlRigPhysics/Public/RigPhysicsData.h
--- a/Engine/Plugins/Experimental/PhysicsControl/Source/ControlRigPhysics/Public/RigPhysicsData.h
+++ b/Engine/Plugins/Experimental/PhysicsControl/Source/ControlRigPhysics/Public/RigPhysicsData.h
@@ -25,13 +25,18 @@ struct FRigPhysicsVisualizationSettings
 {
 	GENERATED_BODY()
 
-	/* Enable visualization (includes enabling low-level chaos debug draw) */
+	// Enable visualization (includes enabling low-level chaos debug draw). Note that this can be
+	// overridden using CVar ControlRig.Physics.VisualizationOveride
 	UPROPERTY(EditAnywhere, Category = Visualization)
 	bool bEnableVisualization = true;
 
 	UPROPERTY(EditAnywhere, Category = Visualization, meta = (ClampMin = "0.0"))
 	float LineThickness = 1.0f;
 
+	// Multiplier on the size things like limit shapes
+	UPROPERTY(EditAnywhere, Category = Visualization, meta = (ClampMin = "0.0"))
+	int32 ShapeSize = 1;
+
 	UPROPERTY(EditAnywhere, Category = Visualization, meta = (ClampMin = "0.0"))
 	int32 ShapeDetail = 16;
 };
