Patch up as a result of CL# 40592633

diff --git a/Engine/Plugins/Experimental/PhysicsControl/Source/ControlRigPhysics/Private/RigPhysicsSimulation.cpp b/Engine/Plugins/Experimental/PhysicsControl/Source/ControlRigPhysics/Private/RigPhysicsSimulation.cpp
--- a/Engine/Plugins/Experimental/PhysicsControl/Source/ControlRigPhysics/Private/RigPhysicsSimulation.cpp
+++ b/Engine/Plugins/Experimental/PhysicsControl/Source/ControlRigPhysics/Private/RigPhysicsSimulation.cpp
@@ -697,9 +697,9 @@ void FRigPhysicsSimulation::InstantiatePhysicsJoints(
 				JointSettings.ConnectorTransforms[1] =
 					JointData.ExtraParentOffset * JointSettings.ConnectorTransforms[1];
 
-				ImmediatePhysics_Chaos::UpdateJointSettingsFromLinearConstraint(JointData.LinearConstraint, JointSettings);
-				ImmediatePhysics_Chaos::UpdateJointSettingsFromConeConstraint(JointData.ConeConstraint, JointSettings);
-				ImmediatePhysics_Chaos::UpdateJointSettingsFromTwistConstraint(JointData.TwistConstraint, JointSettings);
+				ImmediatePhysics::UpdateJointSettingsFromLinearConstraint(JointData.LinearConstraint, JointSettings);
+				ImmediatePhysics::UpdateJointSettingsFromConeConstraint(JointData.ConeConstraint, JointSettings);
+				ImmediatePhysics::UpdateJointSettingsFromTwistConstraint(JointData.TwistConstraint, JointSettings);
 
 				// The physics setting is backwards, because we can't enable collision on bodies that
 				// are set to not collide for other reasons.
@@ -878,8 +878,8 @@ void FRigPhysicsSimulation::UpdateJointPrePhysics(
 		const FRigPhysicsDriveData& DriveData = PhysicsJointComponent->DriveData;
 
 		Chaos::FPBDJointSettings Settings = Constraint->GetSettings();
-		ImmediatePhysics_Chaos::UpdateJointSettingsFromLinearDriveConstraint(DriveData.LinearDriveConstraint, Settings);
-		ImmediatePhysics_Chaos::UpdateJointSettingsFromAngularDriveConstraint(DriveData.AngularDriveConstraint, Settings);
+		ImmediatePhysics::UpdateJointSettingsFromLinearDriveConstraint(DriveData.LinearDriveConstraint, Settings);
+		ImmediatePhysics::UpdateJointSettingsFromAngularDriveConstraint(DriveData.AngularDriveConstraint, Settings);
 		Constraint->SetSettings(Settings);
 
 		// Now set the actual target
