Control rig physics: Fix uninitialised data warnings

diff --git a/Engine/Plugins/Experimental/PhysicsControl/Source/ControlRigPhysics/Public/RigPhysicsExecution.h b/Engine/Plugins/Experimental/PhysicsControl/Source/ControlRigPhysics/Public/RigPhysicsExecution.h
--- a/Engine/Plugins/Experimental/PhysicsControl/Source/ControlRigPhysics/Public/RigPhysicsExecution.h
+++ b/Engine/Plugins/Experimental/PhysicsControl/Source/ControlRigPhysics/Public/RigPhysicsExecution.h
@@ -298,18 +298,18 @@ struct CONTROLRIGPHYSICS_API FRigUnit_GetPhysicsSolverSpaceData : public FRigUni
 	FRigComponentKey PhysicsSolverComponentKey;
 
 	UPROPERTY(meta = (Output))
-	FVector LinearVelocity;
+	FVector LinearVelocity = FVector::ZeroVector;
 
 	UPROPERTY(meta = (Output))
-	FVector AngularVelocity;
+	FVector AngularVelocity = FVector::ZeroVector;
 
 	UPROPERTY(meta = (Output))
-	FVector LinearAcceleration;
+	FVector LinearAcceleration = FVector::ZeroVector;
 
 	UPROPERTY(meta = (Output))
-	FVector AngularAcceleration;
+	FVector AngularAcceleration = FVector::ZeroVector;
 
 	UPROPERTY(meta = (Output))
-	FVector Gravity;
+	FVector Gravity = FVector::ZeroVector;
 };
 
