Initialisation of control rig physics on worker threads. Also adds support for damping.

diff --git a/Engine/Plugins/Experimental/PhysicsControl/Source/ControlRigPhysics/Private/RigPhysicsBodyComponent.cpp b/Engine/Plugins/Experimental/PhysicsControl/Source/ControlRigPhysics/Private/RigPhysicsBodyComponent.cpp
--- a/Engine/Plugins/Experimental/PhysicsControl/Source/ControlRigPhysics/Private/RigPhysicsBodyComponent.cpp
+++ b/Engine/Plugins/Experimental/PhysicsControl/Source/ControlRigPhysics/Private/RigPhysicsBodyComponent.cpp
@@ -150,7 +150,7 @@ void FRigPhysicsBodyComponent::AutoCalculateCollision(URigHierarchy* InHierarchy
 		}
 		else
 		{
-			TM.SetTranslation(InHierarchy->GetLocalTransform(GetElementKey()).GetTranslation() * 0.5);
+			TM.SetTranslation(TM.GetTranslation() * 0.5);
 			// We're going to make a box that goes along the X axis
 			TM.SetRotation(FQuat::FindBetweenVectors(FVector(1, 0, 0), TM.GetTranslation()));
 			FVector Extents;
diff --git a/Engine/Plugins/Experimental/PhysicsControl/Source/ControlRigPhysics/Private/RigPhysicsData.cpp b/Engine/Plugins/Experimental/PhysicsControl/Source/ControlRigPhysics/Private/RigPhysicsData.cpp
--- a/Engine/Plugins/Experimental/PhysicsControl/Source/ControlRigPhysics/Private/RigPhysicsData.cpp
+++ b/Engine/Plugins/Experimental/PhysicsControl/Source/ControlRigPhysics/Private/RigPhysicsData.cpp
@@ -377,4 +377,20 @@ void FRigPhysicsBodySolverSettings::OnRigHierarchyKeyChanged(
 			TargetBone = InNewKey.GetElement();
 		}
 	}
-}
\ No newline at end of file
+}
+
+//======================================================================================================================
+FArchive& operator <<(FArchive& Ar, FRigPhysicsDynamics& Data)
+{
+	Ar << Data.Mass;
+	Ar << Data.bOverrideCentreOfMass;
+	Ar << Data.CentreOfMassOverride;
+	Ar << Data.bOverrideMomentsOfInertia;
+	Ar << Data.MomentsOfInertiaOverride;
+	if (Ar.CustomVer(FPhysicsControlObjectVersion::GUID) >= FPhysicsControlObjectVersion::ControlRigSupportBodyDamping)
+	{
+		Ar << Data.LinearDamping;
+		Ar << Data.AngularDamping;
+	}
+	return Ar;
+}
diff --git a/Engine/Plugins/Experimental/PhysicsControl/Source/ControlRigPhysics/Private/RigPhysicsExecution.cpp b/Engine/Plugins/Experimental/PhysicsControl/Source/ControlRigPhysics/Private/RigPhysicsExecution.cpp
--- a/Engine/Plugins/Experimental/PhysicsControl/Source/ControlRigPhysics/Private/RigPhysicsExecution.cpp
+++ b/Engine/Plugins/Experimental/PhysicsControl/Source/ControlRigPhysics/Private/RigPhysicsExecution.cpp
@@ -327,6 +327,8 @@ FRigUnit_HierarchyInstantiateFromPhysicsAsset_Execute()
 		// Dynamics
 		FRigPhysicsDynamics Dynamics;
 		Dynamics.Mass = SkeletalBodySetup->CalculateMass();
+		Dynamics.LinearDamping = SkeletalBodySetup->DefaultInstance.LinearDamping;
+		Dynamics.AngularDamping = SkeletalBodySetup->DefaultInstance.AngularDamping;
 
 		// Body Data
 		FPhysicsControlModifierData BodyData;
diff --git a/Engine/Plugins/Experimental/PhysicsControl/Source/ControlRigPhysics/Private/RigPhysicsSimulation.cpp b/Engine/Plugins/Experimental/PhysicsControl/Source/ControlRigPhysics/Private/RigPhysicsSimulation.cpp
--- a/Engine/Plugins/Experimental/PhysicsControl/Source/ControlRigPhysics/Private/RigPhysicsSimulation.cpp
+++ b/Engine/Plugins/Experimental/PhysicsControl/Source/ControlRigPhysics/Private/RigPhysicsSimulation.cpp
@@ -25,6 +25,7 @@
 #include "Physics/ImmediatePhysics/ImmediatePhysicsChaos/ImmediatePhysicsSimulation_Chaos.h"
 
 #include "Physics/Experimental/PhysScene_Chaos.h"
+#include "Physics/Experimental/ChaosInterfaceUtils.h"
 
 #include "Chaos/ParticleHandle.h"
 #include "Chaos/PBDJointConstraints.h"
@@ -400,6 +401,124 @@ static Chaos::EJointMotionType GetMotionTypeFromLimitValue(const float Value)
 	}
 }
 
+//======================================================================================================================
+// Danny TODO check all of this
+static bool CreateGeometry(
+	const FRigPhysicsCollision&               Collision,
+	const FRigPhysicsDynamics*                Dynamics,
+	Chaos::FReal&                             OutMass, 
+	Chaos::FVec3&                             OutInertia, 
+	Chaos::FRigidTransform3&                  OutCoMTransform, 
+	Chaos::FImplicitObjectPtr&                OutGeom, 
+	TArray<TUniquePtr<Chaos::FPerShapeData>>& OutShapes)
+{
+	using namespace Chaos;
+
+	OutMass = 0.0f;
+	OutInertia = FVector::ZeroVector;
+	OutCoMTransform = FTransform::Identity;
+
+	// Set the filter to collide with everything (we use a broad phase that only contains particle pairs that are explicitly set to collide)
+	FBodyCollisionData BodyCollisionData;
+	// @todo(chaos): we need an API for setting up filters
+	BodyCollisionData.CollisionFilterData.SimFilter.Word1 = 0xFFFF;
+	BodyCollisionData.CollisionFilterData.SimFilter.Word3 = 0xFFFF;
+
+	// See FBodyInstance::BuildBodyCollisionFlags
+	BodyCollisionData.CollisionFlags.bEnableQueryCollision = false;
+	BodyCollisionData.CollisionFlags.bEnableSimCollisionSimple = true;
+	BodyCollisionData.CollisionFlags.bEnableSimCollisionComplex = false;
+	BodyCollisionData.CollisionFlags.bEnableProbeCollision = false;
+
+	FKAggregateGeom AggGeom;
+	for (const FRigPhysicsCollisionBox& Shape : Collision.Boxes)
+	{
+		FKBoxElem Elem(Shape.Extents.X, Shape.Extents.Y, Shape.Extents.Z);
+		SetCommonProperties(Shape, Elem);
+		Elem.Center = Shape.TM.GetTranslation();
+		Elem.Rotation = Shape.TM.Rotator();
+		AggGeom.BoxElems.Add(Elem);
+	}
+
+	for (const FRigPhysicsCollisionSphere& Shape : Collision.Spheres)
+	{
+		FKSphereElem Elem(Shape.Radius);
+		SetCommonProperties(Shape, Elem);
+		Elem.Center = Shape.TM.GetTranslation();
+		// Note that there is no rotation
+		AggGeom.SphereElems.Add(Elem);
+	}
+
+	for (const FRigPhysicsCollisionCapsule& Shape : Collision.Capsules)
+	{
+		FKSphylElem Elem(Shape.Radius, Shape.Length);
+		SetCommonProperties(Shape, Elem);
+		Elem.Center = Shape.TM.GetTranslation();
+		Elem.Rotation = Shape.TM.Rotator();
+		AggGeom.SphylElems.Add(Elem);
+	}
+
+	FGeometryAddParams AddParams;
+	AddParams.CollisionData = BodyCollisionData;
+	AddParams.CollisionTraceType = ECollisionTraceFlag::CTF_UseDefault;
+	AddParams.Scale = FVector(1, 1, 1);
+	AddParams.LocalTransform = FTransform::Identity;
+	AddParams.WorldTransform = FTransform::Identity;
+	AddParams.Geometry = &AggGeom;
+
+	TArray<Chaos::FImplicitObjectPtr> Geoms;
+	FShapesArray Shapes;
+	ChaosInterface::CreateGeometry(AddParams, Geoms, Shapes);
+
+	if (Geoms.Num() == 0)
+	{
+		return false;
+	}
+
+	if (Dynamics)
+	{
+		// Whether each shape contributes to mass. @todo(chaos): it would be easier if
+		// ComputeMassProperties knew how to extract this info. Maybe it should be a flag in PerShapeData
+		TArray<bool> bContributesToMass;
+		bContributesToMass.Reserve(Shapes.Num());
+		for (int32 ShapeIndex = 0; ShapeIndex < Shapes.Num(); ++ShapeIndex)
+		{
+			const TUniquePtr<FPerShapeData>& Shape = Shapes[ShapeIndex];
+			const FKShapeElem* ShapeElem = FChaosUserData::Get<FKShapeElem>(Shape->GetUserData());
+			bool bHasMass = ShapeElem && ShapeElem->GetContributeToMass();
+			bContributesToMass.Add(bHasMass);
+		}
+
+		// bInertiaScaleIncludeMass = true is to match legacy physics behaviour. This will scale the inertia by the change in mass (density x volumescale) 
+		// as well as the dimension change even though we don't actually change the mass.
+		const float density = 0.001f; // Danny TODO
+		Chaos::FMassProperties MassProperties;
+		ChaosInterface::CalculateMassPropertiesFromShapeCollection(
+			MassProperties, Shapes, bContributesToMass, density);
+
+		OutMass = MassProperties.Mass;
+		OutInertia = MassProperties.InertiaTensor.GetDiagonal();
+		OutCoMTransform = FTransform(MassProperties.RotationOfMass, MassProperties.CenterOfMass);
+	}
+
+	// If we have multiple root shapes, wrap them in a union
+	if (Geoms.Num() == 1)
+	{
+		OutGeom = MoveTemp(Geoms[0]);
+	}
+	else
+	{
+		OutGeom = MakeImplicitObjectPtr<FImplicitObjectUnion>(MoveTemp(Geoms));
+	}
+
+	for (TUniquePtr<FPerShapeData>& Shape : Shapes)
+	{
+		OutShapes.Emplace(MoveTemp(Shape));
+	}
+
+	return true;
+}
+
 //======================================================================================================================
 ImmediatePhysics::FActorHandle* FRigPhysicsSimulation::CreateBody(
 	const FName                 BodyName,
@@ -410,6 +529,12 @@ ImmediatePhysics::FActorHandle* FRigPhysicsSimulation::CreateBody(
 	// Danny TODO Use the correct materials (potentially tricky since we may be called on
 	// any thread - see the RBAN code for reference).
 
+	bool bAddedCollision = false;
+
+	// Disabling this code, but keeping it around for now because we need to double check that the
+	// actors we create directly (in the live code) match what was being created in the
+	// FBodyInstance version, which won't run properly on worker threads.
+#if 0
 	// This appears to be the only way
 	UBodySetup* BodySetup = NewObject<UBodySetup>();
 	BodySetup->CollisionTraceFlag = CTF_UseSimpleAsComplex;
@@ -417,7 +542,6 @@ ImmediatePhysics::FActorHandle* FRigPhysicsSimulation::CreateBody(
 	// Note that the base class (FKShapeElem) for the shapes has got additional
 	// functionality such as enable/disable collision etc.
 
-	bool bAddedCollision = false;
 	for (const FRigPhysicsCollisionBox& Shape : Collision.Boxes)
 	{
 		FKBoxElem Elem(Shape.Extents.X, Shape.Extents.Y, Shape.Extents.Z);
@@ -473,7 +597,55 @@ ImmediatePhysics::FActorHandle* FRigPhysicsSimulation::CreateBody(
 	BodyInstance.InitBody(BodySetup, BodyRelSimSpaceTM, nullptr, nullptr);
 
 	ImmediatePhysics::EActorType ActorType = Dynamics ? ImmediatePhysics::EActorType::DynamicActor : ImmediatePhysics::EActorType::KinematicActor;
-	ImmediatePhysics::FActorHandle* ActorHandle = Simulation->CreateActor(ImmediatePhysics::MakeActorSetup(ActorType, &BodyInstance, BodyRelSimSpaceTM));
+	ImmediatePhysics::FActorSetup&& ActorSetupOld = ImmediatePhysics::MakeActorSetup(ActorType, &BodyInstance, BodyRelSimSpaceTM);
+	ImmediatePhysics::FActorHandle* ActorHandleOld = Simulation->CreateActor(MoveTemp(ActorSetupOld));
+#endif
+	ImmediatePhysics::FActorSetup ActorSetup;
+
+	ActorSetup.ActorType = Dynamics ?
+		ImmediatePhysics::EActorType::DynamicActor : ImmediatePhysics::EActorType::KinematicActor;
+
+	if (Dynamics)
+	{
+		ActorSetup.ActorType = ImmediatePhysics::EActorType::DynamicActor;
+		ActorSetup.bEnableGravity = true;
+		ActorSetup.LinearDamping = 0;
+		ActorSetup.AngularDamping = 0;
+		ActorSetup.LinearDamping = Dynamics->LinearDamping;
+		ActorSetup.AngularDamping = Dynamics->AngularDamping;
+	}
+	else
+	{
+		ActorSetup.ActorType = ImmediatePhysics::EActorType::KinematicActor;
+	}
+
+
+
+	Chaos::FVec3 Inertia;
+	Chaos::FRigidTransform3 CoMTransform;
+	Chaos::FReal Mass;
+
+	Chaos::FImplicitObjectPtr BodyGeom;
+	TArray<TUniquePtr<Chaos::FPerShapeData>> BodyShapes;
+	bool bGeometryCreated = CreateGeometry(
+		Collision, Dynamics, Mass, Inertia, CoMTransform, BodyGeom, BodyShapes);
+
+	ActorSetup.Mass = Mass;
+	ActorSetup.Inertia = Inertia;
+	if (Mass > 0.0f && Dynamics)
+	{
+		ActorSetup.Mass = Dynamics->Mass;
+		ActorSetup.Inertia = (Dynamics->Mass / Mass) * Inertia;
+	}
+
+	ActorSetup.CoMTransform = CoMTransform;
+	ActorSetup.Geometry = MoveTemp(BodyGeom);
+	ActorSetup.Shapes = MoveTemp(BodyShapes);
+	ActorSetup.Material = MakeUnique<Chaos::FChaosPhysicsMaterial>();
+
+	ImmediatePhysics::FActorHandle* ActorHandle = Simulation->CreateActor(MoveTemp(ActorSetup));
+	ActorHandle->SetWorldTransform(BodyRelSimSpaceTM);
+
 
 	// Creation is guaranteed to return something
 	check(ActorHandle);
@@ -486,9 +658,12 @@ ImmediatePhysics::FActorHandle* FRigPhysicsSimulation::CreateBody(
 	}
 #endif
 
-	Simulation->AddToCollidingPairs(ActorHandle);
+	if (bGeometryCreated)
+	{
+		Simulation->AddToCollidingPairs(ActorHandle);
+	}
 
-	if (bAddedCollision && Dynamics)
+	if (bGeometryCreated && Dynamics)
 	{
 		// Note that particles are always created disabled. They will simulate when disabled,
 		// but won't collide!
diff --git a/Engine/Plugins/Experimental/PhysicsControl/Source/ControlRigPhysics/Public/RigPhysicsData.h b/Engine/Plugins/Experimental/PhysicsControl/Source/ControlRigPhysics/Public/RigPhysicsData.h
--- a/Engine/Plugins/Experimental/PhysicsControl/Source/ControlRigPhysics/Public/RigPhysicsData.h
+++ b/Engine/Plugins/Experimental/PhysicsControl/Source/ControlRigPhysics/Public/RigPhysicsData.h
@@ -447,15 +447,7 @@ struct CONTROLRIGPHYSICS_API FRigPhysicsDynamics
 {
 	GENERATED_BODY()
 
-	friend FArchive& operator <<(FArchive& Ar, FRigPhysicsDynamics& Data)
-	{
-		Ar << Data.Mass;
-		Ar << Data.bOverrideCentreOfMass;
-		Ar << Data.CentreOfMassOverride;
-		Ar << Data.bOverrideMomentsOfInertia;
-		Ar << Data.MomentsOfInertiaOverride;
-		return Ar;
-	}
+	friend FArchive& operator <<(FArchive& Ar, FRigPhysicsDynamics& Data);
 
 	// The total mass of this body.
 	UPROPERTY(EditAnywhere, Category = Dynamics)
@@ -474,6 +466,12 @@ struct CONTROLRIGPHYSICS_API FRigPhysicsDynamics
 
 	UPROPERTY(EditAnywhere, Category = Dynamics, meta = (EditCondition = "bOverrideMomentsOfInertia"))
 	FVector MomentsOfInertiaOverride = FVector(1.0f, 1.0f, 1.0f);
+
+	UPROPERTY(EditAnywhere, Category = Dynamics, meta = (ClampMin = "0.0"))
+	float LinearDamping = 0.0f;
+
+	UPROPERTY(EditAnywhere, Category = Dynamics, meta = (ClampMin = "0.0"))
+	float AngularDamping = 0.0f;
 };
 
 // This represents the motor drive associated with a physics joint, that can drive the bodies
diff --git a/Engine/Plugins/Experimental/PhysicsControl/Source/PhysicsControl/Public/PhysicsControlObjectVersion.h b/Engine/Plugins/Experimental/PhysicsControl/Source/PhysicsControl/Public/PhysicsControlObjectVersion.h
--- a/Engine/Plugins/Experimental/PhysicsControl/Source/PhysicsControl/Public/PhysicsControlObjectVersion.h
+++ b/Engine/Plugins/Experimental/PhysicsControl/Source/PhysicsControl/Public/PhysicsControlObjectVersion.h
@@ -30,6 +30,8 @@ struct PHYSICSCONTROL_API FPhysicsControlObjectVersion
 
 		ControlRigSupportNoCollisionBodies,
 
+		ControlRigSupportBodyDamping,
+
 		// -----<new versions can be added above this line>-------------------------------------------------
 		VersionPlusOne,
 		LatestVersion = VersionPlusOne - 1
diff --git a/Engine/Source/Runtime/Engine/Public/Physics/Experimental/ChaosInterfaceUtils.h b/Engine/Source/Runtime/Engine/Public/Physics/Experimental/ChaosInterfaceUtils.h
--- a/Engine/Source/Runtime/Engine/Public/Physics/Experimental/ChaosInterfaceUtils.h
+++ b/Engine/Source/Runtime/Engine/Public/Physics/Experimental/ChaosInterfaceUtils.h
@@ -28,6 +28,6 @@ namespace ChaosInterface
 	 * Generate the mass properties for a set of shapes in the space of the shapes' owner. 
 	 * Rotation will be built into the inertia matrix (it may not be diagonal) and RotationOfMass will be identity.
 	*/
-	void CalculateMassPropertiesFromShapeCollection(Chaos::FMassProperties& OutProperties, const TArray<FPhysicsShapeHandle>& InShapes, float InDensityKGPerCM);
-	void CalculateMassPropertiesFromShapeCollection(Chaos::FMassProperties& OutProperties, const Chaos::FShapesArray& InShapes, const TArray<bool>& bContributesToMass, float InDensityKGPerCM);
+	void ENGINE_API CalculateMassPropertiesFromShapeCollection(Chaos::FMassProperties& OutProperties, const TArray<FPhysicsShapeHandle>& InShapes, float InDensityKGPerCM);
+	void ENGINE_API CalculateMassPropertiesFromShapeCollection(Chaos::FMassProperties& OutProperties, const Chaos::FShapesArray& InShapes, const TArray<bool>& bContributesToMass, float InDensityKGPerCM);
 }
