Improve the code to generate control rig physics on worker threads. Add support for materials. Allow Density to be used instead of mass.

diff --git a/Engine/Plugins/Experimental/PhysicsControl/Source/ControlRigPhysics/Private/RigPhysicsData.cpp b/Engine/Plugins/Experimental/PhysicsControl/Source/ControlRigPhysics/Private/RigPhysicsData.cpp
--- a/Engine/Plugins/Experimental/PhysicsControl/Source/ControlRigPhysics/Private/RigPhysicsData.cpp
+++ b/Engine/Plugins/Experimental/PhysicsControl/Source/ControlRigPhysics/Private/RigPhysicsData.cpp
@@ -352,6 +352,10 @@ FArchive& operator <<(FArchive& Ar, FRigPhysicsCollision& Data)
 	Ar << Data.Boxes;
 	Ar << Data.Spheres;
 	Ar << Data.Capsules;
+	if (Ar.CustomVer(FPhysicsControlObjectVersion::GUID) >= FPhysicsControlObjectVersion::ControlRigCollisionHasMaterial)
+	{
+		Ar << Data.Material;
+	}
 	return Ar;
 }
 
@@ -382,7 +386,11 @@ void FRigPhysicsBodySolverSettings::OnRigHierarchyKeyChanged(
 //======================================================================================================================
 FArchive& operator <<(FArchive& Ar, FRigPhysicsDynamics& Data)
 {
-	Ar << Data.Mass;
+	if (Ar.CustomVer(FPhysicsControlObjectVersion::GUID) >= FPhysicsControlObjectVersion::ControlRigBodyDynamicsHasDensity)
+	{
+		Ar << Data.Density;
+	}
+	Ar << Data.MassOverride;
 	Ar << Data.bOverrideCentreOfMass;
 	Ar << Data.CentreOfMassOverride;
 	Ar << Data.bOverrideMomentsOfInertia;
@@ -394,3 +402,13 @@ FArchive& operator <<(FArchive& Ar, FRigPhysicsDynamics& Data)
 	}
 	return Ar;
 }
+
+//======================================================================================================================
+FArchive& operator <<(FArchive& Ar, FRigPhysicsMaterial& Data)
+{
+	Ar << Data.Friction;
+	Ar << Data.Restitution;
+	Ar << Data.FrictionCombineMode;
+	Ar << Data.RestitutionCombineMode;
+	return Ar;
+}
\ No newline at end of file
diff --git a/Engine/Plugins/Experimental/PhysicsControl/Source/ControlRigPhysics/Private/RigPhysicsExecution.cpp b/Engine/Plugins/Experimental/PhysicsControl/Source/ControlRigPhysics/Private/RigPhysicsExecution.cpp
--- a/Engine/Plugins/Experimental/PhysicsControl/Source/ControlRigPhysics/Private/RigPhysicsExecution.cpp
+++ b/Engine/Plugins/Experimental/PhysicsControl/Source/ControlRigPhysics/Private/RigPhysicsExecution.cpp
@@ -15,6 +15,7 @@
 #include "PhysicsEngine/PhysicsAsset.h"
 #include "PhysicsEngine/SkeletalBodySetup.h"
 #include "PhysicsEngine/PhysicsConstraintTemplate.h"
+#include "PhysicalMaterials/PhysicalMaterial.h"
 
 #include "ControlRig.h"
 #if WITH_EDITOR
@@ -324,9 +325,21 @@ FRigUnit_HierarchyInstantiateFromPhysicsAsset_Execute()
 			Collision.Capsules.Add(Capsule);
 		}
 
+		const UPhysicalMaterial* Material = SkeletalBodySetup->GetPhysMaterial();
+		Material = Material ? Material : GetDefault<UPhysicalMaterial>();
+		if (Material)
+		{
+			Collision.Material.Friction = FMath::Max(Material->Friction, Material->StaticFriction);
+			Collision.Material.Restitution = Material->Restitution;
+			Collision.Material.FrictionCombineMode = (ERigPhysicsCombineMode) 
+				Material->FrictionCombineMode.GetValue();
+			Collision.Material.RestitutionCombineMode = (ERigPhysicsCombineMode) 
+				Material->RestitutionCombineMode.GetValue();
+		}
+
 		// Dynamics
 		FRigPhysicsDynamics Dynamics;
-		Dynamics.Mass = SkeletalBodySetup->CalculateMass();
+		Dynamics.MassOverride = SkeletalBodySetup->CalculateMass();
 		Dynamics.LinearDamping = SkeletalBodySetup->DefaultInstance.LinearDamping;
 		Dynamics.AngularDamping = SkeletalBodySetup->DefaultInstance.AngularDamping;
 
diff --git a/Engine/Plugins/Experimental/PhysicsControl/Source/ControlRigPhysics/Private/RigPhysicsSimulation.cpp b/Engine/Plugins/Experimental/PhysicsControl/Source/ControlRigPhysics/Private/RigPhysicsSimulation.cpp
--- a/Engine/Plugins/Experimental/PhysicsControl/Source/ControlRigPhysics/Private/RigPhysicsSimulation.cpp
+++ b/Engine/Plugins/Experimental/PhysicsControl/Source/ControlRigPhysics/Private/RigPhysicsSimulation.cpp
@@ -55,8 +55,6 @@ constexpr int32 ConstraintParentIndex = 1;
 FRigPhysicsSimulation::FRigPhysicsSimulation(TObjectPtr<UControlRig> InOwningControlRig)
 	: FRigPhysicsSimulationBase(FRigPhysicsSimulation::StaticStruct()), OwningControlRig(InOwningControlRig)
 {
-	// Danny TODO temporary hackery to allow creation in worker thread
-	Chaos::CVars::bEnableAsyncInitBody = true;
 }
 
 //======================================================================================================================
@@ -140,7 +138,7 @@ void FRigPhysicsSimulation::InitialiseSimulation(const FRigPhysicsSolverComponen
 
 	// Always create a world actor at the origin, for attaching controls to. 
 	SimulationActorHandle = CreateBody(
-		FName(TEXT("Simulation")), FRigPhysicsCollision(), nullptr, FTransform());
+		FName(TEXT("Simulation")), FRigPhysicsCollision(), nullptr, nullptr, FTransform());
 }
 
 //======================================================================================================================
@@ -386,7 +384,6 @@ static void SetCommonProperties(const FRigPhysicsCollisionShape& Shape, FKShapeE
 //======================================================================================================================
 static Chaos::EJointMotionType GetMotionTypeFromLimitValue(const float Value)
 {
-	// Danny TODO This +ve/-ve thing isn't nice for users - use an enum
 	if (Value > 0.0f)
 	{
 		return Chaos::EJointMotionType::Limited;
@@ -402,10 +399,10 @@ static Chaos::EJointMotionType GetMotionTypeFromLimitValue(const float Value)
 }
 
 //======================================================================================================================
-// Danny TODO check all of this
 static bool CreateGeometry(
 	const FRigPhysicsCollision&               Collision,
 	const FRigPhysicsDynamics*                Dynamics,
+	const Chaos::FReal                        Density,
 	Chaos::FReal&                             OutMass, 
 	Chaos::FVec3&                             OutInertia, 
 	Chaos::FRigidTransform3&                  OutCoMTransform, 
@@ -418,7 +415,8 @@ static bool CreateGeometry(
 	OutInertia = FVector::ZeroVector;
 	OutCoMTransform = FTransform::Identity;
 
-	// Set the filter to collide with everything (we use a broad phase that only contains particle pairs that are explicitly set to collide)
+	// Set the filter to collide with everything (we use a broad phase that only contains particle
+	// pairs that are explicitly set to collide)
 	FBodyCollisionData BodyCollisionData;
 	// @todo(chaos): we need an API for setting up filters
 	BodyCollisionData.CollisionFilterData.SimFilter.Word1 = 0xFFFF;
@@ -460,9 +458,9 @@ static bool CreateGeometry(
 
 	FGeometryAddParams AddParams;
 	AddParams.CollisionData = BodyCollisionData;
-	AddParams.CollisionTraceType = ECollisionTraceFlag::CTF_UseDefault;
+	AddParams.CollisionTraceType = ECollisionTraceFlag::CTF_UseSimpleAsComplex;
 	AddParams.Scale = FVector(1, 1, 1);
-	AddParams.LocalTransform = FTransform::Identity;
+	AddParams.LocalTransform = FTransform::Identity; // How are these used? We will just set TM afterwards anyway
 	AddParams.WorldTransform = FTransform::Identity;
 	AddParams.Geometry = &AggGeom;
 
@@ -475,10 +473,11 @@ static bool CreateGeometry(
 		return false;
 	}
 
+	// Calculate mass properties, if we have dynamics
 	if (Dynamics)
 	{
-		// Whether each shape contributes to mass. @todo(chaos): it would be easier if
-		// ComputeMassProperties knew how to extract this info. Maybe it should be a flag in PerShapeData
+		// Whether each shape contributes to mass. It would be easier if ComputeMassProperties knew
+		// how to extract this info. Maybe it should be a flag in PerShapeData
 		TArray<bool> bContributesToMass;
 		bContributesToMass.Reserve(Shapes.Num());
 		for (int32 ShapeIndex = 0; ShapeIndex < Shapes.Num(); ++ShapeIndex)
@@ -489,12 +488,9 @@ static bool CreateGeometry(
 			bContributesToMass.Add(bHasMass);
 		}
 
-		// bInertiaScaleIncludeMass = true is to match legacy physics behaviour. This will scale the inertia by the change in mass (density x volumescale) 
-		// as well as the dimension change even though we don't actually change the mass.
-		const float density = 0.001f; // Danny TODO
 		Chaos::FMassProperties MassProperties;
 		ChaosInterface::CalculateMassPropertiesFromShapeCollection(
-			MassProperties, Shapes, bContributesToMass, density);
+			MassProperties, Shapes, bContributesToMass, Density);
 
 		OutMass = MassProperties.Mass;
 		OutInertia = MassProperties.InertiaTensor.GetDiagonal();
@@ -521,96 +517,18 @@ static bool CreateGeometry(
 
 //======================================================================================================================
 ImmediatePhysics::FActorHandle* FRigPhysicsSimulation::CreateBody(
-	const FName                 BodyName,
-	const FRigPhysicsCollision& Collision,
-	const FRigPhysicsDynamics*  Dynamics, 
-	const FTransform&           BodyRelSimSpaceTM) const
+	const FName                        BodyName,
+	const FRigPhysicsCollision&        Collision,
+	const FRigPhysicsDynamics*         Dynamics,
+	const FPhysicsControlModifierData* BodyData,
+	const FTransform&                  BodyRelSimSpaceTM) const
 {
-	// Danny TODO Use the correct materials (potentially tricky since we may be called on
-	// any thread - see the RBAN code for reference).
-
-	bool bAddedCollision = false;
-
-	// Disabling this code, but keeping it around for now because we need to double check that the
-	// actors we create directly (in the live code) match what was being created in the
-	// FBodyInstance version, which won't run properly on worker threads.
-#if 0
-	// This appears to be the only way
-	UBodySetup* BodySetup = NewObject<UBodySetup>();
-	BodySetup->CollisionTraceFlag = CTF_UseSimpleAsComplex;
-
-	// Note that the base class (FKShapeElem) for the shapes has got additional
-	// functionality such as enable/disable collision etc.
-
-	for (const FRigPhysicsCollisionBox& Shape : Collision.Boxes)
-	{
-		FKBoxElem Elem(Shape.Extents.X, Shape.Extents.Y, Shape.Extents.Z);
-		SetCommonProperties(Shape, Elem);
-		Elem.Center = Shape.TM.GetTranslation();
-		Elem.Rotation = Shape.TM.Rotator();
-		BodySetup->AggGeom.BoxElems.Add(Elem);
-		bAddedCollision = true;
-	}
-
-	for (const FRigPhysicsCollisionSphere& Shape : Collision.Spheres)
-	{
-		FKSphereElem Elem(Shape.Radius);
-		SetCommonProperties(Shape, Elem);
-		Elem.Center = Shape.TM.GetTranslation();
-		// Note that there is no rotation
-		BodySetup->AggGeom.SphereElems.Add(Elem);
-		bAddedCollision = true;
-	}
-
-	for (const FRigPhysicsCollisionCapsule& Shape : Collision.Capsules)
-	{
-		FKSphylElem Elem(Shape.Radius, Shape.Length);
-		SetCommonProperties(Shape, Elem);
-		Elem.Center = Shape.TM.GetTranslation();
-		Elem.Rotation = Shape.TM.Rotator();
-		BodySetup->AggGeom.SphylElems.Add(Elem);
-		bAddedCollision = true;
-	}
-
-	// This is a bit odd, but in order to create a low-level immediate physics representation, you
-	// have to create a high-level temporary one. Danny TODO avoid using BodyInstance and BodySetup.
-	//
-	// See USkeletalMeshComponent::InstantiatePhysicsAsset_Internal for how it's done in RBAN (which
-	// makes body instances, converts to immediate physics actors, and then discards the body instances).
-	FBodyInstance BodyInstance;
-	if (Dynamics)
-	{
-		BodyInstance.bSimulatePhysics = true;
-		BodyInstance.bEnableGravity = true;
-		BodyInstance.SetMassOverride(Dynamics->Mass);
-		BodyInstance.LinearDamping = 0;
-		BodyInstance.AngularDamping = 0;
-	}
-	else
-	{
-		BodyInstance.bSimulatePhysics = false;
-	}
-
-	// Without this flag then creation only works when it is on the game thread.
-	BodySetup->bNeverNeedsCookedCollisionData = true;
-
-	BodyInstance.InitBody(BodySetup, BodyRelSimSpaceTM, nullptr, nullptr);
-
-	ImmediatePhysics::EActorType ActorType = Dynamics ? ImmediatePhysics::EActorType::DynamicActor : ImmediatePhysics::EActorType::KinematicActor;
-	ImmediatePhysics::FActorSetup&& ActorSetupOld = ImmediatePhysics::MakeActorSetup(ActorType, &BodyInstance, BodyRelSimSpaceTM);
-	ImmediatePhysics::FActorHandle* ActorHandleOld = Simulation->CreateActor(MoveTemp(ActorSetupOld));
-#endif
 	ImmediatePhysics::FActorSetup ActorSetup;
 
-	ActorSetup.ActorType = Dynamics ?
-		ImmediatePhysics::EActorType::DynamicActor : ImmediatePhysics::EActorType::KinematicActor;
-
 	if (Dynamics)
 	{
 		ActorSetup.ActorType = ImmediatePhysics::EActorType::DynamicActor;
 		ActorSetup.bEnableGravity = true;
-		ActorSetup.LinearDamping = 0;
-		ActorSetup.AngularDamping = 0;
 		ActorSetup.LinearDamping = Dynamics->LinearDamping;
 		ActorSetup.AngularDamping = Dynamics->AngularDamping;
 	}
@@ -619,36 +537,57 @@ ImmediatePhysics::FActorHandle* FRigPhysicsSimulation::CreateBody(
 		ActorSetup.ActorType = ImmediatePhysics::EActorType::KinematicActor;
 	}
 
-
+	if (BodyData)
+	{
+		ActorSetup.bUpdateKinematicFromSimulation = BodyData->bUpdateKinematicFromSimulation;
+	}
+	else
+	{
+		ActorSetup.bUpdateKinematicFromSimulation = false;
+	}
 
 	Chaos::FVec3 Inertia;
 	Chaos::FRigidTransform3 CoMTransform;
 	Chaos::FReal Mass;
-
 	Chaos::FImplicitObjectPtr BodyGeom;
 	TArray<TUniquePtr<Chaos::FPerShapeData>> BodyShapes;
+	Chaos::FReal Density = Dynamics ? Dynamics->Density : 1.0f;
+	// Convert from g/cm^3 to kg/cm^3
+	Density *= 1e-6;
+
 	bool bGeometryCreated = CreateGeometry(
-		Collision, Dynamics, Mass, Inertia, CoMTransform, BodyGeom, BodyShapes);
+		Collision, Dynamics, Density, Mass, Inertia, CoMTransform, BodyGeom, BodyShapes);
 
+	// We will have created with an arbitrary density - adjust to result in the desired mass.
 	ActorSetup.Mass = Mass;
 	ActorSetup.Inertia = Inertia;
-	if (Mass > 0.0f && Dynamics)
+	if (Mass > 0.0f && Dynamics && Dynamics->MassOverride > 0.0f)
 	{
-		ActorSetup.Mass = Dynamics->Mass;
-		ActorSetup.Inertia = (Dynamics->Mass / Mass) * Inertia;
+		ActorSetup.Mass = Dynamics->MassOverride;
+		ActorSetup.Inertia = (Dynamics->MassOverride / Mass) * Inertia;
 	}
 
+	ActorSetup.Transform = BodyRelSimSpaceTM;
 	ActorSetup.CoMTransform = CoMTransform;
 	ActorSetup.Geometry = MoveTemp(BodyGeom);
 	ActorSetup.Shapes = MoveTemp(BodyShapes);
+
 	ActorSetup.Material = MakeUnique<Chaos::FChaosPhysicsMaterial>();
+	ActorSetup.Material->Friction = Collision.Material.Friction;
+	ActorSetup.Material->StaticFriction = Collision.Material.Friction;
+	ActorSetup.Material->Restitution = Collision.Material.Restitution;
+	ActorSetup.Material->FrictionCombineMode =
+		(Chaos::FChaosPhysicsMaterial::ECombineMode)Collision.Material.FrictionCombineMode;
+	ActorSetup.Material->RestitutionCombineMode =
+		(Chaos::FChaosPhysicsMaterial::ECombineMode)Collision.Material.RestitutionCombineMode;
 
 	ImmediatePhysics::FActorHandle* ActorHandle = Simulation->CreateActor(MoveTemp(ActorSetup));
-	ActorHandle->SetWorldTransform(BodyRelSimSpaceTM);
-
-
-	// Creation is guaranteed to return something
-	check(ActorHandle);
+	if (!ActorHandle)
+	{
+		UE_LOG(LogRigPhysics, Warning,
+			TEXT("Unable to create body %s"), *BodyName.ToString());
+		return nullptr;
+	}
 
 	ActorHandle->SetName(BodyName);
 #if CHAOS_DEBUG_NAME
@@ -661,13 +600,12 @@ ImmediatePhysics::FActorHandle* FRigPhysicsSimulation::CreateBody(
 	if (bGeometryCreated)
 	{
 		Simulation->AddToCollidingPairs(ActorHandle);
-	}
-
-	if (bGeometryCreated && Dynamics)
-	{
-		// Note that particles are always created disabled. They will simulate when disabled,
-		// but won't collide!
-		ActorHandle->SetEnabled(true);
+		if (Dynamics)
+		{
+			// Note that particles are always created disabled. They will simulate when disabled,
+			// but won't collide!
+			ActorHandle->SetEnabled(true);
+		}
 	}
 
 	return ActorHandle;
@@ -736,10 +674,9 @@ void FRigPhysicsSimulation::InstantiateSolverCollision(
 			SolverComponent->SolverSettings, FTransform());
 
 		CollisionActorHandle = CreateBody(
-			FName(TEXT("Environment")), SolverComponent->SolverSettings.Collision, nullptr, BodyRelSimSpaceTM);
+			FName(TEXT("Environment")), SolverComponent->SolverSettings.Collision, nullptr, nullptr, BodyRelSimSpaceTM);
 	}
 
-
 	// Add no-collision pairs
 	TArray<ImmediatePhysics_Chaos::FSimulation::FIgnorePair> ChaosIgnorePairs;
 
@@ -791,7 +728,8 @@ void FRigPhysicsSimulation::InstantiatePhysicsBodies(
 				FTransform SourceComponentSpaceTM = Hierarchy->GetGlobalTransform(SourceKey);
 				const FTransform SourceSimulationSpaceTM =
 					ConvertComponentSpaceTransformToSimSpace(SolverComponent->SolverSettings, SourceComponentSpaceTM);
-				Record.ActorHandle = CreateBody(ComponentKey.ElementKey.Name, Collision, &Dynamics, SourceSimulationSpaceTM);
+				Record.ActorHandle = CreateBody(
+					ComponentKey.ElementKey.Name, Collision, &Dynamics, &BodyData, SourceSimulationSpaceTM);
 			}
 
 			for (const FRigComponentKey& NoCollisionKey : PhysicsComponent->NoCollisionBodies)
@@ -857,7 +795,6 @@ void FRigPhysicsSimulation::InstantiatePhysicsJoints(
 				Chaos::FPBDJointSettings JointSettings;
 				if (JointData.bAutoCalculateChildOffset)
 				{
-					// Danny TODO - is the "non-override" actually always simply the identity TM?
 					JointSettings.ConnectorTransforms[0] = FTransform();
 				}
 				JointSettings.ConnectorTransforms[0] =
@@ -865,8 +802,6 @@ void FRigPhysicsSimulation::InstantiatePhysicsJoints(
 
 				if (JointData.bAutoCalculateParentOffset)
 				{
-					// Danny TODO is it right to use the TM of the child bone relative to parent bone?
-					// What about if the body has got a TM too?
 					JointSettings.ConnectorTransforms[1] = Hierarchy->GetLocalTransform(ChildBoneKey, true);
 				}
 				JointSettings.ConnectorTransforms[1] =
diff --git a/Engine/Plugins/Experimental/PhysicsControl/Source/ControlRigPhysics/Private/RigPhysicsSimulation.h b/Engine/Plugins/Experimental/PhysicsControl/Source/ControlRigPhysics/Private/RigPhysicsSimulation.h
--- a/Engine/Plugins/Experimental/PhysicsControl/Source/ControlRigPhysics/Private/RigPhysicsSimulation.h
+++ b/Engine/Plugins/Experimental/PhysicsControl/Source/ControlRigPhysics/Private/RigPhysicsSimulation.h
@@ -226,10 +226,11 @@ private:
 
 	// Creates an actor with collision. This will by dynamic if Dynamics is valid, or otherwise kinematic
 	ImmediatePhysics::FActorHandle* CreateBody(
-		const FName                 BodyName,
-		const FRigPhysicsCollision& Collision,
-		const FRigPhysicsDynamics*  Dynamics,
-		const FTransform&           BodyRelSimSpaceTM) const;
+		const FName                        BodyName,
+		const FRigPhysicsCollision&        Collision,
+		const FRigPhysicsDynamics*         Dynamics,
+		const FPhysicsControlModifierData* BodyData,
+		const FTransform&                  BodyRelSimSpaceTM) const;
 
 	// This releases the actors etc in the simulation, and destroys the simulation, but it doesn't
 	// change our records (apart from resetting pointers to actors etc).
diff --git a/Engine/Plugins/Experimental/PhysicsControl/Source/ControlRigPhysics/Private/RigPhysicsSolverComponent.cpp b/Engine/Plugins/Experimental/PhysicsControl/Source/ControlRigPhysics/Private/RigPhysicsSolverComponent.cpp
--- a/Engine/Plugins/Experimental/PhysicsControl/Source/ControlRigPhysics/Private/RigPhysicsSolverComponent.cpp
+++ b/Engine/Plugins/Experimental/PhysicsControl/Source/ControlRigPhysics/Private/RigPhysicsSolverComponent.cpp
@@ -25,6 +25,20 @@ void FRigPhysicsSolverComponent::Load(FArchive& Ar)
 	Ar << SimulationSpaceSettings;
 }
 
+//======================================================================================================================
+void FRigPhysicsSolverComponent::OnAddedToHierarchy(URigHierarchy* InHierarchy, URigHierarchyController* InController)
+{
+	if (!IsProcedural())
+	{
+		// Default the material here to have friction and restitution. Then the interactions are
+		// easily adjusted on the dynamic bodies.
+		SolverSettings.Collision.Material.Friction = 1.0f;
+		SolverSettings.Collision.Material.Restitution = 1.0f;
+
+		InHierarchy->Notify(ERigHierarchyNotification::ComponentContentChanged, this);
+	}
+}
+
 //======================================================================================================================
 #if WITH_EDITOR
 const FSlateIcon& FRigPhysicsSolverComponent::GetIconForUI() const
diff --git a/Engine/Plugins/Experimental/PhysicsControl/Source/ControlRigPhysics/Public/RigPhysicsData.h b/Engine/Plugins/Experimental/PhysicsControl/Source/ControlRigPhysics/Public/RigPhysicsData.h
--- a/Engine/Plugins/Experimental/PhysicsControl/Source/ControlRigPhysics/Public/RigPhysicsData.h
+++ b/Engine/Plugins/Experimental/PhysicsControl/Source/ControlRigPhysics/Public/RigPhysicsData.h
@@ -196,9 +196,44 @@ struct CONTROLRIGPHYSICS_API FRigPhysicsCollisionCapsule : public FRigPhysicsCol
 	float Length = RigPhysicsDefaultShapeSize;
 };
 
-/**
- * Collection of shapes that define the collision and (optionally) mass distribution of the body
- */
+// How to combine friction/restitution values. 
+UENUM(BlueprintType)
+enum class ERigPhysicsCombineMode : uint8
+{
+	Avg,
+	Min,
+	Multiply,
+	Max
+	// The values here must match those in Chaos::FChaosPhysicsMaterial::ECombineMode
+};
+
+// Material properties of a collision shapes.
+USTRUCT(BlueprintType)
+struct CONTROLRIGPHYSICS_API FRigPhysicsMaterial
+{
+	GENERATED_BODY()
+
+	friend FArchive& operator <<(FArchive& Ar, FRigPhysicsMaterial& Data);
+
+	// Friction is simplified here - just one value used for static and dynamic friction. 
+	UPROPERTY(EditAnywhere, Category = Material, meta = (ClampMin = "0.0"))
+	float Friction = 0.5f;
+
+	UPROPERTY(EditAnywhere, Category = Material, meta = (ClampMin = "0.0", ClampMax = "1.0"))
+	float Restitution = 0.0f;
+
+	// How to combine friction values. If the materials have different combine modes, then Max is
+	// used over Multiply, Multiply over Min and Min over Avg.
+	UPROPERTY(EditAnywhere, Category = Material)
+	ERigPhysicsCombineMode FrictionCombineMode = ERigPhysicsCombineMode::Multiply;
+
+	// How to combine restitution values. If the materials have different combine modes, then Max is
+	// used over Multiply, Multiply over Min and Min over Avg.
+	UPROPERTY(EditAnywhere, Category = Material)
+	ERigPhysicsCombineMode RestitutionCombineMode = ERigPhysicsCombineMode::Multiply;
+};
+
+// Collection of shapes that define the collision and (optionally) mass distribution of the body
 USTRUCT(BlueprintType)
 struct CONTROLRIGPHYSICS_API FRigPhysicsCollision
 {
@@ -215,6 +250,9 @@ struct CONTROLRIGPHYSICS_API FRigPhysicsCollision
 	UPROPERTY(EditAnywhere, Category = Shapes)
 	TArray<FRigPhysicsCollisionCapsule> Capsules;
 
+	UPROPERTY(EditAnywhere, Category = Shapes)
+	FRigPhysicsMaterial Material;
+
 	bool IsEmpty() const { return Boxes.IsEmpty() && Spheres.IsEmpty() && Capsules.IsEmpty(); }
 	int32 NumShapes() const { return Boxes.Num() + Capsules.Num() + Spheres.Num(); }
 };
@@ -449,9 +487,13 @@ struct CONTROLRIGPHYSICS_API FRigPhysicsDynamics
 
 	friend FArchive& operator <<(FArchive& Ar, FRigPhysicsDynamics& Data);
 
-	// The total mass of this body.
+	// Density - units g per cm^3 so this defaults to "water"
+	UPROPERTY(EditAnywhere, Category = Dynamics, meta = (EditCondition = "MassOverride <= 0"))
+	float Density = 1.0f;
+
+	// The total mass of this body - will override density if positive
 	UPROPERTY(EditAnywhere, Category = Dynamics)
-	float Mass = 1.0f;
+	float MassOverride = 1.0f;
 
 	// Use this to specify the location of the centre of mass
 	UPROPERTY(EditAnywhere, Category = Dynamics)
diff --git a/Engine/Plugins/Experimental/PhysicsControl/Source/ControlRigPhysics/Public/RigPhysicsExecution.h b/Engine/Plugins/Experimental/PhysicsControl/Source/ControlRigPhysics/Public/RigPhysicsExecution.h
--- a/Engine/Plugins/Experimental/PhysicsControl/Source/ControlRigPhysics/Public/RigPhysicsExecution.h
+++ b/Engine/Plugins/Experimental/PhysicsControl/Source/ControlRigPhysics/Public/RigPhysicsExecution.h
@@ -41,6 +41,10 @@ struct CONTROLRIGPHYSICS_API FRigUnit_AddPhysicsSolver : public FRigUnit_Physics
 	FRigUnit_AddPhysicsSolver()
 	{
 		Owner.Type = ERigElementType::Bone;
+		// Default the material here to have friction and restitution. Then the interactions are
+		// easily adjusted on the dynamic bodies.
+		SolverSettings.Collision.Material.Friction = 1.0f;
+		SolverSettings.Collision.Material.Restitution = 1.0f;
 	}
 
 	RIGVM_METHOD()
diff --git a/Engine/Plugins/Experimental/PhysicsControl/Source/ControlRigPhysics/Public/RigPhysicsSolverComponent.h b/Engine/Plugins/Experimental/PhysicsControl/Source/ControlRigPhysics/Public/RigPhysicsSolverComponent.h
--- a/Engine/Plugins/Experimental/PhysicsControl/Source/ControlRigPhysics/Public/RigPhysicsSolverComponent.h
+++ b/Engine/Plugins/Experimental/PhysicsControl/Source/ControlRigPhysics/Public/RigPhysicsSolverComponent.h
@@ -32,6 +32,8 @@ public:
 	virtual const FSlateIcon& GetIconForUI() const override;
 #endif
 
+	virtual void OnAddedToHierarchy(URigHierarchy* InHierarchy, URigHierarchyController* InController) override;
+
 	virtual FName GetDefaultComponentName() const override { return GetDefaultName(); }
 
 	static FName GetDefaultName() { return TEXT("PhysicsSolver"); }
diff --git a/Engine/Plugins/Experimental/PhysicsControl/Source/PhysicsControl/Public/PhysicsControlObjectVersion.h b/Engine/Plugins/Experimental/PhysicsControl/Source/PhysicsControl/Public/PhysicsControlObjectVersion.h
--- a/Engine/Plugins/Experimental/PhysicsControl/Source/PhysicsControl/Public/PhysicsControlObjectVersion.h
+++ b/Engine/Plugins/Experimental/PhysicsControl/Source/PhysicsControl/Public/PhysicsControlObjectVersion.h
@@ -32,6 +32,10 @@ struct PHYSICSCONTROL_API FPhysicsControlObjectVersion
 
 		ControlRigSupportBodyDamping,
 
+		ControlRigCollisionHasMaterial,
+
+		ControlRigBodyDynamicsHasDensity,
+
 		// -----<new versions can be added above this line>-------------------------------------------------
 		VersionPlusOne,
 		LatestVersion = VersionPlusOne - 1
