Control Rig Physics - expose damping and collision functionality

diff --git a/Engine/Plugins/Experimental/PhysicsControl/Source/ControlRigPhysics/Private/RigPhysicsBodyComponent.cpp b/Engine/Plugins/Experimental/PhysicsControl/Source/ControlRigPhysics/Private/RigPhysicsBodyComponent.cpp
--- a/Engine/Plugins/Experimental/PhysicsControl/Source/ControlRigPhysics/Private/RigPhysicsBodyComponent.cpp
+++ b/Engine/Plugins/Experimental/PhysicsControl/Source/ControlRigPhysics/Private/RigPhysicsBodyComponent.cpp
@@ -118,7 +118,8 @@ void FRigPhysicsBodyComponent::OnAddedToHierarchy(URigHierarchy* InHierarchy, UR
 }
 
 //======================================================================================================================
-void FRigPhysicsBodyComponent::AutoCalculateCollision(URigHierarchy* InHierarchy, float MinAspectRatio)
+void FRigPhysicsBodyComponent::AutoCalculateCollision(
+	URigHierarchy* InHierarchy, float MinAspectRatio, float MinSize)
 {
 	// Start clean
 	Collision = FRigPhysicsCollision();
@@ -146,7 +147,8 @@ void FRigPhysicsBodyComponent::AutoCalculateCollision(URigHierarchy* InHierarchy
 		{
 			// We have no children, and are co-located with our parent. Make a single, arbitrarily
 			// sized, shape.
-			Collision.Boxes.Add(FRigPhysicsCollisionBox(TM, FVector(10.0)));
+			float Size = FMath::Max(MinSize, 10.0f);
+			Collision.Boxes.Add(FRigPhysicsCollisionBox(TM, FVector(Size)));
 		}
 		else
 		{
@@ -154,8 +156,8 @@ void FRigPhysicsBodyComponent::AutoCalculateCollision(URigHierarchy* InHierarchy
 			// We're going to make a box that goes along the X axis
 			TM.SetRotation(FQuat::FindBetweenVectors(FVector(1, 0, 0), TM.GetTranslation()));
 			FVector Extents;
-			Extents.X = TM.GetTranslation().Length() * 2.0;
-			Extents.Y = Extents.X * MinAspectRatio;
+			Extents.X = FMath::Max(MinSize, TM.GetTranslation().Length() * 2.0);
+			Extents.Y = FMath::Max(MinSize, Extents.X * MinAspectRatio);
 			Extents.Z = Extents.Y;
 			Collision.Boxes.Add(FRigPhysicsCollisionBox(TM, Extents));
 		}
@@ -169,7 +171,7 @@ void FRigPhysicsBodyComponent::AutoCalculateCollision(URigHierarchy* InHierarchy
 		TM.SetRotation(FQuat::FindBetweenVectors(FVector(0, 0, 1), Points[1]));
 
 		float Length = TM.GetTranslation().Length() * 2.0f;
-		float Radius = Length * MinAspectRatio;
+		float Radius = FMath::Max(MinSize, Length * MinAspectRatio);
 		Length -= 2.0f * Radius;
 		Length = FMath::Max(0.0f, Length);
 		Collision.Capsules.Add(FRigPhysicsCollisionCapsule(TM, Radius, Length));
@@ -186,9 +188,9 @@ void FRigPhysicsBodyComponent::AutoCalculateCollision(URigHierarchy* InHierarchy
 		FTransform TM(Box.Center());
 		FVector Extents = Box.Extents() * 2.0; // The returned extents are half the extents!
 		double MaxExtent = Extents.GetAbsMax();
-		Extents.X = FMath::Max(Extents.X, MaxExtent * MinAspectRatio);
-		Extents.Y = FMath::Max(Extents.Y, MaxExtent * MinAspectRatio);
-		Extents.Z = FMath::Max(Extents.Z, MaxExtent * MinAspectRatio);
+		Extents.X = FMath::Max(MinSize, FMath::Max(Extents.X, MaxExtent * MinAspectRatio));
+		Extents.Y = FMath::Max(MinSize, FMath::Max(Extents.Y, MaxExtent * MinAspectRatio));
+		Extents.Z = FMath::Max(MinSize, FMath::Max(Extents.Z, MaxExtent * MinAspectRatio));
 		Collision.Boxes.Add(FRigPhysicsCollisionBox(TM, Extents));
 	}
 
diff --git a/Engine/Plugins/Experimental/PhysicsControl/Source/ControlRigPhysics/Private/RigPhysicsBodyExecution.cpp b/Engine/Plugins/Experimental/PhysicsControl/Source/ControlRigPhysics/Private/RigPhysicsBodyExecution.cpp
--- a/Engine/Plugins/Experimental/PhysicsControl/Source/ControlRigPhysics/Private/RigPhysicsBodyExecution.cpp
+++ b/Engine/Plugins/Experimental/PhysicsControl/Source/ControlRigPhysics/Private/RigPhysicsBodyExecution.cpp
@@ -67,7 +67,7 @@ FRigUnit_HierarchyAutoCalculateCollision_Execute()
 	if (FRigPhysicsBodyComponent* Component = Cast<FRigPhysicsBodyComponent>(
 		ExecuteContext.Hierarchy->FindComponent(PhysicsBodyComponentKey)))
 	{
-		Component->AutoCalculateCollision(ExecuteContext.Hierarchy, MinAspectRatio);
+		Component->AutoCalculateCollision(ExecuteContext.Hierarchy, MinAspectRatio, MinSize);
 	}
 }
 
@@ -267,3 +267,19 @@ FRigUnit_HierarchySetPhysicsBodyUpdateKinematicFromSimulation_Execute()
 		Component->BodyData.bUpdateKinematicFromSimulation = bUpdateKinematicFromSimulation;
 	}
 }
+
+//======================================================================================================================
+FRigUnit_HierarchySetPhysicsBodyDamping_Execute()
+{
+	DECLARE_SCOPE_HIERARCHICAL_COUNTER_RIGUNIT();
+	if (!ExecuteContext.Hierarchy)
+	{
+		return;
+	}
+	if (FRigPhysicsBodyComponent* Component = Cast<FRigPhysicsBodyComponent>(
+		ExecuteContext.Hierarchy->FindComponent(PhysicsBodyComponentKey)))
+	{
+		Component->Dynamics.LinearDamping = LinearDamping;
+		Component->Dynamics.AngularDamping = AngularDamping;
+	}
+}
diff --git a/Engine/Plugins/Experimental/PhysicsControl/Source/ControlRigPhysics/Private/RigPhysicsSimulation.cpp b/Engine/Plugins/Experimental/PhysicsControl/Source/ControlRigPhysics/Private/RigPhysicsSimulation.cpp
--- a/Engine/Plugins/Experimental/PhysicsControl/Source/ControlRigPhysics/Private/RigPhysicsSimulation.cpp
+++ b/Engine/Plugins/Experimental/PhysicsControl/Source/ControlRigPhysics/Private/RigPhysicsSimulation.cpp
@@ -965,6 +965,12 @@ void FRigPhysicsSimulation::UpdateBodyPrePhysics(
 			SolverComponent->SolverSettings, KinematicTargetCS);
 		Record.ActorHandle->SetKinematicTarget(KinematicTargetTM);
 	}
+	else
+	{
+		// Danny TODO move damping into BodyData - any PhysicsControl system should be able to use it
+		Record.ActorHandle->SetLinearDamping(PhysicsComponent->Dynamics.LinearDamping);
+		Record.ActorHandle->SetAngularDamping(PhysicsComponent->Dynamics.AngularDamping);
+	}
 }
 
 //======================================================================================================================
diff --git a/Engine/Plugins/Experimental/PhysicsControl/Source/ControlRigPhysics/Public/RigPhysicsBodyComponent.h b/Engine/Plugins/Experimental/PhysicsControl/Source/ControlRigPhysics/Public/RigPhysicsBodyComponent.h
--- a/Engine/Plugins/Experimental/PhysicsControl/Source/ControlRigPhysics/Public/RigPhysicsBodyComponent.h
+++ b/Engine/Plugins/Experimental/PhysicsControl/Source/ControlRigPhysics/Public/RigPhysicsBodyComponent.h
@@ -49,7 +49,7 @@ public:
 	// Removes any existing collision, and replaces it with a shape calculated from the joint
 	// positions (if possible). The shape will be a single box.
 	// @param MinAspectRatio the minimum box extent, as a proportion of the maximum box extent.
-	void AutoCalculateCollision(URigHierarchy* InHierarchy, float MinAspectRatio = 0.25f);
+	void AutoCalculateCollision(URigHierarchy* InHierarchy, float MinAspectRatio = 0.25f, float MinSize = 0.0f);
 
 	virtual void Save(FArchive& Ar) override;
 	virtual void Load(FArchive& Ar) override;
diff --git a/Engine/Plugins/Experimental/PhysicsControl/Source/ControlRigPhysics/Public/RigPhysicsBodyExecution.h b/Engine/Plugins/Experimental/PhysicsControl/Source/ControlRigPhysics/Public/RigPhysicsBodyExecution.h
--- a/Engine/Plugins/Experimental/PhysicsControl/Source/ControlRigPhysics/Public/RigPhysicsBodyExecution.h
+++ b/Engine/Plugins/Experimental/PhysicsControl/Source/ControlRigPhysics/Public/RigPhysicsBodyExecution.h
@@ -78,9 +78,15 @@ struct CONTROLRIGPHYSICS_API FRigUnit_HierarchyAutoCalculateCollision : public F
 	UPROPERTY(meta = (Input))
 	FRigComponentKey PhysicsBodyComponentKey;
 
-	// The minimum box extent, as a proportion of the maximum box extent.
+	// For boxes: The minimum box extent, as a proportion of the maximum box extent.
+	// For capsules: The minimum radius, as a proportion of the length (not including the radius)
 	UPROPERTY(meta = (Input))
 	float MinAspectRatio = 0.25f;
+
+	// For boxes: The minimum side length. 
+	// For capsules: The minimum radius
+	UPROPERTY(meta = (Input))
+	float MinSize = 0.0f;
 };
 
 // Sets the mass etc for a physics component body
@@ -406,3 +412,32 @@ struct CONTROLRIGPHYSICS_API FRigUnit_HierarchySetPhysicsBodyUpdateKinematicFrom
 	bool bUpdateKinematicFromSimulation = true;
 };
 
+/**
+ * If true, then kinematic objects will be written back from simulation to the bones. This only 
+ * necessary when either kinematic targets are being used, or when the target bone differs from the source bone.
+ * Danny TODO - check the implementation, and also can we auto-calculate this at runtime?
+ */
+USTRUCT(meta = (DisplayName = "Set Physics Body Damping", Varying))
+struct CONTROLRIGPHYSICS_API FRigUnit_HierarchySetPhysicsBodyDamping : public FRigUnit_PhysicsBaseMutable
+{
+	GENERATED_BODY()
+
+	FRigUnit_HierarchySetPhysicsBodyDamping()
+	{
+		PhysicsBodyComponentKey.ElementKey.Type = ERigElementType::Bone;
+		PhysicsBodyComponentKey.Name = FRigPhysicsBodyComponent::GetDefaultName();
+	}
+
+	RIGVM_METHOD()
+	virtual void Execute() override;
+
+	UPROPERTY(meta = (Input))
+	FRigComponentKey PhysicsBodyComponentKey;
+
+	UPROPERTY(meta = (Input))
+	float LinearDamping = 0;
+
+	UPROPERTY(meta = (Input))
+	float AngularDamping = 0;
+};
+
