Expose CVars to allow overriding the control rig physics step time - to allow profiling to not double count if being done at low frame rates (i.e. at frame rates that would trigger double stepping the simulation)
To use a single, variable timestep: "ControlRig.Physics.FixedTimeStepOverride 0"

diff --git a/Engine/Plugins/Experimental/PhysicsControl/Source/ControlRigPhysics/Private/RigPhysicsSimulation.cpp b/Engine/Plugins/Experimental/PhysicsControl/Source/ControlRigPhysics/Private/RigPhysicsSimulation.cpp
--- a/Engine/Plugins/Experimental/PhysicsControl/Source/ControlRigPhysics/Private/RigPhysicsSimulation.cpp
+++ b/Engine/Plugins/Experimental/PhysicsControl/Source/ControlRigPhysics/Private/RigPhysicsSimulation.cpp
@@ -48,6 +48,14 @@
 
 DEFINE_LOG_CATEGORY(LogRigPhysics);
 
+TAutoConsoleVariable<float> CVarControlRigPhysicsFixedTimeStepOverride(
+	TEXT("ControlRig.Physics.FixedTimeStepOverride"), -1.0f,
+	TEXT("-1.0 disables the override, so the timestep authored in the simulation settings will be used (which may or may not imply a fixed timestep). A value of 0 forces a variable timestep to be used. A +ve value is used to specify a fixed timestep."));
+
+TAutoConsoleVariable<int> CVarControlRigPhysicsMaxTimeStepsOverride(
+	TEXT("ControlRig.Physics.MaxTimeStepsOverride"), -1,
+	TEXT("-1 disables the override, so the max timesteps authored in the simulation settings will be used. A +ve value is used to specify the maximum number of timesteps."));
+
 constexpr int32 ConstraintChildIndex = 0;
 constexpr int32 ConstraintParentIndex = 1;
 
@@ -1412,9 +1420,14 @@ void FRigPhysicsSimulation::StepSimulation(
 	const FRigPhysicsSolverSettings& SolverSettings = SolverComponent->SolverSettings;
 	const FRigPhysicsSimulationSpaceSettings& SimulationSpaceSettings = SolverComponent->SimulationSpaceSettings;
 
+	float FixedTimeStep = CVarControlRigPhysicsFixedTimeStepOverride.GetValueOnAnyThread() < 0 
+		? SolverSettings.FixedTimeStep : CVarControlRigPhysicsFixedTimeStepOverride.GetValueOnAnyThread();
+	int MaxTimeSteps = CVarControlRigPhysicsMaxTimeStepsOverride.GetValueOnAnyThread() < 0
+		? SolverSettings.MaxTimeSteps : CVarControlRigPhysicsMaxTimeStepsOverride.GetValueOnAnyThread();
+
 	// Set settings that might change
 	Simulation->SetSolverSettings(
-		SolverSettings.FixedTimeStep,
+		FixedTimeStep,
 		999999999.0, // Danny TODO cull distance - expose this? Yuck
 		SolverSettings.MaxDepenetrationVelocity,
 		SolverSettings.bUseLinearJointSolver,
@@ -1452,7 +1465,7 @@ void FRigPhysicsSimulation::StepSimulation(
 		UpdatePrePhysics(ExecuteContext, SolverComponent, PhysicsDeltaTime);
 
 		Simulation->Simulate(
-			PhysicsDeltaTime, SolverSettings.MaxDeltaTime, SolverSettings.MaxTimeSteps,
+			PhysicsDeltaTime, SolverSettings.MaxDeltaTime, MaxTimeSteps,
 			SimulationSpaceData.Gravity, &ChaosJointSolverSettings);
 	}
 
