Control Rig Physics visualization of contacts. Better teleport detection. Fix the bounds expansion setting. Fix disabling of collision between body pairs.

diff --git a/Engine/Plugins/Experimental/PhysicsControl/Source/ControlRigPhysics/Private/RigPhysicsData.cpp b/Engine/Plugins/Experimental/PhysicsControl/Source/ControlRigPhysics/Private/RigPhysicsData.cpp
--- a/Engine/Plugins/Experimental/PhysicsControl/Source/ControlRigPhysics/Private/RigPhysicsData.cpp
+++ b/Engine/Plugins/Experimental/PhysicsControl/Source/ControlRigPhysics/Private/RigPhysicsData.cpp
@@ -202,6 +202,11 @@ FArchive& operator <<(FArchive& Ar, FRigPhysicsSimulationSpaceSettings& Data)
 	Ar << Data.MaxAngularAcceleration;
 	Ar << Data.LinearAccelerationThresholdForTeleport;
 	Ar << Data.AngularAccelerationThresholdForTeleport;
+	if (Ar.CustomVer(FPhysicsControlObjectVersion::GUID) >= FPhysicsControlObjectVersion::ControlRigDetectTeleportFromDistanceChange)
+	{
+		Ar << Data.PositionChangeThresholdForTeleport;
+		Ar << Data.OrientationChangeThresholdForTeleport;
+	}
 	Ar << Data.ExternalLinearDragV;
 	Ar << Data.ExternalLinearVelocity;
 	Ar << Data.ExternalAngularVelocity;
@@ -220,6 +225,12 @@ FArchive& operator <<(FArchive& Ar, FRigPhysicsSolverSettings& Data)
 	Ar << Data.VelocityIterations;
 	Ar << Data.ProjectionIterations;
 	Ar << Data.MaxNumRollingAverageStepTimes;
+	if (Ar.CustomVer(FPhysicsControlObjectVersion::GUID) >= FPhysicsControlObjectVersion::ControlRigSolverSettingsIncludesCollisionBoundsExpansion)
+	{
+		Ar << Data.CollisionBoundsExpansion;
+		Ar << Data.BoundsVelocityMultiplier;
+		Ar << Data.MaxVelocityBoundsExpansion;
+	}
 	Ar << Data.MaxDepenetrationVelocity;
 	Ar << Data.FixedTimeStep;
 	Ar << Data.MaxTimeSteps;
diff --git a/Engine/Plugins/Experimental/PhysicsControl/Source/ControlRigPhysics/Private/RigPhysicsSimulation.cpp b/Engine/Plugins/Experimental/PhysicsControl/Source/ControlRigPhysics/Private/RigPhysicsSimulation.cpp
--- a/Engine/Plugins/Experimental/PhysicsControl/Source/ControlRigPhysics/Private/RigPhysicsSimulation.cpp
+++ b/Engine/Plugins/Experimental/PhysicsControl/Source/ControlRigPhysics/Private/RigPhysicsSimulation.cpp
@@ -35,6 +35,7 @@
 #include "Chaos/ChaosScene.h"
 #include "Chaos/PBDJointConstraintUtilities.h"
 #include "Chaos/Evolution/SimulationSpace.h"
+#include "Chaos/PBDCollisionConstraints.h"
 
 #include "PhysicsEngine/BodyInstance.h"
 #include "PhysicsEngine/ConstraintInstance.h"
@@ -615,6 +616,10 @@ ImmediatePhysics::FActorHandle* FRigPhysicsSimulation::CreateBody(
 			ActorHandle->SetEnabled(true);
 		}
 	}
+	else
+	{
+		Simulation->SetHasCollision(ActorHandle, false);
+	}
 
 	return ActorHandle;
 }
@@ -691,8 +696,8 @@ void FRigPhysicsSimulation::InstantiateSolverCollision(
 	for (const FRigPhysicsIgnorePair& IgnorePair : IgnorePairs)
 	{
 		ImmediatePhysics::FSimulation::FIgnorePair ChaosPair;
-		ChaosPair.A = GetActor(IgnorePair.Key);
-		ChaosPair.B = GetActor(IgnorePair.Value);
+		ChaosPair.A = GetActor(IgnorePair.A);
+		ChaosPair.B = GetActor(IgnorePair.B);
 		if (ChaosPair.A && ChaosPair.B)
 		{
 			ChaosIgnorePairs.Add(ChaosPair);
@@ -1428,13 +1433,19 @@ void FRigPhysicsSimulation::StepSimulation(
 	// Set settings that might change
 	Simulation->SetSolverSettings(
 		FixedTimeStep,
-		999999999.0, // Danny TODO cull distance - expose this? Yuck
+		SolverSettings.CollisionBoundsExpansion,
 		SolverSettings.MaxDepenetrationVelocity,
 		SolverSettings.bUseLinearJointSolver,
 		SolverSettings.PositionIterations,
 		SolverSettings.VelocityIterations,
 		SolverSettings.ProjectionIterations,
 		SolverSettings.bUseManifolds);
+
+	Chaos::FCollisionDetectorSettings CollisionDetectorSettings = Simulation->GetCollisionDetectorSettings();
+	CollisionDetectorSettings.BoundsVelocityInflation = SolverSettings.BoundsVelocityMultiplier;
+	CollisionDetectorSettings.MaxVelocityBoundsExpansion = SolverSettings.MaxVelocityBoundsExpansion;
+	Simulation->SetCollisionDetectorSettings(CollisionDetectorSettings);
+
 	// This gets reset to 100 after every simulation step!
 	Simulation->SetMaxNumRollingAverageStepTimes(SolverSettings.MaxNumRollingAverageStepTimes);
 
@@ -1443,7 +1454,7 @@ void FRigPhysicsSimulation::StepSimulation(
 	ChaosJointSolverSettings.bSortEnabled = true;
 
 	// Simulation space
-	Chaos::FSimulationSpaceSettings ChaosSimulationSpaceSettings;
+	Chaos::FSimulationSpaceSettings ChaosSimulationSpaceSettings = Simulation->GetSimulationSpaceSettings();
 	ChaosSimulationSpaceSettings.bEnabled = (SimulationSpaceSettings.SpaceMovementAmount > 0.0f);
 	ChaosSimulationSpaceSettings.ExternalLinearEtherDrag = SimulationSpaceSettings.ExternalLinearDragV;
 	Simulation->SetSimulationSpaceSettings(ChaosSimulationSpaceSettings);
diff --git a/Engine/Plugins/Experimental/PhysicsControl/Source/ControlRigPhysics/Private/RigPhysicsSimulation.h b/Engine/Plugins/Experimental/PhysicsControl/Source/ControlRigPhysics/Private/RigPhysicsSimulation.h
--- a/Engine/Plugins/Experimental/PhysicsControl/Source/ControlRigPhysics/Private/RigPhysicsSimulation.h
+++ b/Engine/Plugins/Experimental/PhysicsControl/Source/ControlRigPhysics/Private/RigPhysicsSimulation.h
@@ -10,6 +10,7 @@
 #include "Rigs/RigHierarchyComponents.h"
 
 #include "Physics/ImmediatePhysics/ImmediatePhysicsDeclares.h"
+#include "Physics/ImmediatePhysics/ImmediatePhysicsSimulation.h"
 #include "Chaos/PBDJointConstraintTypes.h"
 
 #include "RigPhysicsSimulation.generated.h"
@@ -35,23 +36,22 @@ DECLARE_LOG_CATEGORY_EXTERN(LogRigPhysics, Log, All);
 //======================================================================================================================
 // For internal use only
 //======================================================================================================================
-typedef TPair<FRigComponentKey, FRigComponentKey> FRigPhysicsIgnorePair;
-inline bool operator<(const FRigPhysicsIgnorePair& A, const FRigPhysicsIgnorePair& B)
+struct FRigPhysicsIgnorePair
 {
-	if (A.Key < B.Key)
-	{
-		return true;
-	}
-	else if (A.Key > B.Key)
-	{
-		return false;
-	}
-	// keys are the same...
-	if (A.Value < B.Value)
+	FRigPhysicsIgnorePair(const FRigComponentKey& InA, const FRigComponentKey& InB)
+		: A(InA), B(InB) {}
+	FRigComponentKey A;
+	FRigComponentKey B;
+
+	bool operator==(const FRigPhysicsIgnorePair& Other) const
 	{
-		return true;
+		return (A == Other.A && B == Other.B) || (A == Other.B && B == Other.A);
 	}
-	return false;
+};
+
+FORCEINLINE uint32 GetTypeHash(const FRigPhysicsIgnorePair& Pair)
+{
+	return HashCombine(GetTypeHash(FMath::Min(Pair.A, Pair.B)), GetTypeHash(FMath::Max(Pair.A, Pair.B)));
 }
 typedef TSet<FRigPhysicsIgnorePair> FRigPhysicsIgnorePairs;
 
@@ -172,10 +172,6 @@ private:
 	{
 		FTransform ComponentTM;
 		FTransform BoneRelComponentTM;
-		FVector    ComponentLinearVelocity;
-		FVector    ComponentAngularVelocity;
-		FVector    BoneLinearVelocity;
-		FVector    BoneAngularVelocity;
 
 		// The world transform of the simulation space
 		FTransform SimulationSpaceTM;
diff --git a/Engine/Plugins/Experimental/PhysicsControl/Source/ControlRigPhysics/Private/RigPhysicsSimulationSpace.cpp b/Engine/Plugins/Experimental/PhysicsControl/Source/ControlRigPhysics/Private/RigPhysicsSimulationSpace.cpp
--- a/Engine/Plugins/Experimental/PhysicsControl/Source/ControlRigPhysics/Private/RigPhysicsSimulationSpace.cpp
+++ b/Engine/Plugins/Experimental/PhysicsControl/Source/ControlRigPhysics/Private/RigPhysicsSimulationSpace.cpp
@@ -128,10 +128,6 @@ void FRigPhysicsSimulation::InitSimulationSpace(
 {
 	SimulationSpaceState.ComponentTM = ComponentTM;
 	SimulationSpaceState.BoneRelComponentTM = BoneRelComponentTM;
-	SimulationSpaceState.ComponentLinearVelocity = FVector::ZeroVector;
-	SimulationSpaceState.ComponentAngularVelocity = FVector::ZeroVector;
-	SimulationSpaceState.BoneLinearVelocity = FVector::ZeroVector;
-	SimulationSpaceState.BoneAngularVelocity = FVector::ZeroVector;
 }
 
 //======================================================================================================================
@@ -141,6 +137,12 @@ FVector CalculateAngularVelocity(const FQuat& PrevQ, const FQuat& CurrentQ, floa
 	return DeltaQ.ToRotationVector() / Dt;
 }
 
+//======================================================================================================================
+FVector CalculateLinearVelocity(const FVector& PrevP, const FVector& CurrentP, float Dt)
+{
+	return (CurrentP - PrevP) / Dt;
+}
+
 //======================================================================================================================
 // Note - don't use the space conversion functions here as the state won't have been set yet.
 FRigPhysicsSimulation::FSimulationSpaceData FRigPhysicsSimulation::UpdateSimulationSpaceStateAndCalculateData(
@@ -194,7 +196,7 @@ FRigPhysicsSimulation::FSimulationSpaceData FRigPhysicsSimulation::UpdateSimulat
 
 	// If the timestep is zero, then it doesn't actually matter what the velocity is - but make sure
 	// it doesn't corrupt anything.
-	if (SimulationSpaceState.Dt < SMALL_NUMBER || SimulationSpaceState.PrevDt < SMALL_NUMBER)
+	if (SimulationSpaceState.Dt < SMALL_NUMBER)
 	{
 		SimulationSpaceData.LinearVelocity = FVector::ZeroVector;
 		SimulationSpaceData.AngularVelocity = FVector::ZeroVector;
@@ -211,23 +213,29 @@ FRigPhysicsSimulation::FSimulationSpaceData FRigPhysicsSimulation::UpdateSimulat
 	// finite difference approximations, but this wouldn't necessarily be correct.
 
 	// World-space component linear velocity and acceleration
-	const FVector PrevSpaceLinearVel = Chaos::FVec3::CalculateVelocity(
-		SimulationSpaceState.PrevPrevSimulationSpaceTM.GetTranslation(),
-		SimulationSpaceState.PrevSimulationSpaceTM.GetTranslation(), SimulationSpaceState.PrevDt);
-	SimulationSpaceData.LinearVelocity = Chaos::FVec3::CalculateVelocity(
+	SimulationSpaceData.LinearVelocity = CalculateLinearVelocity(
 		SimulationSpaceState.PrevSimulationSpaceTM.GetTranslation(),
 		SimulationSpaceState.SimulationSpaceTM.GetTranslation(), SimulationSpaceState.Dt);
-	SimulationSpaceData.LinearAcceleration = 
+	const FVector PrevSpaceLinearVel = 
+		SimulationSpaceState.PrevDt < SMALL_NUMBER 
+		? SimulationSpaceData.LinearVelocity 
+		: CalculateLinearVelocity(
+			SimulationSpaceState.PrevPrevSimulationSpaceTM.GetTranslation(),
+			SimulationSpaceState.PrevSimulationSpaceTM.GetTranslation(), SimulationSpaceState.PrevDt);
+	SimulationSpaceData.LinearAcceleration =
 		(SimulationSpaceData.LinearVelocity - PrevSpaceLinearVel) / SimulationSpaceState.Dt;
 
 	// World-space component angular velocity and acceleration
-	const FVector PrevSpaceAngularVel = CalculateAngularVelocity(
-		SimulationSpaceState.PrevPrevSimulationSpaceTM.GetRotation(),
-		SimulationSpaceState.PrevSimulationSpaceTM.GetRotation(), SimulationSpaceState.PrevDt);
 	SimulationSpaceData.AngularVelocity = CalculateAngularVelocity(
 		SimulationSpaceState.PrevSimulationSpaceTM.GetRotation(),
 		SimulationSpaceState.SimulationSpaceTM.GetRotation(), SimulationSpaceState.Dt);
-	SimulationSpaceData.AngularAcceleration = 
+	const FVector PrevSpaceAngularVel = 
+		SimulationSpaceState.PrevDt < SMALL_NUMBER 
+		? SimulationSpaceData.AngularVelocity
+		: CalculateAngularVelocity(
+			SimulationSpaceState.PrevPrevSimulationSpaceTM.GetRotation(),
+			SimulationSpaceState.PrevSimulationSpaceTM.GetRotation(), SimulationSpaceState.PrevDt);
+	SimulationSpaceData.AngularAcceleration =
 		(SimulationSpaceData.AngularVelocity - PrevSpaceAngularVel) / SimulationSpaceState.Dt;
 
 	// Apply Z scale
@@ -241,13 +249,26 @@ FRigPhysicsSimulation::FSimulationSpaceData FRigPhysicsSimulation::UpdateSimulat
 			FMath::Square(SimulationSpaceSettings.LinearAccelerationThresholdForTeleport)) ||
 		(SimulationSpaceSettings.AngularAccelerationThresholdForTeleport > 0 &&
 			SimulationSpaceData.AngularAcceleration.SquaredLength() >
-			FMath::Square(SimulationSpaceSettings.AngularAccelerationThresholdForTeleport))
+			FMath::Square(SimulationSpaceSettings.AngularAccelerationThresholdForTeleport)) ||
+		(SimulationSpaceSettings.PositionChangeThresholdForTeleport > 0 &&
+			SimulationSpaceData.LinearVelocity.SquaredLength() >
+			FMath::Square(SimulationSpaceSettings.PositionChangeThresholdForTeleport / SimulationSpaceState.Dt)) ||
+		(SimulationSpaceSettings.OrientationChangeThresholdForTeleport > 0 &&
+			SimulationSpaceData.AngularVelocity.SquaredLength() >
+			FMath::Square(FMath::DegreesToRadians(
+				SimulationSpaceSettings.OrientationChangeThresholdForTeleport / SimulationSpaceState.Dt)))
 		)
 	{
 		SimulationSpaceData.LinearVelocity = FVector::ZeroVector;
 		SimulationSpaceData.AngularVelocity = FVector::ZeroVector;
 		SimulationSpaceData.LinearAcceleration = FVector::ZeroVector;
 		SimulationSpaceData.AngularAcceleration = FVector::ZeroVector;
+
+		// This will stop the next step from using bogus values too.
+		SimulationSpaceState.PrevSimulationSpaceTM = SimulationSpaceState.SimulationSpaceTM;
+		SimulationSpaceState.PrevPrevSimulationSpaceTM = SimulationSpaceState.SimulationSpaceTM;
+		SimulationSpaceState.PrevDt = 0;
+		SimulationSpaceState.Dt = 0;
 	}
 	else
 	{
@@ -278,7 +299,7 @@ FRigPhysicsSimulation::FSimulationSpaceData FRigPhysicsSimulation::UpdateSimulat
 
 	// Transform world-space motion into simulation space Danny TODO note that this matches the code
 	// in RBAN, and is doing what the interface requires (i.e. movement of the space in the space of
-	// the space!). However, the resultant motion is incorrect for rotation of the space.
+	// the space!). 
 	SimulationSpaceData.LinearVelocity =
 		SimulationSpaceState.SimulationSpaceTM.InverseTransformVector(SimulationSpaceData.LinearVelocity);
 	SimulationSpaceData.AngularVelocity =
diff --git a/Engine/Plugins/Experimental/PhysicsControl/Source/ControlRigPhysics/Private/RigPhysicsSimulationVisualization.cpp b/Engine/Plugins/Experimental/PhysicsControl/Source/ControlRigPhysics/Private/RigPhysicsSimulationVisualization.cpp
--- a/Engine/Plugins/Experimental/PhysicsControl/Source/ControlRigPhysics/Private/RigPhysicsSimulationVisualization.cpp
+++ b/Engine/Plugins/Experimental/PhysicsControl/Source/ControlRigPhysics/Private/RigPhysicsSimulationVisualization.cpp
@@ -29,8 +29,20 @@ namespace RigPhysicsSolverDraw
 {
 	FLinearColor DynamicColor = FColor::Yellow;
 	FLinearColor KinematicColor = FColor::Blue;
+
+	FLinearColor ActiveContactColor = FColor::Red;
+	FLinearColor InactiveContactColor = FColor::Silver;
+
 }
 
+TAutoConsoleVariable<int> CVarControlRigPhysicsShowActiveContactsOverride(
+	TEXT("ControlRig.Physics.ShowActiveContactsOveride"), -1,
+	TEXT("Whether to draw active contacts (requires visualization to be enabled). -1 uses the visualization setting, 0 forces drawing to be disabled, 1 forces it to be enabled."));
+
+TAutoConsoleVariable<int> CVarControlRigPhysicsShowInactiveContactsOverride(
+	TEXT("ControlRig.Physics.ShowInactiveContactsOveride"), -1,
+	TEXT("Whether to draw inactive contacts (requires visualization to be enabled). -1 uses the visualization setting, 0 forces drawing to be disabled, 1 forces it to be enabled."));
+
 //======================================================================================================================
 static void DrawShapes(
 	FRigVMDrawInterface*                    DI,
@@ -352,5 +364,52 @@ void FRigPhysicsSimulation::Draw(
 			DrawJoint(DI, VisualizationSettings, SpaceTransform, *JointHandle);
 		}
 	}
+
+	int32 ShowActiveOverride = CVarControlRigPhysicsShowActiveContactsOverride.GetValueOnAnyThread();
+	int32 ShowInactiveOverride = CVarControlRigPhysicsShowInactiveContactsOverride.GetValueOnAnyThread();
+
+	bool bShowActiveContacts = ShowActiveOverride < 0 ? VisualizationSettings.bShowActiveContacts : 
+		(ShowActiveOverride ? true : false);
+	bool bShowInactiveContacts = ShowInactiveOverride < 0 ? VisualizationSettings.bShowInactiveContacts :
+		(ShowInactiveOverride ? true : false);
+
+	if (bShowActiveContacts || bShowInactiveContacts)
+	{
+		Simulation->VisitCollisions(
+			[DI, &SpaceTransform, &VisualizationSettings, bShowActiveContacts, bShowInactiveContacts]
+			(const ImmediatePhysics::FSimulation::FCollisionData& Collision)
+			{
+				using namespace Chaos;
+
+				bool bIsActive = Collision.GetCollisionAccumulatedImpulse().IsZero();
+				if ((bIsActive && bShowActiveContacts) || (!bIsActive && bShowInactiveContacts))
+				{
+					float Size = 5.0f * VisualizationSettings.ShapeSize;
+
+					int32 NumManifoldPoints = Collision.GetNumManifoldPoints();
+					for (int32 PointIndex = 0; PointIndex != NumManifoldPoints; ++PointIndex)
+					{
+						Chaos::FRealSingle Depth;
+						Chaos::FVec3 PlaneNormal;
+						Chaos::FVec3 PointLocation;
+						Chaos::FVec3 PlaneLocation;
+						Collision.GetManifoldPointData(
+							PointIndex, Depth, PlaneNormal, PointLocation, PlaneLocation);
+
+						const FVec3 PointPlaneLocation = PointLocation -
+							FVec3::DotProduct(PointLocation - PlaneLocation, PlaneNormal) * PlaneNormal;
+
+						FMatrix Axes = FRotationMatrix::MakeFromZ(PlaneNormal);
+						FTransform TM(Axes);
+						TM.SetTranslation(PointPlaneLocation);
+
+						DI->DrawCircle(
+							SpaceTransform, TM, Size,
+							bIsActive ? RigPhysicsSolverDraw::ActiveContactColor : RigPhysicsSolverDraw::InactiveContactColor,
+							VisualizationSettings.LineThickness, VisualizationSettings.ShapeDetail);
+					}
+				}
+			}, Chaos::ECollisionVisitorFlags::VisitDefault);
+	}
 }
 
diff --git a/Engine/Plugins/Experimental/PhysicsControl/Source/ControlRigPhysics/Public/RigPhysicsData.h b/Engine/Plugins/Experimental/PhysicsControl/Source/ControlRigPhysics/Public/RigPhysicsData.h
--- a/Engine/Plugins/Experimental/PhysicsControl/Source/ControlRigPhysics/Public/RigPhysicsData.h
+++ b/Engine/Plugins/Experimental/PhysicsControl/Source/ControlRigPhysics/Public/RigPhysicsData.h
@@ -39,6 +39,16 @@ struct FRigPhysicsVisualizationSettings
 
 	UPROPERTY(EditAnywhere, Category = Visualization, meta = (ClampMin = "0.0"))
 	int32 ShapeDetail = 16;
+
+	// If visualization is enabled, show active contacts. Note that this can be overridden using
+	// CVar ControlRig.Physics.DrawActiveContactsOveride
+	UPROPERTY(EditAnywhere, Category = Visualization)
+	bool bShowActiveContacts = true;
+
+	// If visualization is enabled, show inactive contacts. Note that this can be overridden using
+	// CVar ControlRig.Physics.DrawInactiveContactsOveride
+	UPROPERTY(EditAnywhere, Category = Visualization)
+	bool bShowInactiveContacts = false;
 };
 
 /** 
@@ -337,6 +347,20 @@ struct FRigPhysicsSimulationSpaceSettings
 	UPROPERTY(EditAnywhere, Category = Settings, meta = (ClampMin = "0.0"))
 	float AngularAccelerationThresholdForTeleport = 100000;
 
+	// If the linear position change is above this threshold, the movement is treated as a teleport.
+	// The calculated accelerations and velocities will be set to zero, without resetting the
+	// simulation state.
+	// A value of zero will disable this detection
+	UPROPERTY(EditAnywhere, Category = Settings, meta = (ClampMin = "0.0"))
+	float PositionChangeThresholdForTeleport = 100;
+
+	// If the angular change (degrees) is above this threshold, the movement is treated as a teleport.
+	// The calculated accelerations and velocities will be set to zero, without resetting the
+	// simulation state.
+	// A value of zero will disable this detection
+	UPROPERTY(EditAnywhere, Category = Settings, meta = (ClampMin = "0.0"))
+	float OrientationChangeThresholdForTeleport = 45;
+
 	// Additional linear drag applied to every body in addition to linear drag specified on them in
 	// the physics asset. When combined with ExternalLinearVelocity, this can be used to add a
 	// temporary wind-blown effect without having to tune linear drag on all the bodies in the
@@ -423,9 +447,24 @@ struct CONTROLRIGPHYSICS_API FRigPhysicsSolverSettings
 	UPROPERTY(EditAnywhere, Category = SolverSettings, meta = (ClampMin = 1))
 	int32 MaxNumRollingAverageStepTimes = 1;
 
-	/**
-	 * When bodies are penetrating, this is the maximum velocity delta that can be applied in one frame.
-	 */
+	// This is the margin around shapes used to detect collisions. Increasing this will increase the
+	// number of inactive contacts, which will reduce the likelihood of penetration, but will also
+	// increase solver cost.
+	UPROPERTY(EditAnywhere, Category = SolverSettings, meta = (ClampMin = 0))
+	float CollisionBoundsExpansion = 2.0f;
+
+	// Expands the shape bounds to detect potentially active contacts when moving. Increasing this
+	// (typically up to 1) will increase the number of inactive contacts, which will reduce the
+	// likelihood of penetration, but will also increase solver cost.
+	UPROPERTY(EditAnywhere, Category = SolverSettings, meta = (ClampMin = 0))
+	float BoundsVelocityMultiplier = 1.0f;
+
+	// The maximum margin added due to movement. Reducing this can prevent excessive numbers of inactive
+	// contacts being generated.
+	UPROPERTY(EditAnywhere, Category = SolverSettings, meta = (ClampMin = 0))
+	float MaxVelocityBoundsExpansion = 25.0f;
+
+	// When bodies are penetrating, this is the maximum velocity delta that can be applied in one frame.
 	UPROPERTY(EditAnywhere, Category = SolverSettings, meta = (ClampMin = 0))
 	float MaxDepenetrationVelocity = 0.0f;
 
diff --git a/Engine/Plugins/Experimental/PhysicsControl/Source/PhysicsControl/Public/PhysicsControlObjectVersion.h b/Engine/Plugins/Experimental/PhysicsControl/Source/PhysicsControl/Public/PhysicsControlObjectVersion.h
--- a/Engine/Plugins/Experimental/PhysicsControl/Source/PhysicsControl/Public/PhysicsControlObjectVersion.h
+++ b/Engine/Plugins/Experimental/PhysicsControl/Source/PhysicsControl/Public/PhysicsControlObjectVersion.h
@@ -36,6 +36,10 @@ struct PHYSICSCONTROL_API FPhysicsControlObjectVersion
 
 		ControlRigBodyDynamicsHasDensity,
 
+		ControlRigSolverSettingsIncludesCollisionBoundsExpansion,
+
+		ControlRigDetectTeleportFromDistanceChange,
+
 		// -----<new versions can be added above this line>-------------------------------------------------
 		VersionPlusOne,
 		LatestVersion = VersionPlusOne - 1
diff --git a/Engine/Source/Runtime/Engine/Private/PhysicsEngine/ImmediatePhysics/ImmediatePhysicsChaos/ImmediatePhysicsSimulation_Chaos.cpp b/Engine/Source/Runtime/Engine/Private/PhysicsEngine/ImmediatePhysics/ImmediatePhysicsChaos/ImmediatePhysicsSimulation_Chaos.cpp
--- a/Engine/Source/Runtime/Engine/Private/PhysicsEngine/ImmediatePhysics/ImmediatePhysicsChaos/ImmediatePhysicsSimulation_Chaos.cpp
+++ b/Engine/Source/Runtime/Engine/Private/PhysicsEngine/ImmediatePhysics/ImmediatePhysicsChaos/ImmediatePhysicsSimulation_Chaos.cpp
@@ -798,6 +798,11 @@ namespace ImmediatePhysics_Chaos
 		Implementation->Evolution.SetSimulationSpaceSettings(SimulationSpaceSettings);
 	}
 
+	const Chaos::FSimulationSpaceSettings& FSimulation::GetSimulationSpaceSettings() const
+	{
+		return Implementation->Evolution.GetSimulationSpaceSettings();
+	}
+
 	void FSimulation::SetSolverSettings(const FReal FixedDt, const FReal CullDistance, const FReal MaxDepenetrationVelocity, const int32 UseLinearJointSolver, const int32 PositionIts, const int32 VelocityIts, const int32 ProjectionIts, const int32 bUseManifoilds)
 	{
 		if (FixedDt >= FReal(0))
@@ -1203,6 +1208,76 @@ namespace ImmediatePhysics_Chaos
 #endif
 	}
 
+	const Chaos::FCollisionDetectorSettings& FSimulation::GetCollisionDetectorSettings() const
+	{
+		return Implementation->Collisions.GetDetectorSettings();
+	}
+
+	void FSimulation::SetCollisionDetectorSettings(const Chaos::FCollisionDetectorSettings& Settings)
+	{
+		Implementation->Collisions.SetDetectorSettings(Settings);
+	}
+
+	Chaos::FVec3 FSimulation::FCollisionData::GetCollisionAccumulatedImpulse() const
+	{
+		return Collision.GetAccumulatedImpulse();
+	}
+
+	int32 FSimulation::FCollisionData::GetNumManifoldPoints() const
+	{
+		return Collision.NumManifoldPoints();
+	}
+
+	const void FSimulation::FCollisionData::GetManifoldPointData(
+		const int32              PointIndex,
+		FRealSingle&             Depth,
+		Chaos::FVec3&            PlaneNormal,
+		Chaos::FVec3&            PointLocation,
+		Chaos::FVec3&            PlaneLocation) const
+	{
+		using namespace Chaos;
+
+		Depth = Collision.GetManifoldPoint(PointIndex).InitialPhi;
+
+		const FConstGenericParticleHandle Particle0 = Collision.GetParticle0();
+		const FConstGenericParticleHandle Particle1 = Collision.GetParticle1();
+		const FRigidTransform3 WorldActorTransform0 = FParticleUtilities::GetActorWorldTransform(Particle0);
+		const FRigidTransform3 WorldActorTransform1 = FParticleUtilities::GetActorWorldTransform(Particle1);
+
+		const FManifoldPoint& ManifoldPoint = Collision.GetManifoldPoint(PointIndex);
+		const FManifoldPointResult& ManifoldPointResult = Collision.GetManifoldPointResult(PointIndex);
+
+		const int32 ContactPlaneOwner = 1;
+		const int32 ContactPointOwner = 1 - ContactPlaneOwner;
+		const Chaos::FRigidTransform3 PlaneTransform = (ContactPlaneOwner == 0)
+			? Collision.GetShapeRelativeTransform0() * WorldActorTransform0
+			: Collision.GetShapeRelativeTransform1() * WorldActorTransform1;
+		const Chaos::FRigidTransform3 PointTransform = (ContactPlaneOwner == 0)
+			? Collision.GetShapeRelativeTransform1() * WorldActorTransform1
+			: Collision.GetShapeRelativeTransform0() * WorldActorTransform0;
+
+		PlaneNormal = PlaneTransform.TransformVectorNoScale(
+			FVec3(ManifoldPoint.ContactPoint.ShapeContactNormal));
+		PointLocation = PointTransform.TransformPosition(
+			FVec3(ManifoldPoint.ContactPoint.ShapeContactPoints[ContactPointOwner]));
+		PlaneLocation = PlaneTransform.TransformPosition(
+			FVec3(ManifoldPoint.ContactPoint.ShapeContactPoints[ContactPlaneOwner]));
+	}
+
+	void FSimulation::VisitCollisions(
+		TFunction<void(const FCollisionData&)> Visitor, Chaos::ECollisionVisitorFlags VisitorFlags) const
+	{
+		using namespace Chaos;
+		const Private::FCollisionConstraintAllocator& ConstraintAllocator = Implementation->Collisions.GetConstraintAllocator();
+		ConstraintAllocator.VisitConstCollisions(
+			[&](const FPBDCollisionConstraint& Collision)
+			{
+				Visitor(FCollisionData(Collision));
+				return ECollisionVisitorResult::Continue;
+			}, VisitorFlags);
+	}
+
+
 	void FSimulation::DebugDrawConstraints()
 	{
 #if CHAOS_DEBUG_DRAW
diff --git a/Engine/Source/Runtime/Engine/Public/Physics/ImmediatePhysics/ImmediatePhysicsChaos/ImmediatePhysicsSimulation_Chaos.h b/Engine/Source/Runtime/Engine/Public/Physics/ImmediatePhysics/ImmediatePhysicsChaos/ImmediatePhysicsSimulation_Chaos.h
--- a/Engine/Source/Runtime/Engine/Public/Physics/ImmediatePhysics/ImmediatePhysicsChaos/ImmediatePhysicsSimulation_Chaos.h
+++ b/Engine/Source/Runtime/Engine/Public/Physics/ImmediatePhysics/ImmediatePhysicsChaos/ImmediatePhysicsSimulation_Chaos.h
@@ -6,6 +6,7 @@
 #include "ChaosVisualDebugger/ChaosVDContextProvider.h"
 #include "Physics/ImmediatePhysics/ImmediatePhysicsChaos/ImmediatePhysicsCore_Chaos.h"
 
+#include "Chaos/Collision/CollisionVisitor.h"
 #include "Engine/EngineTypes.h"
 #include "Templates/UniquePtr.h"
 
@@ -123,9 +124,15 @@ namespace ImmediatePhysics_Chaos
 			const FReal DampingAlpha,
 			const FVector& ExternalLinearEtherDrag);
 
+		ENGINE_API const Chaos::FSimulationSpaceSettings& GetSimulationSpaceSettings() const;
+
 		ENGINE_API void SetSimulationSpaceSettings(
 			const Chaos::FSimulationSpaceSettings& SimulationSpaceSettings);
 
+		ENGINE_API const Chaos::FCollisionDetectorSettings& GetCollisionDetectorSettings() const;
+
+		ENGINE_API void SetCollisionDetectorSettings(const Chaos::FCollisionDetectorSettings& Settings);
+
 		/** Set settings. Invalid (negative) values will leave that value unchanged from defaults */
 		ENGINE_API void SetSolverSettings(
 			const FReal FixedDt,
@@ -145,6 +152,38 @@ namespace ImmediatePhysics_Chaos
 		/** Explicit debug draw path if the use case needs it to happen at a point outside of the simulation **/
 		ENGINE_API void DebugDraw();
 
+		// Data relating to contacts returned by VisitCollisions
+		struct FCollisionData
+		{
+		public:
+			// The overall accumulated impulse applied due to this collision/contact. Only valid if
+			// called after the solve has completed.
+			ENGINE_API Chaos::FVec3 GetCollisionAccumulatedImpulse() const;
+
+			ENGINE_API int GetNumManifoldPoints() const;
+
+			// The contact represents a point attached to one particle, and the plane attached to
+			// the other. These will be calculated and returned relative to where the particles are now.
+			ENGINE_API const void GetManifoldPointData(
+				const int32              ManifoldPointIndex,
+				FRealSingle&             Depth, // The initial penetration depth
+				Chaos::FVec3&            PlaneNormal,
+				Chaos::FVec3&            PointLocation,
+				Chaos::FVec3&            PlaneLocation) const;
+
+		private:
+			friend struct FSimulation;
+			// This is constructed by Chaos, so it passes in the original CollisionConstraint
+			FCollisionData(const Chaos::FPBDCollisionConstraint& InCollisionConstraint) 
+				: Collision(InCollisionConstraint) {}
+
+			const Chaos::FPBDCollisionConstraint& Collision;
+		};
+
+		// Access to collisions detected during the previous solve. 
+		ENGINE_API void VisitCollisions(
+			TFunction<void(const FCollisionData&)> Visitor, Chaos::ECollisionVisitorFlags VisitorFlags) const;
+
 	private:
 		void RemoveFromCollidingPairs(FActorHandle* ActorHandle);
 		void UpdateInertiaConditioning(const FVector& Gravity);
