Control Rig Physics: Better teleport detection and simulation resetting

diff --git a/Engine/Plugins/Experimental/PhysicsControl/Source/ControlRigPhysics/Private/RigPhysicsData.cpp b/Engine/Plugins/Experimental/PhysicsControl/Source/ControlRigPhysics/Private/RigPhysicsData.cpp
--- a/Engine/Plugins/Experimental/PhysicsControl/Source/ControlRigPhysics/Private/RigPhysicsData.cpp
+++ b/Engine/Plugins/Experimental/PhysicsControl/Source/ControlRigPhysics/Private/RigPhysicsData.cpp
@@ -239,6 +239,14 @@ FArchive& operator <<(FArchive& Ar, FRigPhysicsSolverSettings& Data)
 	Ar << Data.bSolveJointPositionsLast;
 	Ar << Data.bUseManifolds;
 	Ar << Data.DistanceThresholdForReset;
+	if (Ar.CustomVer(FPhysicsControlObjectVersion::GUID) >= FPhysicsControlObjectVersion::ControlRigSpeedThresholdForReset)
+	{
+		Ar << Data.SpeedThresholdForReset;
+	}
+	if (Ar.CustomVer(FPhysicsControlObjectVersion::GUID) >= FPhysicsControlObjectVersion::ControlRigSpeedThresholdForReset)
+	{
+		Ar << Data.ResetCooldownFrames;
+	}
 	if (Ar.CustomVer(FPhysicsControlObjectVersion::GUID) >= FPhysicsControlObjectVersion::ControlRigAutomaticallyAddPhysicsComponents)
 	{
 		Ar << Data.bAutomaticallyAddPhysicsComponents;
diff --git a/Engine/Plugins/Experimental/PhysicsControl/Source/ControlRigPhysics/Private/RigPhysicsSimulation.cpp b/Engine/Plugins/Experimental/PhysicsControl/Source/ControlRigPhysics/Private/RigPhysicsSimulation.cpp
--- a/Engine/Plugins/Experimental/PhysicsControl/Source/ControlRigPhysics/Private/RigPhysicsSimulation.cpp
+++ b/Engine/Plugins/Experimental/PhysicsControl/Source/ControlRigPhysics/Private/RigPhysicsSimulation.cpp
@@ -745,6 +745,12 @@ void FRigPhysicsSimulation::InstantiatePhysicsBodies(
 					ComponentKey.ElementKey.Name, Collision, &Dynamics, &BodyData, SourceSimulationSpaceTM);
 			}
 
+			Record.TargetElementKey = PhysicsComponent->BodySolverSettings.TargetBone;
+			if (!Record.TargetElementKey.IsValid())
+			{
+				Record.TargetElementKey = ComponentKey.ElementKey;
+			}
+
 			for (const FRigComponentKey& NoCollisionKey : PhysicsComponent->NoCollisionBodies)
 			{
 				IgnorePairs.Add(FRigPhysicsIgnorePair(ComponentKey, NoCollisionKey));
@@ -1294,62 +1300,60 @@ void FRigPhysicsSimulation::UpdatePostPhysics(const FRigPhysicsSolverComponent*
 	bool bGotInvalidSimulationData = false;
 	URigHierarchy* Hierarchy = OwningControlRig->GetHierarchy();
 
+	double SpeedThresholdForResetSquared = SolverSettings.SpeedThresholdForReset > 0 ?
+		FMath::Square(SolverSettings.SpeedThresholdForReset) : TNumericLimits<double>::Max();
+
+	double DistanceThresholdForResetSquared = SolverSettings.DistanceThresholdForReset > 0 ?
+		FMath::Square(SolverSettings.DistanceThresholdForReset) : TNumericLimits<double>::Max();
+
 	// Traverse using the sorted keys
 	for (const FRigComponentKey& ComponentKey : SortedBodyComponentKeys)
 	{
-		const FRigBodyRecord& Record = *BodyRecords.Find(ComponentKey);
-		if (Record.ActorHandle)
+		FRigBodyRecord& Record = *BodyRecords.Find(ComponentKey);
+		if (Record.ActorHandle && Record.TargetElementKey.IsValid())
 		{
-			if (const FRigPhysicsBodyComponent* PhysicsComponent = Cast<FRigPhysicsBodyComponent>(
-				Hierarchy->FindComponent(ComponentKey)))
+			FTransform SimSpaceTM = Record.ActorHandle->GetWorldTransform();
+			double DistSq = SimSpaceTM.GetTranslation().SquaredLength();
+			if (!SimSpaceTM.IsValid() || DistSq > DistanceThresholdForResetSquared)
 			{
-				// We will only have created this body record if it's associated with our solver, so
-				// we don't need to check again that our solver is in its list.
-				FRigElementKey TargetElementKey = PhysicsComponent->BodySolverSettings.TargetBone;
-				if (!TargetElementKey.IsValid())
-				{
-					TargetElementKey = ComponentKey.ElementKey;
-				}
+				bGotInvalidSimulationData = true;
+			}
+			else if (Record.ActorHandle->GetLinearVelocity().SquaredLength() > SpeedThresholdForResetSquared)
+			{
+				bGotInvalidSimulationData = true;
+			}
 
-				if (TargetElementKey.IsValid())
-				{
-					FTransform TargetSimSpaceTM = Record.ActorHandle->GetWorldTransform();
-					double DistSq = TargetSimSpaceTM.GetTranslation().SquaredLength();
-					if (!TargetSimSpaceTM.IsValid() || (SolverSettings.DistanceThresholdForReset > 0.0f &&
-						DistSq > FMath::Square(SolverSettings.DistanceThresholdForReset)))
-					{
-						bGotInvalidSimulationData = true;
-					}
-					else
-					{
-						FTransform TargetComponentSpaceTM = ConvertSimSpaceTransformToComponentSpace(
-							SolverSettings, TargetSimSpaceTM);
-						if (Alpha < 0.999f)
-						{
-							// Danny TODO Note that this uses Alpha to blend in component space.
-							// This can cause joint separation. We probably want an option to blend
-							// in local (joint) space, perhaps splitting the alpha into orientation
-							// and position.
-							FTransform CurrentTM = Hierarchy->GetGlobalTransform(TargetElementKey);
-							FQuat TargetQ = FQuat::Slerp(
-								CurrentTM.GetRotation(), TargetComponentSpaceTM.GetRotation(), Alpha);
-							FVector TargetT = FMath::Lerp(
-								CurrentTM.GetTranslation(), TargetComponentSpaceTM.GetTranslation(), Alpha);
-							TargetComponentSpaceTM.SetRotation(TargetQ);
-							TargetComponentSpaceTM.SetTranslation(TargetT);
-						}
-						// Note that we set bAffectChildren = true (i.e. don't counter-animate
-						// children), so that attached animation bones will follow physics, but we
-						// rely on our bodies being sorted so we work out from the leaf nodes so as
-						// not to disturb previously set bodies.
-						Hierarchy->SetGlobalTransform(TargetElementKey, TargetComponentSpaceTM, false, true);
-					}
-				}
+			// Calculate the target TM even if we're going to reset - it's likely useful for
+			// debugging (and this should be rare!)
+			Record.TargetComponentSpaceTM = ConvertSimSpaceTransformToComponentSpace(
+				SolverSettings, SimSpaceTM);
+			if (Alpha < 0.999f)
+			{
+				// Danny TODO Note that this uses Alpha to blend in component space.
+				// This can cause joint separation. We probably want an option to blend
+				// in local (joint) space, perhaps splitting the alpha into orientation
+				// and position.
+				FTransform CurrentTM = Hierarchy->GetGlobalTransform(Record.TargetElementKey);
+				FQuat TargetQ = FQuat::Slerp(
+					CurrentTM.GetRotation(), Record.TargetComponentSpaceTM.GetRotation(), Alpha);
+				FVector TargetT = FMath::Lerp(
+					CurrentTM.GetTranslation(), Record.TargetComponentSpaceTM.GetTranslation(), Alpha);
+				Record.TargetComponentSpaceTM.SetRotation(TargetQ);
+				Record.TargetComponentSpaceTM.SetTranslation(TargetT);
 			}
 		}
 	}
 
 	if (bGotInvalidSimulationData)
+	{
+		// Avoid cached transforms being used in controls by bumping the update counter. 
+		UpdateCounter += 1;
+		ResetCooldownCounter = SolverSettings.ResetCooldownFrames;
+	}
+
+	// If we found something invalid, or we're in cooldown, then we force the simulation to be as
+	// good as we can make it, and we don't write back to the hierarchy.
+	if (bGotInvalidSimulationData || ResetCooldownCounter > 0)
 	{
 		for (TPair<FRigComponentKey, FRigBodyRecord>& BodyRecordPair : BodyRecords)
 		{
@@ -1380,6 +1384,23 @@ void FRigPhysicsSimulation::UpdatePostPhysics(const FRigPhysicsSolverComponent*
 			}
 		}
 	}
+	else
+	{
+		// All is good - write the transforms we cached
+		for (const FRigComponentKey& ComponentKey : SortedBodyComponentKeys)
+		{
+			const FRigBodyRecord& Record = *BodyRecords.Find(ComponentKey);
+			if (Record.ActorHandle && Record.TargetElementKey.IsValid())
+			{
+				// Note that we set bAffectChildren = true (i.e. don't counter-animate
+				// children), so that attached animation bones will follow physics, but we
+				// rely on our bodies being sorted so we work out from the leaf nodes so as
+				// not to disturb previously set bodies.
+				Hierarchy->SetGlobalTransform(Record.TargetElementKey, Record.TargetComponentSpaceTM, false, true);
+
+			}
+		}
+	}
 }
 
 //======================================================================================================================
@@ -1422,6 +1443,11 @@ void FRigPhysicsSimulation::StepSimulation(
 		return;
 	}
 
+	if (ResetCooldownCounter > 0)
+	{
+		--ResetCooldownCounter;
+	}
+
 	const FRigPhysicsSolverSettings& SolverSettings = SolverComponent->SolverSettings;
 	const FRigPhysicsSimulationSpaceSettings& SimulationSpaceSettings = SolverComponent->SimulationSpaceSettings;
 
diff --git a/Engine/Plugins/Experimental/PhysicsControl/Source/ControlRigPhysics/Private/RigPhysicsSimulation.h b/Engine/Plugins/Experimental/PhysicsControl/Source/ControlRigPhysics/Private/RigPhysicsSimulation.h
--- a/Engine/Plugins/Experimental/PhysicsControl/Source/ControlRigPhysics/Private/RigPhysicsSimulation.h
+++ b/Engine/Plugins/Experimental/PhysicsControl/Source/ControlRigPhysics/Private/RigPhysicsSimulation.h
@@ -66,7 +66,12 @@ struct FRigBodyRecord
 	// Things that are set during instantiation
 	ImmediatePhysics::FActorHandle* ActorHandle = nullptr;
 
-	// Things that are updated as the simulation progresses
+	// Cache the element key for where we will write the simulation result
+	FRigElementKey TargetElementKey;
+
+	// The final/simulated TM is stored before writing it into the output, so we can avoid
+	// corrupting the output if anything is bad and we need to reset.
+	FTransform TargetComponentSpaceTM;
 };
 
 //======================================================================================================================
@@ -328,6 +333,10 @@ private:
 	// Note that the universe will roll over before a uint64 does
 	uint64 UpdateCounter = 0;
 
+	// If we have to reset, then we can be told to wait a few frames - this will be set > 0 and will
+	// be decremented each update. When it reaches 0, things can continue.
+	int ResetCooldownCounter = 0;
+
 	// This will be set to false after instantiation, which can happen manually during construction,
 	// or be postponed until there is a simulation step.
 	bool bNeedToInstantiate = true;
diff --git a/Engine/Plugins/Experimental/PhysicsControl/Source/ControlRigPhysics/Private/RigPhysicsSimulationSpace.cpp b/Engine/Plugins/Experimental/PhysicsControl/Source/ControlRigPhysics/Private/RigPhysicsSimulationSpace.cpp
--- a/Engine/Plugins/Experimental/PhysicsControl/Source/ControlRigPhysics/Private/RigPhysicsSimulationSpace.cpp
+++ b/Engine/Plugins/Experimental/PhysicsControl/Source/ControlRigPhysics/Private/RigPhysicsSimulationSpace.cpp
@@ -259,6 +259,8 @@ FRigPhysicsSimulation::FSimulationSpaceData FRigPhysicsSimulation::UpdateSimulat
 				SimulationSpaceSettings.OrientationChangeThresholdForTeleport / SimulationSpaceState.Dt)))
 		)
 	{
+		// Note that a teleport detection shouldn't change the pose, or the current motion. We just
+		// don't want to bring in that unwanted global motion.
 		SimulationSpaceData.LinearVelocity = FVector::ZeroVector;
 		SimulationSpaceData.AngularVelocity = FVector::ZeroVector;
 		SimulationSpaceData.LinearAcceleration = FVector::ZeroVector;
@@ -269,6 +271,8 @@ FRigPhysicsSimulation::FSimulationSpaceData FRigPhysicsSimulation::UpdateSimulat
 		SimulationSpaceState.PrevPrevSimulationSpaceTM = SimulationSpaceState.SimulationSpaceTM;
 		SimulationSpaceState.PrevDt = 0;
 		SimulationSpaceState.Dt = 0;
+		// Avoid cached transforms being used in controls by bumping the update counter. 
+		UpdateCounter += 1;
 	}
 	else
 	{
diff --git a/Engine/Plugins/Experimental/PhysicsControl/Source/ControlRigPhysics/Public/RigPhysicsData.h b/Engine/Plugins/Experimental/PhysicsControl/Source/ControlRigPhysics/Public/RigPhysicsData.h
--- a/Engine/Plugins/Experimental/PhysicsControl/Source/ControlRigPhysics/Public/RigPhysicsData.h
+++ b/Engine/Plugins/Experimental/PhysicsControl/Source/ControlRigPhysics/Public/RigPhysicsData.h
@@ -514,6 +514,19 @@ struct CONTROLRIGPHYSICS_API FRigPhysicsSolverSettings
 	// Disabled when the value is zero.
 	UPROPERTY(EditAnywhere, Category = Settings, meta = (ClampMin = "0.0"))
 	float DistanceThresholdForReset = 0;
+
+	// If any object in the simulation exceeds this velocity, we will reset the entire simulation
+	// (not just that object). This is to detect problems/explosions - for example if the target
+	// animation has teleports etc. Disabled when the value is zero.
+	UPROPERTY(EditAnywhere, Category = Settings, meta = (ClampMin = "0.0"))
+	float SpeedThresholdForReset = 10000.0f;
+
+	// The number of frames/updates to make sure a reset persists for. This can be useful when the
+	// resets are needed to deal with teleports that do not always complete in a single frame. For
+	// example, a value of 3 will disable simulation for 3 frames following a teleport, and then
+	// reset the simulation state to the input pose before continuing. Disabled when the value is zero.
+	UPROPERTY(EditAnywhere, Category = Settings, meta = (ClampMin = "0"))
+	int ResetCooldownFrames = 3;
 };
 
 /**
diff --git a/Engine/Plugins/Experimental/PhysicsControl/Source/PhysicsControl/Public/PhysicsControlObjectVersion.h b/Engine/Plugins/Experimental/PhysicsControl/Source/PhysicsControl/Public/PhysicsControlObjectVersion.h
--- a/Engine/Plugins/Experimental/PhysicsControl/Source/PhysicsControl/Public/PhysicsControlObjectVersion.h
+++ b/Engine/Plugins/Experimental/PhysicsControl/Source/PhysicsControl/Public/PhysicsControlObjectVersion.h
@@ -40,6 +40,10 @@ struct PHYSICSCONTROL_API FPhysicsControlObjectVersion
 
 		ControlRigDetectTeleportFromDistanceChange,
 
+		ControlRigSpeedThresholdForReset,
+
+		ControlRigResetCooldownFrames,
+
 		// -----<new versions can be added above this line>-------------------------------------------------
 		VersionPlusOne,
 		LatestVersion = VersionPlusOne - 1
