Expose more settings for control rig physics drag

diff --git a/Engine/Plugins/Experimental/PhysicsControl/Source/ControlRigPhysics/Private/RigPhysicsData.cpp b/Engine/Plugins/Experimental/PhysicsControl/Source/ControlRigPhysics/Private/RigPhysicsData.cpp
--- a/Engine/Plugins/Experimental/PhysicsControl/Source/ControlRigPhysics/Private/RigPhysicsData.cpp
+++ b/Engine/Plugins/Experimental/PhysicsControl/Source/ControlRigPhysics/Private/RigPhysicsData.cpp
@@ -207,7 +207,12 @@ FArchive& operator <<(FArchive& Ar, FRigPhysicsSimulationSpaceSettings& Data)
 		Ar << Data.PositionChangeThresholdForTeleport;
 		Ar << Data.OrientationChangeThresholdForTeleport;
 	}
-	Ar << Data.ExternalLinearDragV;
+	if (Ar.CustomVer(FPhysicsControlObjectVersion::GUID) >= FPhysicsControlObjectVersion::ControlRigDetectTeleportFromDistanceChange)
+	{
+		Ar << Data.LinearDragMultiplier;
+		Ar << Data.AngularDragMultiplier;
+	}
+	Ar << Data.ExternalLinearDrag;
 	Ar << Data.ExternalLinearVelocity;
 	Ar << Data.ExternalAngularVelocity;
 	return Ar;
diff --git a/Engine/Plugins/Experimental/PhysicsControl/Source/ControlRigPhysics/Private/RigPhysicsSimulation.cpp b/Engine/Plugins/Experimental/PhysicsControl/Source/ControlRigPhysics/Private/RigPhysicsSimulation.cpp
--- a/Engine/Plugins/Experimental/PhysicsControl/Source/ControlRigPhysics/Private/RigPhysicsSimulation.cpp
+++ b/Engine/Plugins/Experimental/PhysicsControl/Source/ControlRigPhysics/Private/RigPhysicsSimulation.cpp
@@ -1482,7 +1482,9 @@ void FRigPhysicsSimulation::StepSimulation(
 	// Simulation space
 	Chaos::FSimulationSpaceSettings ChaosSimulationSpaceSettings = Simulation->GetSimulationSpaceSettings();
 	ChaosSimulationSpaceSettings.bEnabled = (SimulationSpaceSettings.SpaceMovementAmount > 0.0f);
-	ChaosSimulationSpaceSettings.ExternalLinearEtherDrag = SimulationSpaceSettings.ExternalLinearDragV;
+	ChaosSimulationSpaceSettings.ExternalLinearEtherDrag = SimulationSpaceSettings.ExternalLinearDrag;
+	ChaosSimulationSpaceSettings.LinearVelocityAlpha = SimulationSpaceSettings.LinearDragMultiplier;
+	ChaosSimulationSpaceSettings.AngularVelocityAlpha = SimulationSpaceSettings.AngularDragMultiplier;
 	Simulation->SetSimulationSpaceSettings(ChaosSimulationSpaceSettings);
 	Simulation->UpdateSimulationSpace(
 		SimulationSpaceState.SimulationSpaceTM,
diff --git a/Engine/Plugins/Experimental/PhysicsControl/Source/ControlRigPhysics/Public/RigPhysicsData.h b/Engine/Plugins/Experimental/PhysicsControl/Source/ControlRigPhysics/Public/RigPhysicsData.h
--- a/Engine/Plugins/Experimental/PhysicsControl/Source/ControlRigPhysics/Public/RigPhysicsData.h
+++ b/Engine/Plugins/Experimental/PhysicsControl/Source/ControlRigPhysics/Public/RigPhysicsData.h
@@ -361,14 +361,26 @@ struct FRigPhysicsSimulationSpaceSettings
 	UPROPERTY(EditAnywhere, Category = Settings, meta = (ClampMin = "0.0"))
 	float OrientationChangeThresholdForTeleport = 45;
 
-	// Additional linear drag applied to every body in addition to linear drag specified on them in
-	// the physics asset. When combined with ExternalLinearVelocity, this can be used to add a
-	// temporary wind-blown effect without having to tune linear drag on all the bodies in the
-	// physics asset. The result is that each body has a force equal to -ExternalLinearDragV *
-	// ExternalLinearVelocity applied to it, in addition to all other forces. The vector is in
-	// simulation local space.
+	// How much of the simulation frame's linear velocity to pass onto the bodies (linear ether
+	// drag). This relies on there being drag authored on the bodies. If set to zero, the only drag
+	// will be from the body's local movement (in the simulation space).
+	UPROPERTY(EditAnywhere, Category = Settings, meta = (ClampMin = "0.0", ClampMax = "1.0"))
+	float LinearDragMultiplier = 1.0;
+
+	// How much of the simulation frame's angular velocity to pass onto the bodies (angular ether
+	// drag). This relies on there being drag authored on the bodies. If set to zero, the only drag
+	// will be from the body's local movement (in the simulation space).
+	UPROPERTY(EditAnywhere, Category = Settings, meta = (ClampMin = "0.0", ClampMax = "1.0"))
+	float AngularDragMultiplier = 1.0f;
+
+	// Additional linear drag from movement of the simulation space, applied to every body in
+	// addition to linear drag specified on them in the physics asset. When combined with
+	// ExternalLinearVelocity, this can be used to add a temporary wind-blown effect without having
+	// to tune linear drag on all the bodies in the physics asset. The result is that each body has
+	// a force equal to -ExternalLinearDrag * ExternalLinearVelocity applied to it, in addition to
+	// all other forces. The vector is in simulation local space.
 	UPROPERTY(EditAnywhere, Category = Settings)
-	FVector ExternalLinearDragV = FVector::ZeroVector;
+	FVector ExternalLinearDrag = FVector::ZeroVector;
 
 	// Additional velocity that is added to the component velocity so the simulation acts as if the
 	// actor is moving at speed, even when stationary. The vector is in world space. This could be
@@ -474,7 +486,7 @@ struct CONTROLRIGPHYSICS_API FRigPhysicsSolverSettings
 	 * which increases the cost.
 	 */
 	UPROPERTY(EditAnywhere, Category = SolverSettings, meta = (ClampMin = 0))
-	float FixedTimeStep = 0.01666667f;
+	float FixedTimeStep = 0.02f;
 
 	// The maximum number of solver steps that can be made
 	UPROPERTY(EditAnywhere, Category = SolverSettings, meta = (ClampMin = 1))
@@ -485,7 +497,7 @@ struct CONTROLRIGPHYSICS_API FRigPhysicsSolverSettings
 	 * requested total delta time then multiple steps will be used.
 	 */
 	UPROPERTY(EditAnywhere, Category = SolverSettings, meta = (ClampMin = 0))
-	float MaxDeltaTime = 1.f / 30.f;
+	float MaxDeltaTime = 0.02f;
 
 	/**
 	 * Whether to use the linear or non-linear solver for RBAN Joints. The linear solver is significantly cheaper than
