Improvements to the way control rig physics teleports are detected, and how resets are implemented.

diff --git a/Engine/Plugins/Experimental/PhysicsControl/Source/ControlRigPhysics/Private/RigPhysicsData.cpp b/Engine/Plugins/Experimental/PhysicsControl/Source/ControlRigPhysics/Private/RigPhysicsData.cpp
--- a/Engine/Plugins/Experimental/PhysicsControl/Source/ControlRigPhysics/Private/RigPhysicsData.cpp
+++ b/Engine/Plugins/Experimental/PhysicsControl/Source/ControlRigPhysics/Private/RigPhysicsData.cpp
@@ -243,14 +243,22 @@ FArchive& operator <<(FArchive& Ar, FRigPhysicsSolverSettings& Data)
 	Ar << Data.bUseLinearJointSolver;
 	Ar << Data.bSolveJointPositionsLast;
 	Ar << Data.bUseManifolds;
-	Ar << Data.DistanceThresholdForReset;
+	Ar << Data.PositionThresholdForReset;
 	if (Ar.CustomVer(FPhysicsControlObjectVersion::GUID) >= FPhysicsControlObjectVersion::ControlRigSpeedThresholdForReset)
 	{
-		Ar << Data.SpeedThresholdForReset;
+		Ar << Data.KinematicSpeedThresholdForReset;
 	}
-	if (Ar.CustomVer(FPhysicsControlObjectVersion::GUID) >= FPhysicsControlObjectVersion::ControlRigSpeedThresholdForReset)
+	if (Ar.CustomVer(FPhysicsControlObjectVersion::GUID) >= FPhysicsControlObjectVersion::ControlRigAccelerationThresholdForReset)
+	{
+		Ar << Data.KinematicAccelerationThresholdForReset;
+	}
+	if (Ar.CustomVer(FPhysicsControlObjectVersion::GUID) < FPhysicsControlObjectVersion::ControlRigRemoveResetCooldownFrames)
 	{
-		Ar << Data.ResetCooldownFrames;
+		if (Ar.CustomVer(FPhysicsControlObjectVersion::GUID) >= FPhysicsControlObjectVersion::ControlRigResetCooldownFrames)
+		{
+			int ResetCooldownFrames = 0;
+			Ar << ResetCooldownFrames;
+		}
 	}
 	if (Ar.CustomVer(FPhysicsControlObjectVersion::GUID) >= FPhysicsControlObjectVersion::ControlRigAutomaticallyAddPhysicsComponents)
 	{
diff --git a/Engine/Plugins/Experimental/PhysicsControl/Source/ControlRigPhysics/Private/RigPhysicsExecution.cpp b/Engine/Plugins/Experimental/PhysicsControl/Source/ControlRigPhysics/Private/RigPhysicsExecution.cpp
--- a/Engine/Plugins/Experimental/PhysicsControl/Source/ControlRigPhysics/Private/RigPhysicsExecution.cpp
+++ b/Engine/Plugins/Experimental/PhysicsControl/Source/ControlRigPhysics/Private/RigPhysicsExecution.cpp
@@ -60,6 +60,31 @@ FRigUnit_AddPhysicsSolver_Execute()
 	}
 }
 
+//======================================================================================================================
+FRigUnit_TrackInputPose_Execute()
+{
+	DECLARE_SCOPE_HIERARCHICAL_COUNTER_RIGUNIT();
+	if (!ExecuteContext.Hierarchy)
+	{
+		return;
+	}
+	if (UControlRig* ControlRig = Cast<UControlRig>(ExecuteContext.Hierarchy->GetOuter()))
+	{
+		if (FRigPhysicsSolverComponent* SolverComponent = Cast<FRigPhysicsSolverComponent>(
+			ExecuteContext.Hierarchy->FindComponent(PhysicsSolverComponentKey)))
+		{
+			if (bForceNumberOfFrames)
+			{
+				SolverComponent->TrackInputCounter = NumberOfFrames;
+			}
+			else
+			{
+				SolverComponent->TrackInputCounter = FMath::Max(SolverComponent->TrackInputCounter, NumberOfFrames);
+			}
+		}
+	}
+}
+
 //======================================================================================================================
 FRigUnit_StepPhysicsSolver_Execute()
 {
diff --git a/Engine/Plugins/Experimental/PhysicsControl/Source/ControlRigPhysics/Private/RigPhysicsSimulation.cpp b/Engine/Plugins/Experimental/PhysicsControl/Source/ControlRigPhysics/Private/RigPhysicsSimulation.cpp
--- a/Engine/Plugins/Experimental/PhysicsControl/Source/ControlRigPhysics/Private/RigPhysicsSimulation.cpp
+++ b/Engine/Plugins/Experimental/PhysicsControl/Source/ControlRigPhysics/Private/RigPhysicsSimulation.cpp
@@ -627,6 +627,9 @@ ImmediatePhysics::FActorHandle* FRigPhysicsSimulation::CreateBody(
 //======================================================================================================================
 void FRigPhysicsSimulation::Instantiate(const FRigPhysicsSolverComponent* SolverComponent)
 {
+	DECLARE_SCOPE_HIERARCHICAL_COUNTER_FUNC()
+	QUICK_SCOPE_CYCLE_COUNTER(STAT_RigPhysics_Instantiate);
+
 	if (!bNeedToInstantiate || !OwningControlRig)
 	{
 		return;
@@ -712,6 +715,9 @@ void FRigPhysicsSimulation::InstantiatePhysicsBodies(
 	const FRigPhysicsSolverComponent* SolverComponent,
 	FRigPhysicsIgnorePairs&           IgnorePairs)
 {
+	DECLARE_SCOPE_HIERARCHICAL_COUNTER_FUNC()
+	QUICK_SCOPE_CYCLE_COUNTER(STAT_RigPhysics_InstantiatePhysicsBodies);
+
 	const URigHierarchy* Hierarchy = OwningControlRig->GetHierarchy();
 
 	const FRigPhysicsSolverSettings& SolverSettings = SolverComponent->SolverSettings;
@@ -725,6 +731,8 @@ void FRigPhysicsSimulation::InstantiatePhysicsBodies(
 		if (const FRigPhysicsBodyComponent* PhysicsComponent = Cast<FRigPhysicsBodyComponent>(
 			Hierarchy->FindComponent(ComponentKey)))
 		{
+			Record.PhysicsComponent = PhysicsComponent;
+
 			const FRigPhysicsCollision& Collision = PhysicsComponent->Collision;
 			const FRigPhysicsDynamics& Dynamics = PhysicsComponent->Dynamics;
 			const FPhysicsControlModifierData& BodyData = PhysicsComponent->BodyData;
@@ -764,6 +772,9 @@ void FRigPhysicsSimulation::InstantiatePhysicsJoints(
 	const FRigPhysicsSolverComponent* SolverComponent,
 	FRigPhysicsIgnorePairs&           IgnorePairs)
 {
+	DECLARE_SCOPE_HIERARCHICAL_COUNTER_FUNC()
+	QUICK_SCOPE_CYCLE_COUNTER(STAT_RigPhysics_InstantiatePhysicsJoints);
+
 	const URigHierarchy* Hierarchy = OwningControlRig->GetHierarchy();
 
 	const FRigPhysicsSolverSettings& SolverSettings = SolverComponent->SolverSettings;
@@ -863,6 +874,9 @@ void FRigPhysicsSimulation::InstantiateControls(
 	const FRigPhysicsSolverComponent* SolverComponent,
 	FRigPhysicsIgnorePairs&           IgnorePairs)
 {
+	DECLARE_SCOPE_HIERARCHICAL_COUNTER_FUNC()
+	QUICK_SCOPE_CYCLE_COUNTER(STAT_RigPhysics_InstantiateControls);
+
 	const URigHierarchy* Hierarchy = OwningControlRig->GetHierarchy();
 
 	for (TPair<FRigComponentKey, FRigControlRecord>& ControlRecordPair : ControlRecords)
@@ -902,10 +916,10 @@ void FRigPhysicsSimulation::InstantiateControls(
 }
 
 //======================================================================================================================
-void FRigPhysicsSimulation::UpdateBodyPrePhysics(
-	const FRigVMExecuteContext&       ExecuteContext,
+void FRigPhysicsSimulation::UpdateBodyRecordPrePhysics(
 	const FRigPhysicsSolverComponent* SolverComponent,
-	const FRigBodyRecord&             Record, 
+	const float                       DeltaTime,
+	FRigBodyRecord&                   Record, 
 	const FRigPhysicsBodyComponent*   PhysicsComponent)
 {
 	URigHierarchy* Hierarchy = OwningControlRig->GetHierarchy();
@@ -915,8 +929,56 @@ void FRigPhysicsSimulation::UpdateBodyPrePhysics(
 	}
 	const FRigComponentKey ComponentKey = PhysicsComponent->GetKey();
 
+	// Shuffle the record data
+	Record.PrevSourceComponentSpaceVelocity = Record.SourceComponentSpaceVelocity;
+	Record.PrevSourceComponentSpaceAngularVelocity = Record.SourceComponentSpaceAngularVelocity;
+	Record.PrevSourceComponentSpaceTM = Record.SourceComponentSpaceTM;
+
+	FRigElementKey SourceKey = PhysicsComponent->BodySolverSettings.SourceBone;
+	if (!SourceKey.IsValid())
+	{
+		SourceKey = ComponentKey.ElementKey;
+	}
+	if (SourceKey.IsValid())
+	{
+		Record.SourceComponentSpaceTM = Hierarchy->GetGlobalTransform(SourceKey);
+	}
+
+	if (UpdateCounter == PreviousUpdateCounter + 1)
+	{
+		Record.SourceComponentSpaceVelocity = UE::PhysicsControl::CalculateLinearVelocity(
+			Record.PrevSourceComponentSpaceTM.GetTranslation(), Record.SourceComponentSpaceTM.GetTranslation(), DeltaTime);
+		Record.SourceComponentSpaceAngularVelocity = UE::PhysicsControl::CalculateAngularVelocity(
+			Record.PrevSourceComponentSpaceTM.GetRotation(), Record.SourceComponentSpaceTM.GetRotation(), DeltaTime);
+	}
+	else
+	{
+		Record.PrevSourceComponentSpaceTM = Record.SourceComponentSpaceTM;
+		Record.SourceComponentSpaceVelocity = FVector::ZeroVector;
+		Record.SourceComponentSpaceAngularVelocity = FVector::ZeroVector;
+	}
+}
+
+//======================================================================================================================
+void FRigPhysicsSimulation::UpdateBodyPrePhysics(
+	const FRigVMExecuteContext&       ExecuteContext, 
+	const FRigPhysicsSolverComponent* SolverComponent,
+	const FRigBodyRecord&             Record, 
+	const FRigPhysicsBodyComponent*   PhysicsComponent)
+{
+	const FRigComponentKey ComponentKey = PhysicsComponent->GetKey();
+
+	FPhysicsControlModifierData BodyData = PhysicsComponent->BodyData;
+	ERigPhysicsKinematicTargetSpace KinematicTargetSpace = PhysicsComponent->KinematicTargetSpace;
+
+	if (SolverComponent->TrackInputCounter > 0)
+	{
+		BodyData.MovementType = EPhysicsMovementType::Kinematic;
+		KinematicTargetSpace = ERigPhysicsKinematicTargetSpace::IgnoreTarget;
+	}
+
 	UpdateBodyFromModifierData(
-		Record.ActorHandle, Simulation.Get(), PhysicsComponent->BodyData, SimulationSpaceData.Gravity);
+		Record.ActorHandle, Simulation.Get(), BodyData, SimulationSpaceData.Gravity);
 
 	if (Record.ActorHandle->GetIsKinematic())
 	{
@@ -949,19 +1011,17 @@ void FRigPhysicsSimulation::UpdateBodyPrePhysics(
 			}
 			if (SourceKey.IsValid())
 			{
-				FTransform SourceComponentSpaceTM = Hierarchy->GetGlobalTransform(SourceKey);
-
 				switch (PhysicsComponent->KinematicTargetSpace)
 				{
 				case ERigPhysicsKinematicTargetSpace::OffsetInBoneSpace:
 				{
-					KinematicTargetCS = PhysicsComponent->KinematicTarget * SourceComponentSpaceTM;
+					KinematicTargetCS = PhysicsComponent->KinematicTarget * Record.SourceComponentSpaceTM.ToTransform();
 					break;
 				}
 				case ERigPhysicsKinematicTargetSpace::OffsetInWorldSpace:
 				{
 					// Convert the bone to WS, apply the target, and convert back
-					FTransform BoneWS = SourceComponentSpaceTM * SimulationSpaceState.ComponentTM;
+					FTransform BoneWS = Record.SourceComponentSpaceTM.ToTransform() * SimulationSpaceState.ComponentTM;
 					FTransform KinematicTargetWS = BoneWS * PhysicsComponent->KinematicTarget;
 					// Danny TODO figure out which of the GetRelativeTransform versions this is
 					KinematicTargetCS = SimulationSpaceState.ComponentTM.Inverse() * KinematicTargetWS;
@@ -969,7 +1029,12 @@ void FRigPhysicsSimulation::UpdateBodyPrePhysics(
 				}
 				case ERigPhysicsKinematicTargetSpace::OffsetInComponentSpace:
 				{
-					KinematicTargetCS = SourceComponentSpaceTM * PhysicsComponent->KinematicTarget;
+					KinematicTargetCS = Record.SourceComponentSpaceTM.ToTransform() * PhysicsComponent->KinematicTarget;
+					break;
+				}
+				case ERigPhysicsKinematicTargetSpace::IgnoreTarget:
+				{
+					KinematicTargetCS = Record.SourceComponentSpaceTM.ToTransform();
 					break;
 				}
 				default:
@@ -1054,6 +1119,9 @@ void FRigPhysicsSimulation::UpdateJointPrePhysics(
 				return;
 			}
 
+			// Danny TODO now all transforms are being cached in the record, get them from there
+			// rather than the hierarchy
+
 			// Note that the drive operates between a parent and child part, so we
 			// don't need to worry about global/component (etc) space.
 			const FTransform ChildTM = Hierarchy->GetGlobalTransform(ChildSourceKey);
@@ -1124,11 +1192,62 @@ static UE::PhysicsControl::FPosQuat CalculateTargetTM(
 	return ChildTargetTM;
 }
 
+//======================================================================================================================
+void FRigPhysicsSimulation::CheckForResetsPrePhysics(
+	FRigPhysicsSolverComponent* SolverComponent, const float DeltaTime)
+{
+	DECLARE_SCOPE_HIERARCHICAL_COUNTER_FUNC()
+	QUICK_SCOPE_CYCLE_COUNTER(STAT_RigPhysics_CheckForResetsPrePhysics);
+	const FRigPhysicsSolverSettings& SolverSettings = SolverComponent->SolverSettings;
+
+	double SpeedThresholdForResetSquared = SolverSettings.KinematicSpeedThresholdForReset > 0 ?
+		FMath::Square(SolverSettings.KinematicSpeedThresholdForReset) : TNumericLimits<double>::Max();
+
+	double AccelerationThresholdForResetSquared = SolverSettings.KinematicAccelerationThresholdForReset > 0 ?
+		FMath::Square(SolverSettings.KinematicAccelerationThresholdForReset) : TNumericLimits<double>::Max();
+
+	double HighestSpeedSq = -1.0;
+	double HighestAccelerationSq = -1.0;
+
+	for (TPair<FRigComponentKey, FRigBodyRecord>& BodyRecordPair : BodyRecords)
+	{
+		const FRigComponentKey& ComponentKey = BodyRecordPair.Key;
+		FRigBodyRecord& Record = BodyRecordPair.Value;
+
+		if (Record.ActorHandle && Record.ActorHandle->GetIsKinematic())
+		{
+			const FVector Velocity = Record.SourceComponentSpaceVelocity;
+			const FVector Acceleration =
+				DeltaTime > SMALL_NUMBER && (UpdateCounter == PreviousUpdateCounter + 1)
+				? (Record.SourceComponentSpaceVelocity - Record.PrevSourceComponentSpaceVelocity) / DeltaTime
+				: FVector::ZeroVector;
+
+			HighestSpeedSq = FMath::Max(HighestSpeedSq, Velocity.SquaredLength());
+			HighestAccelerationSq = FMath::Max(HighestAccelerationSq, Acceleration.SquaredLength());
+		}
+	}
+
+	if (HighestSpeedSq > SpeedThresholdForResetSquared || HighestAccelerationSq > AccelerationThresholdForResetSquared)
+	{
+		if (HighestSpeedSq > SpeedThresholdForResetSquared)
+		{
+			UE_LOG(LogRigPhysics, Warning, TEXT("Speed %f triggered reset in %s"),
+				FMath::Sqrt(HighestSpeedSq), *OwningControlRig->GetName());
+		}
+		if (HighestAccelerationSq > AccelerationThresholdForResetSquared)
+		{
+			UE_LOG(LogRigPhysics, Warning, TEXT("Acceleration %f triggered reset in %s"),
+				FMath::Sqrt(HighestAccelerationSq), *OwningControlRig->GetName());
+		}
+		SolverComponent->TrackInputCounter = FMath::Max(SolverComponent->TrackInputCounter, 2);
+	}
+}
+
 //======================================================================================================================
 void FRigPhysicsSimulation::UpdatePrePhysics(
-	const FRigVMExecuteContext&       ExecuteContext, 
-	const FRigPhysicsSolverComponent* SolverComponent, 
-	const float                       DeltaTime)
+	const FRigVMExecuteContext& ExecuteContext, 
+	FRigPhysicsSolverComponent* SolverComponent, 
+	const float                 DeltaTime)
 {
 	DECLARE_SCOPE_HIERARCHICAL_COUNTER_FUNC()
 	QUICK_SCOPE_CYCLE_COUNTER(STAT_RigPhysics_UpdatePrePhysics);
@@ -1146,35 +1265,62 @@ void FRigPhysicsSimulation::UpdatePrePhysics(
 	{
 		const FRigComponentKey& ComponentKey = BodyRecordPair.Key;
 		FRigBodyRecord& Record = BodyRecordPair.Value;
-
-		if (const FRigPhysicsBodyComponent* PhysicsComponent = Cast<FRigPhysicsBodyComponent>(
-			Hierarchy->FindComponent(ComponentKey)))
+		if (const FRigPhysicsBodyComponent* PhysicsComponent = Record.PhysicsComponent)
 		{
-			UpdateBodyPrePhysics(ExecuteContext, SolverComponent, Record, PhysicsComponent);
+			UpdateBodyRecordPrePhysics(SolverComponent, DeltaTime, Record, PhysicsComponent);
 		}
 	}
 
-	for (TPair<FRigComponentKey, FRigJointRecord>& JointRecordPair : JointRecords)
+	CheckForResetsPrePhysics(SolverComponent, DeltaTime);
+	if (SolverComponent->TrackInputCounter > 0)
 	{
-		const FRigComponentKey& ComponentKey = JointRecordPair.Key;
-		FRigJointRecord& JointRecord = JointRecordPair.Value;
+		UE_LOG(LogRigPhysics, Warning, TEXT("Forcing tracking of input for %s"), *OwningControlRig->GetName());
+	}
 
-		if (const FRigPhysicsJointComponent* JointComponent = Cast<FRigPhysicsJointComponent>(
-			Hierarchy->FindComponent(ComponentKey)))
+
+	{
+		DECLARE_SCOPE_HIERARCHICAL_COUNTER_FUNC()
+		QUICK_SCOPE_CYCLE_COUNTER(STAT_RigPhysics_UpdateBodiesPrePhysics);
+		for (TPair<FRigComponentKey, FRigBodyRecord>& BodyRecordPair : BodyRecords)
 		{
-			UpdateJointPrePhysics(JointRecord, JointComponent, Hierarchy, DeltaTime);
+			const FRigComponentKey& ComponentKey = BodyRecordPair.Key;
+			FRigBodyRecord& Record = BodyRecordPair.Value;
+			if (const FRigPhysicsBodyComponent* PhysicsComponent = Record.PhysicsComponent)
+			{
+				UpdateBodyPrePhysics(ExecuteContext, SolverComponent, Record, PhysicsComponent);
+			}
 		}
 	}
 
-	for (TPair<FRigComponentKey, FRigControlRecord>& ControlRecordPair : ControlRecords)
 	{
-		const FRigComponentKey& ComponentKey = ControlRecordPair.Key;
-		FRigControlRecord& ControlRecord = ControlRecordPair.Value;
+		DECLARE_SCOPE_HIERARCHICAL_COUNTER_FUNC()
+		QUICK_SCOPE_CYCLE_COUNTER(STAT_RigPhysics_UpdateJointsPrePhysics);
+		for (TPair<FRigComponentKey, FRigJointRecord>& JointRecordPair : JointRecords)
+		{
+			const FRigComponentKey& ComponentKey = JointRecordPair.Key;
+			FRigJointRecord& JointRecord = JointRecordPair.Value;
 
-		if (const FRigPhysicsControlComponent* ControlComponent = Cast<FRigPhysicsControlComponent>(
-			Hierarchy->FindComponent(ComponentKey)))
+			if (const FRigPhysicsJointComponent* JointComponent = Cast<FRigPhysicsJointComponent>(
+				Hierarchy->FindComponent(ComponentKey)))
+			{
+				UpdateJointPrePhysics(JointRecord, JointComponent, Hierarchy, DeltaTime);
+			}
+		}
+	}
+
+	{
+		DECLARE_SCOPE_HIERARCHICAL_COUNTER_FUNC()
+		QUICK_SCOPE_CYCLE_COUNTER(STAT_RigPhysics_UpdateControlsPrePhysics);
+		for (TPair<FRigComponentKey, FRigControlRecord>& ControlRecordPair : ControlRecords)
 		{
-			UpdateControlPrePhysics(ControlRecord, ControlComponent, Hierarchy, DeltaTime);
+			const FRigComponentKey& ComponentKey = ControlRecordPair.Key;
+			FRigControlRecord& ControlRecord = ControlRecordPair.Value;
+
+			if (const FRigPhysicsControlComponent* ControlComponent = Cast<FRigPhysicsControlComponent>(
+				Hierarchy->FindComponent(ComponentKey)))
+			{
+				UpdateControlPrePhysics(ControlRecord, ControlComponent, Hierarchy, DeltaTime);
+			}
 		}
 	}
 }
@@ -1287,7 +1433,8 @@ void FRigPhysicsSimulation::UpdateControlPrePhysics(
 //======================================================================================================================
 // Note that we read back into a target bone, which may have been specified explicitly, or will
 // otherwise default to the physics element parent.
-void FRigPhysicsSimulation::UpdatePostPhysics(const FRigPhysicsSolverComponent* SolverComponent, float Alpha)
+void FRigPhysicsSimulation::UpdatePostPhysics(
+	FRigPhysicsSolverComponent* SolverComponent, const float Alpha, const float DeltaTime)
 {
 	DECLARE_SCOPE_HIERARCHICAL_COUNTER_FUNC()
 	QUICK_SCOPE_CYCLE_COUNTER(STAT_RigPhysics_UpdatePostPhysics);
@@ -1300,11 +1447,9 @@ void FRigPhysicsSimulation::UpdatePostPhysics(const FRigPhysicsSolverComponent*
 	bool bGotInvalidSimulationData = false;
 	URigHierarchy* Hierarchy = OwningControlRig->GetHierarchy();
 
-	double SpeedThresholdForResetSquared = SolverSettings.SpeedThresholdForReset > 0 ?
-		FMath::Square(SolverSettings.SpeedThresholdForReset) : TNumericLimits<double>::Max();
-
-	double DistanceThresholdForResetSquared = SolverSettings.DistanceThresholdForReset > 0 ?
-		FMath::Square(SolverSettings.DistanceThresholdForReset) : TNumericLimits<double>::Max();
+	double PositionThresholdForResetSquared = SolverSettings.PositionThresholdForReset > 0 ?
+		FMath::Square(SolverSettings.PositionThresholdForReset) : TNumericLimits<double>::Max();
+	double HighestPosition = -1.0;
 
 	// Traverse using the sorted keys
 	for (const FRigComponentKey& ComponentKey : SortedBodyComponentKeys)
@@ -1312,14 +1457,12 @@ void FRigPhysicsSimulation::UpdatePostPhysics(const FRigPhysicsSolverComponent*
 		FRigBodyRecord& Record = *BodyRecords.Find(ComponentKey);
 		if (Record.ActorHandle && Record.TargetElementKey.IsValid())
 		{
+			// Check the simulation output
 			FTransform SimSpaceTM = Record.ActorHandle->GetWorldTransform();
 			double DistSq = SimSpaceTM.GetTranslation().SquaredLength();
-			if (!SimSpaceTM.IsValid() || DistSq > DistanceThresholdForResetSquared)
-			{
-				bGotInvalidSimulationData = true;
-			}
-			else if (Record.ActorHandle->GetLinearVelocity().SquaredLength() > SpeedThresholdForResetSquared)
+			if (!SimSpaceTM.IsValid() || DistSq > PositionThresholdForResetSquared)
 			{
+				HighestPosition = FMath::Max(HighestPosition, SimSpaceTM.GetTranslation().Length());
 				bGotInvalidSimulationData = true;
 			}
 
@@ -1333,7 +1476,7 @@ void FRigPhysicsSimulation::UpdatePostPhysics(const FRigPhysicsSolverComponent*
 				// This can cause joint separation. We probably want an option to blend
 				// in local (joint) space, perhaps splitting the alpha into orientation
 				// and position.
-				FTransform CurrentTM = Hierarchy->GetGlobalTransform(Record.TargetElementKey);
+				FTransform CurrentTM = Record.SourceComponentSpaceTM.ToTransform();
 				FQuat TargetQ = FQuat::Slerp(
 					CurrentTM.GetRotation(), Record.TargetComponentSpaceTM.GetRotation(), Alpha);
 				FVector TargetT = FMath::Lerp(
@@ -1346,22 +1489,30 @@ void FRigPhysicsSimulation::UpdatePostPhysics(const FRigPhysicsSolverComponent*
 
 	if (bGotInvalidSimulationData)
 	{
+		if (HighestPosition > 0)
+		{
+			UE_LOG(LogRigPhysics, Warning, TEXT("Position %f triggered teleport in %s - resetting pose"),
+				HighestPosition, *OwningControlRig->GetName());
+		}
 		// Avoid cached transforms being used in controls by bumping the update counter. 
 		UpdateCounter += 1;
-		ResetCooldownCounter = SolverSettings.ResetCooldownFrames;
+		// Set this to 3 since it gets decremented at the end of the update, and we need it to take
+		// effect at the start of the next update.
+		SolverComponent->TrackInputCounter = FMath::Max(SolverComponent->TrackInputCounter, 3);
 	}
 
-	// If we found something invalid, or we're in cooldown, then we force the simulation to be as
-	// good as we can make it, and we don't write back to the hierarchy.
-	if (bGotInvalidSimulationData || ResetCooldownCounter > 0)
+	// If we found something invalid then we force the simulation to be as good as we can make it,
+	// and we don't write back to the hierarchy.
+	if (HighestPosition > 0)
 	{
+		UE_LOG(LogRigPhysics, Warning, TEXT("Resetting state to input pose in %s"), 
+			*OwningControlRig->GetName());
 		for (TPair<FRigComponentKey, FRigBodyRecord>& BodyRecordPair : BodyRecords)
 		{
 			const FRigComponentKey& ComponentKey = BodyRecordPair.Key;
 			const FRigBodyRecord& Record = BodyRecordPair.Value;
 
-			if (const FRigPhysicsBodyComponent* PhysicsComponent = Cast<FRigPhysicsBodyComponent>(
-				Hierarchy->FindComponent(ComponentKey)))
+			if (const FRigPhysicsBodyComponent* PhysicsComponent = Record.PhysicsComponent)
 			{
 				if (Record.ActorHandle && !Record.ActorHandle->GetIsKinematic())
 				{
@@ -1373,13 +1524,13 @@ void FRigPhysicsSimulation::UpdatePostPhysics(const FRigPhysicsSolverComponent*
 					}
 					if (SourceKey.IsValid())
 					{
-						FTransform SourceComponentSpaceTM = Hierarchy->GetGlobalTransform(SourceKey);
+						FTransform SourceComponentSpaceTM = Record.SourceComponentSpaceTM.ToTransform();
 						const FTransform SourceSimulationSpaceTM =
 							ConvertComponentSpaceTransformToSimSpace(SolverSettings, SourceComponentSpaceTM);
 						Record.ActorHandle->SetWorldTransform(SourceSimulationSpaceTM);
-						Record.ActorHandle->SetLinearVelocity(FVector::ZeroVector);
-						Record.ActorHandle->SetAngularVelocity(FVector::ZeroVector);
 					}
+					Record.ActorHandle->SetLinearVelocity(FVector::ZeroVector);
+					Record.ActorHandle->SetAngularVelocity(FVector::ZeroVector);
 				}
 			}
 		}
@@ -1397,7 +1548,6 @@ void FRigPhysicsSimulation::UpdatePostPhysics(const FRigPhysicsSolverComponent*
 				// rely on our bodies being sorted so we work out from the leaf nodes so as
 				// not to disturb previously set bodies.
 				Hierarchy->SetGlobalTransform(Record.TargetElementKey, Record.TargetComponentSpaceTM, false, true);
-
 			}
 		}
 	}
@@ -1405,15 +1555,18 @@ void FRigPhysicsSimulation::UpdatePostPhysics(const FRigPhysicsSolverComponent*
 
 //======================================================================================================================
 void FRigPhysicsSimulation::StepSimulation(
-	const FRigVMExecuteContext&       ExecuteContext, 
-	const FRigPhysicsSolverComponent* SolverComponent, 
-	const float                       DeltaTimeOverride, 
-	const float                       SimulationSpaceDeltaTimeOverride,
-	const float                       Alpha)
+	const FRigVMExecuteContext& ExecuteContext,
+	FRigPhysicsSolverComponent* SolverComponent,
+	const float                 DeltaTimeOverride,
+	const float                 SimulationSpaceDeltaTimeOverride,
+	const float                 Alpha)
 {
 	DECLARE_SCOPE_HIERARCHICAL_COUNTER_FUNC()
 	QUICK_SCOPE_CYCLE_COUNTER(STAT_RigPhysics_StepSimulation);
 
+	// Increment the update counter at the start - and always update it so this tells us our "frame number".
+	++UpdateCounter;
+
 	float PhysicsDeltaTime = ExecuteContext.GetDeltaTime();
 	if (DeltaTimeOverride > 0.0f)
 	{
@@ -1443,11 +1596,6 @@ void FRigPhysicsSimulation::StepSimulation(
 		return;
 	}
 
-	if (ResetCooldownCounter > 0)
-	{
-		--ResetCooldownCounter;
-	}
-
 	const FRigPhysicsSolverSettings& SolverSettings = SolverComponent->SolverSettings;
 	const FRigPhysicsSimulationSpaceSettings& SimulationSpaceSettings = SolverComponent->SimulationSpaceSettings;
 
@@ -1506,10 +1654,15 @@ void FRigPhysicsSimulation::StepSimulation(
 		Simulation->Simulate(
 			PhysicsDeltaTime, SolverSettings.MaxDeltaTime, MaxTimeSteps,
 			SimulationSpaceData.Gravity, &ChaosJointSolverSettings);
+
+		PreviousUpdateCounter = UpdateCounter;
 	}
 
 	// Always do a read-back, even for zero Dt
-	UpdatePostPhysics(SolverComponent, Alpha);
+	UpdatePostPhysics(SolverComponent, Alpha, PhysicsDeltaTime);
 
-	++UpdateCounter;
+	if (SolverComponent->TrackInputCounter > 0)
+	{
+		--SolverComponent->TrackInputCounter;
+	}
 }
diff --git a/Engine/Plugins/Experimental/PhysicsControl/Source/ControlRigPhysics/Private/RigPhysicsSimulation.h b/Engine/Plugins/Experimental/PhysicsControl/Source/ControlRigPhysics/Private/RigPhysicsSimulation.h
--- a/Engine/Plugins/Experimental/PhysicsControl/Source/ControlRigPhysics/Private/RigPhysicsSimulation.h
+++ b/Engine/Plugins/Experimental/PhysicsControl/Source/ControlRigPhysics/Private/RigPhysicsSimulation.h
@@ -63,6 +63,9 @@ struct FRigBodyRecord
 {
 	GENERATED_BODY()
 
+	// The component that generated us;
+	const FRigPhysicsBodyComponent* PhysicsComponent = nullptr;
+
 	// Things that are set during instantiation
 	ImmediatePhysics::FActorHandle* ActorHandle = nullptr;
 
@@ -72,6 +75,16 @@ struct FRigBodyRecord
 	// The final/simulated TM is stored before writing it into the output, so we can avoid
 	// corrupting the output if anything is bad and we need to reset.
 	FTransform TargetComponentSpaceTM;
+
+	// These are updated for all records in UpdatePrePhysics. The times/validity are determined by
+	// the CurrentDeltaTime, PrevDeltaTime and update counters in the simulation itself.
+	UE::PhysicsControl::FPosQuat SourceComponentSpaceTM;
+	FVector                      SourceComponentSpaceVelocity = FVector::ZeroVector;
+	FVector                      SourceComponentSpaceAngularVelocity = FVector::ZeroVector;
+
+	UE::PhysicsControl::FPosQuat PrevSourceComponentSpaceTM;
+	FVector                      PrevSourceComponentSpaceVelocity = FVector::ZeroVector;
+	FVector                      PrevSourceComponentSpaceAngularVelocity = FVector::ZeroVector;
 };
 
 //======================================================================================================================
@@ -143,11 +156,11 @@ public:
 	// If it is zero, then delta time is taken from the execute context
 	// If it is negative, then the simulation isn't stepped.
 	void StepSimulation(
-		const FRigVMExecuteContext&       ExecuteContext, 
-		const FRigPhysicsSolverComponent* SolverComponent, 
-		const float                       DeltaTimeOverride, 
-		const float                       SimulationSpaceDeltaTimeOverride,
-		const float                       Alpha);
+		const FRigVMExecuteContext& ExecuteContext, 
+		FRigPhysicsSolverComponent* SolverComponent, 
+		const float                 DeltaTimeOverride, 
+		const float                 SimulationSpaceDeltaTimeOverride,
+		const float                 Alpha);
 
 	// Draws shapes etc AND (potentially) enables the low-level Chaos debug draw
 	void Draw(
@@ -239,11 +252,17 @@ private:
 
 	void UpdatePrePhysics(
 		const FRigVMExecuteContext&       ExecuteContext, 
-		const FRigPhysicsSolverComponent* SolverComponent, 
+		FRigPhysicsSolverComponent*       SolverComponent, 
 		const float                       DeltaTime);
 
+	void UpdateBodyRecordPrePhysics(
+		const FRigPhysicsSolverComponent* SolverComponent,
+		const float                       DeltaTime,
+		FRigBodyRecord&                   Record,
+		const FRigPhysicsBodyComponent*   PhysicsComponent);
+
 	void UpdateBodyPrePhysics(
-		const FRigVMExecuteContext&       ExecuteContext,
+		const FRigVMExecuteContext&       ExecuteContext, 
 		const FRigPhysicsSolverComponent* SolverComponent,
 		const FRigBodyRecord&             Record, 
 		const FRigPhysicsBodyComponent*   PhysicsComponent);
@@ -260,10 +279,15 @@ private:
 		const URigHierarchy*               Hierarchy,
 		const float                        DeltaTime);
 
-	void UpdatePostPhysics(const FRigPhysicsSolverComponent* SolverComponent, const float Alpha);
+	void UpdatePostPhysics(
+		FRigPhysicsSolverComponent* SolverComponent, 
+		const float                 Alpha, 
+		const float                 DeltaTime);
 
 	void InitSimulationSpace(const FTransform& ComponentToWorld, const FTransform& BoneToComponent);
 
+	void CheckForResetsPrePhysics(FRigPhysicsSolverComponent* SolverComponent, const float DeltaTime);
+
 	FSimulationSpaceData UpdateSimulationSpaceStateAndCalculateData(
 		const FRigPhysicsSolverComponent* SolverComponent, const float Dt);
 
@@ -333,9 +357,8 @@ private:
 	// Note that the universe will roll over before a uint64 does
 	uint64 UpdateCounter = 0;
 
-	// If we have to reset, then we can be told to wait a few frames - this will be set > 0 and will
-	// be decremented each update. When it reaches 0, things can continue.
-	int ResetCooldownCounter = 0;
+	// The update counter when PrevSourceComponentSpaceTM etc were written. Check this before using them
+	uint64 PreviousUpdateCounter = 0;
 
 	// This will be set to false after instantiation, which can happen manually during construction,
 	// or be postponed until there is a simulation step.
diff --git a/Engine/Plugins/Experimental/PhysicsControl/Source/ControlRigPhysics/Private/RigPhysicsSimulationSpace.cpp b/Engine/Plugins/Experimental/PhysicsControl/Source/ControlRigPhysics/Private/RigPhysicsSimulationSpace.cpp
--- a/Engine/Plugins/Experimental/PhysicsControl/Source/ControlRigPhysics/Private/RigPhysicsSimulationSpace.cpp
+++ b/Engine/Plugins/Experimental/PhysicsControl/Source/ControlRigPhysics/Private/RigPhysicsSimulationSpace.cpp
@@ -7,6 +7,8 @@
 #include "Rigs/RigHierarchyController.h"
 #include "ControlRig.h"
 
+#include "PhysicsControlHelpers.h"
+
 #include "Physics/ImmediatePhysics/ImmediatePhysicsActorHandle.h"
 #include "Physics/ImmediatePhysics/ImmediatePhysicsJointHandle.h"
 #include "Physics/ImmediatePhysics/ImmediatePhysicsSimulation.h"
@@ -130,19 +132,6 @@ void FRigPhysicsSimulation::InitSimulationSpace(
 	SimulationSpaceState.BoneRelComponentTM = BoneRelComponentTM;
 }
 
-//======================================================================================================================
-FVector CalculateAngularVelocity(const FQuat& PrevQ, const FQuat& CurrentQ, float Dt)
-{
-	const FQuat DeltaQ = (CurrentQ * PrevQ.Inverse()).GetShortestArcWith(FQuat::Identity);
-	return DeltaQ.ToRotationVector() / Dt;
-}
-
-//======================================================================================================================
-FVector CalculateLinearVelocity(const FVector& PrevP, const FVector& CurrentP, float Dt)
-{
-	return (CurrentP - PrevP) / Dt;
-}
-
 //======================================================================================================================
 // Note - don't use the space conversion functions here as the state won't have been set yet.
 FRigPhysicsSimulation::FSimulationSpaceData FRigPhysicsSimulation::UpdateSimulationSpaceStateAndCalculateData(
@@ -213,26 +202,26 @@ FRigPhysicsSimulation::FSimulationSpaceData FRigPhysicsSimulation::UpdateSimulat
 	// finite difference approximations, but this wouldn't necessarily be correct.
 
 	// World-space component linear velocity and acceleration
-	SimulationSpaceData.LinearVelocity = CalculateLinearVelocity(
+	SimulationSpaceData.LinearVelocity = UE::PhysicsControl::CalculateLinearVelocity(
 		SimulationSpaceState.PrevSimulationSpaceTM.GetTranslation(),
 		SimulationSpaceState.SimulationSpaceTM.GetTranslation(), SimulationSpaceState.Dt);
 	const FVector PrevSpaceLinearVel = 
 		SimulationSpaceState.PrevDt < SMALL_NUMBER 
 		? SimulationSpaceData.LinearVelocity 
-		: CalculateLinearVelocity(
+		: UE::PhysicsControl::CalculateLinearVelocity(
 			SimulationSpaceState.PrevPrevSimulationSpaceTM.GetTranslation(),
 			SimulationSpaceState.PrevSimulationSpaceTM.GetTranslation(), SimulationSpaceState.PrevDt);
 	SimulationSpaceData.LinearAcceleration =
 		(SimulationSpaceData.LinearVelocity - PrevSpaceLinearVel) / SimulationSpaceState.Dt;
 
 	// World-space component angular velocity and acceleration
-	SimulationSpaceData.AngularVelocity = CalculateAngularVelocity(
+	SimulationSpaceData.AngularVelocity = UE::PhysicsControl::CalculateAngularVelocity(
 		SimulationSpaceState.PrevSimulationSpaceTM.GetRotation(),
 		SimulationSpaceState.SimulationSpaceTM.GetRotation(), SimulationSpaceState.Dt);
 	const FVector PrevSpaceAngularVel = 
 		SimulationSpaceState.PrevDt < SMALL_NUMBER 
 		? SimulationSpaceData.AngularVelocity
-		: CalculateAngularVelocity(
+		: UE::PhysicsControl::CalculateAngularVelocity(
 			SimulationSpaceState.PrevPrevSimulationSpaceTM.GetRotation(),
 			SimulationSpaceState.PrevSimulationSpaceTM.GetRotation(), SimulationSpaceState.PrevDt);
 	SimulationSpaceData.AngularAcceleration =
@@ -242,23 +231,52 @@ FRigPhysicsSimulation::FSimulationSpaceData FRigPhysicsSimulation::UpdateSimulat
 	SimulationSpaceData.LinearVelocity.Z *= SimulationSpaceSettings.VelocityScaleZ;
 	SimulationSpaceData.LinearAcceleration.Z *= SimulationSpaceSettings.VelocityScaleZ;
 
+	bool bLinearAccelerationTrigger = SimulationSpaceSettings.LinearAccelerationThresholdForTeleport > 0 &&
+		SimulationSpaceData.LinearAcceleration.SquaredLength() >
+		FMath::Square(SimulationSpaceSettings.LinearAccelerationThresholdForTeleport);
+	bool bAngularAccelerationTrigger = SimulationSpaceSettings.AngularAccelerationThresholdForTeleport > 0 &&
+		SimulationSpaceData.AngularAcceleration.SquaredLength() >
+		FMath::Square(FMath::DegreesToRadians(SimulationSpaceSettings.AngularAccelerationThresholdForTeleport));
+	bool bPositionTrigger = SimulationSpaceSettings.PositionChangeThresholdForTeleport > 0 &&
+		SimulationSpaceData.LinearVelocity.SquaredLength() >
+		FMath::Square(SimulationSpaceSettings.PositionChangeThresholdForTeleport / SimulationSpaceState.Dt);
+	bool bOrientationTrigger = SimulationSpaceSettings.OrientationChangeThresholdForTeleport > 0 &&
+		SimulationSpaceData.AngularVelocity.SquaredLength() >
+		FMath::Square(FMath::DegreesToRadians(
+			SimulationSpaceSettings.OrientationChangeThresholdForTeleport / SimulationSpaceState.Dt));
+
 	// Clamped world-space motion of the simulation space
-	if (
-		(SimulationSpaceSettings.LinearAccelerationThresholdForTeleport > 0 &&
-			SimulationSpaceData.LinearAcceleration.SquaredLength() >
-			FMath::Square(SimulationSpaceSettings.LinearAccelerationThresholdForTeleport)) ||
-		(SimulationSpaceSettings.AngularAccelerationThresholdForTeleport > 0 &&
-			SimulationSpaceData.AngularAcceleration.SquaredLength() >
-			FMath::Square(SimulationSpaceSettings.AngularAccelerationThresholdForTeleport)) ||
-		(SimulationSpaceSettings.PositionChangeThresholdForTeleport > 0 &&
-			SimulationSpaceData.LinearVelocity.SquaredLength() >
-			FMath::Square(SimulationSpaceSettings.PositionChangeThresholdForTeleport / SimulationSpaceState.Dt)) ||
-		(SimulationSpaceSettings.OrientationChangeThresholdForTeleport > 0 &&
-			SimulationSpaceData.AngularVelocity.SquaredLength() >
-			FMath::Square(FMath::DegreesToRadians(
-				SimulationSpaceSettings.OrientationChangeThresholdForTeleport / SimulationSpaceState.Dt)))
-		)
+	if (bLinearAccelerationTrigger || bAngularAccelerationTrigger || bPositionTrigger || bOrientationTrigger)
 	{
+		if (bLinearAccelerationTrigger)
+		{
+			UE_LOG(LogRigPhysics, Log, TEXT("Detected linear Acceleration (%f > %f) teleport in %s"),
+				SimulationSpaceData.LinearAcceleration.Length(), 
+				SimulationSpaceSettings.LinearAccelerationThresholdForTeleport,
+				*OwningControlRig->GetName());
+		}
+		if (bAngularAccelerationTrigger)
+		{
+			UE_LOG(LogRigPhysics, Log, TEXT("Detected angular Acceleration (%f > %f) teleport in %s"),
+				SimulationSpaceData.AngularAcceleration.Length(),
+				SimulationSpaceSettings.AngularAccelerationThresholdForTeleport,
+				*OwningControlRig->GetName());
+		}
+		if (bPositionTrigger)
+		{
+			UE_LOG(LogRigPhysics, Log, TEXT("Detected position (%f > %f) teleport in %s"),
+				SimulationSpaceData.LinearVelocity.Length() * SimulationSpaceState.Dt,
+				SimulationSpaceSettings.PositionChangeThresholdForTeleport,
+				*OwningControlRig->GetName());
+		}
+		if (bOrientationTrigger)
+		{
+			UE_LOG(LogRigPhysics, Log, TEXT("Detected orientation (%f > %f) teleport in %s"),
+				FMath::RadiansToDegrees(SimulationSpaceData.AngularVelocity.Length() * SimulationSpaceState.Dt),
+				SimulationSpaceSettings.OrientationChangeThresholdForTeleport,
+				*OwningControlRig->GetName());
+		}
+
 		// Note that a teleport detection shouldn't change the pose, or the current motion. We just
 		// don't want to bring in that unwanted global motion.
 		SimulationSpaceData.LinearVelocity = FVector::ZeroVector;
diff --git a/Engine/Plugins/Experimental/PhysicsControl/Source/ControlRigPhysics/Private/RigPhysicsSimulationVisualization.cpp b/Engine/Plugins/Experimental/PhysicsControl/Source/ControlRigPhysics/Private/RigPhysicsSimulationVisualization.cpp
--- a/Engine/Plugins/Experimental/PhysicsControl/Source/ControlRigPhysics/Private/RigPhysicsSimulationVisualization.cpp
+++ b/Engine/Plugins/Experimental/PhysicsControl/Source/ControlRigPhysics/Private/RigPhysicsSimulationVisualization.cpp
@@ -347,9 +347,7 @@ void FRigPhysicsSimulation::Draw(
 	for (const TPair<FRigComponentKey, FRigBodyRecord>& BodyRecordPair : BodyRecords)
 	{
 		const FRigBodyRecord& Record = BodyRecordPair.Value;
-		const ImmediatePhysics::FActorHandle* ActorHandle = Record.ActorHandle;
-
-		if (ActorHandle)
+		if (const ImmediatePhysics::FActorHandle* ActorHandle = Record.ActorHandle)
 		{
 			DrawActor(DI, VisualizationSettings, SpaceTransform, *ActorHandle);
 		}
@@ -357,9 +355,7 @@ void FRigPhysicsSimulation::Draw(
 	for (const TPair<FRigComponentKey, FRigJointRecord>& JointRecordPair : JointRecords)
 	{
 		const FRigJointRecord& Record = JointRecordPair.Value;
-		const ImmediatePhysics::FJointHandle* JointHandle = Record.JointHandle;
-
-		if (JointHandle)
+		if (const ImmediatePhysics::FJointHandle* JointHandle = Record.JointHandle)
 		{
 			DrawJoint(DI, VisualizationSettings, SpaceTransform, *JointHandle);
 		}
diff --git a/Engine/Plugins/Experimental/PhysicsControl/Source/ControlRigPhysics/Public/RigPhysicsData.h b/Engine/Plugins/Experimental/PhysicsControl/Source/ControlRigPhysics/Public/RigPhysicsData.h
--- a/Engine/Plugins/Experimental/PhysicsControl/Source/ControlRigPhysics/Public/RigPhysicsData.h
+++ b/Engine/Plugins/Experimental/PhysicsControl/Source/ControlRigPhysics/Public/RigPhysicsData.h
@@ -59,25 +59,21 @@ struct FRigPhysicsVisualizationSettings
 UENUM()
 enum class ERigPhysicsKinematicTargetSpace : uint8
 {
-	/** Target is in world space */
+	// Target is in world space
 	World,
-	/** Target is in the space of the component */
+	// Target is in the space of the component
 	Component,
-	/**
-	 * Target is specified relative to the bone that controls the (kinematic) body, with the offset applied 
-	 * in the space of that body. 
-	 */
+	// Target is specified relative to the bone that controls the (kinematic) body, with the offset
+	// applied in the space of that body.
 	OffsetInBoneSpace,
-	/**
-	 * Target is specified relative to the bone that controls the (kinematic) body, with the offset applied
-	 * in world space. 
-	 */
+	// Target is specified relative to the bone that controls the (kinematic) body, with the offset
+	// applied in world space.
 	OffsetInWorldSpace,
-	/**
-	 * Target is specified relative to the bone that controls the (kinematic) body, with the offset applied
-	 * in component space. 
-	 */
+	// Target is specified relative to the bone that controls the (kinematic) body, with the offset
+	// applied in component space.
 	OffsetInComponentSpace,
+	// Ignore the specified target, and just track the bone
+	IgnoreTarget
 };
 
 /** Specifies what space the simulation should run in */
@@ -338,14 +334,14 @@ struct FRigPhysicsSimulationSpaceSettings
 	// simulation state.
 	// A value of zero will disable this detection
 	UPROPERTY(EditAnywhere, Category = Settings, meta = (ClampMin = "0.0"))
-	float LinearAccelerationThresholdForTeleport = 100000;
+	float LinearAccelerationThresholdForTeleport = 10000;
 
-	// If the angular acceleration (rand/s/s) is above this threshold, the movement is treated as a teleport.
+	// If the angular acceleration (deg/s/s) is above this threshold, the movement is treated as a teleport.
 	// The calculated accelerations and velocities will be set to zero, without resetting the
 	// simulation state.
 	// A value of zero will disable this detection
 	UPROPERTY(EditAnywhere, Category = Settings, meta = (ClampMin = "0.0"))
-	float AngularAccelerationThresholdForTeleport = 100000;
+	float AngularAccelerationThresholdForTeleport = 10000;
 
 	// If the linear position change is above this threshold, the movement is treated as a teleport.
 	// The calculated accelerations and velocities will be set to zero, without resetting the
@@ -359,7 +355,7 @@ struct FRigPhysicsSimulationSpaceSettings
 	// simulation state.
 	// A value of zero will disable this detection
 	UPROPERTY(EditAnywhere, Category = Settings, meta = (ClampMin = "0.0"))
-	float OrientationChangeThresholdForTeleport = 45;
+	float OrientationChangeThresholdForTeleport = 30;
 
 	// How much of the simulation frame's linear velocity to pass onto the bodies (linear ether
 	// drag). This relies on there being drag authored on the bodies. If set to zero, the only drag
@@ -522,23 +518,22 @@ struct CONTROLRIGPHYSICS_API FRigPhysicsSolverSettings
 	bool bUseManifolds = true;
 
 	// If any object in the simulation exceeds this distance from the simulation origin, we will
-	// reset the entire simulation. This is to detect problems/explosions.
+	// reset the pose and velocity of the entire simulation. This is to detect problems/explosions.
 	// Disabled when the value is zero.
 	UPROPERTY(EditAnywhere, Category = Settings, meta = (ClampMin = "0.0"))
-	float DistanceThresholdForReset = 0;
+	float PositionThresholdForReset = 0;
+
+	// If any kinematic object in the simulation exceeds this speed, we will reset the velocities
+	// of all objects in the simulation.This is to detect problems - for example if the
+	// target animation has teleports etc. Disabled when the value is zero.
+	UPROPERTY(EditAnywhere, Category = Settings, meta = (ClampMin = "0.0"))
+	float KinematicSpeedThresholdForReset = 5000.0f;
 
-	// If any object in the simulation exceeds this velocity, we will reset the entire simulation
-	// (not just that object). This is to detect problems/explosions - for example if the target
-	// animation has teleports etc. Disabled when the value is zero.
+	// If any kinematic object in the simulation exceeds this acceleration, we will reset the
+	// velocities of all objects in the simulation.This is to detect problems - for example if the
+	// target animation has teleports etc. Disabled when the value is zero.
 	UPROPERTY(EditAnywhere, Category = Settings, meta = (ClampMin = "0.0"))
-	float SpeedThresholdForReset = 10000.0f;
-
-	// The number of frames/updates to make sure a reset persists for. This can be useful when the
-	// resets are needed to deal with teleports that do not always complete in a single frame. For
-	// example, a value of 3 will disable simulation for 3 frames following a teleport, and then
-	// reset the simulation state to the input pose before continuing. Disabled when the value is zero.
-	UPROPERTY(EditAnywhere, Category = Settings, meta = (ClampMin = "0"))
-	int ResetCooldownFrames = 3;
+	float KinematicAccelerationThresholdForReset = 25000.0f;
 };
 
 /**
diff --git a/Engine/Plugins/Experimental/PhysicsControl/Source/ControlRigPhysics/Public/RigPhysicsExecution.h b/Engine/Plugins/Experimental/PhysicsControl/Source/ControlRigPhysics/Public/RigPhysicsExecution.h
--- a/Engine/Plugins/Experimental/PhysicsControl/Source/ControlRigPhysics/Public/RigPhysicsExecution.h
+++ b/Engine/Plugins/Experimental/PhysicsControl/Source/ControlRigPhysics/Public/RigPhysicsExecution.h
@@ -137,6 +137,39 @@ struct CONTROLRIGPHYSICS_API FRigUnit_StepPhysicsSolver : public FRigUnit_Physic
 	FRigPhysicsVisualizationSettings VisualizationSettings;
 };
 
+
+/**
+ * Forces tracking of the input animation (on all physics bodies) for the next N frames
+ */
+USTRUCT(meta = (DisplayName = "Track Input Pose", Keywords = "Simulate"))
+struct CONTROLRIGPHYSICS_API FRigUnit_TrackInputPose : public FRigUnit_PhysicsBaseMutable
+{
+	GENERATED_BODY()
+
+	FRigUnit_TrackInputPose()
+	{
+		PhysicsSolverComponentKey.ElementKey.Type = ERigElementType::Bone;
+		PhysicsSolverComponentKey.Name = FRigPhysicsSolverComponent::GetDefaultName();
+	}
+
+	RIGVM_METHOD()
+	virtual void Execute() override;
+
+	// The solver to relate this new physics element to
+	UPROPERTY(meta = (Input))
+	FRigComponentKey PhysicsSolverComponentKey;
+
+	// The number of frames to track the input pose for
+	UPROPERTY(meta = (Input, ClampMin = "0"))
+	int NumberOfFrames = 1;
+
+	// If true, then the number will be forced, potentially reducing the number. If false, then the
+	// NumberOfFrames will only be used to increase the number of frames remaining.
+	UPROPERTY(meta = (Input))
+	bool bForceNumberOfFrames = false;
+};
+
+
 // Adds a set of physics components including the body, joint and controls
 USTRUCT(meta = (DisplayName = "Add Physics Components", Keywords = "Construction,Create,New", Varying))
 struct CONTROLRIGPHYSICS_API FRigUnit_AddPhysicsComponents : public FRigUnit_PhysicsBaseMutable
diff --git a/Engine/Plugins/Experimental/PhysicsControl/Source/ControlRigPhysics/Public/RigPhysicsSolverComponent.h b/Engine/Plugins/Experimental/PhysicsControl/Source/ControlRigPhysics/Public/RigPhysicsSolverComponent.h
--- a/Engine/Plugins/Experimental/PhysicsControl/Source/ControlRigPhysics/Public/RigPhysicsSolverComponent.h
+++ b/Engine/Plugins/Experimental/PhysicsControl/Source/ControlRigPhysics/Public/RigPhysicsSolverComponent.h
@@ -25,6 +25,11 @@ public:
 	UPROPERTY(BlueprintReadOnly, EditAnywhere, Category = Physics, meta = (ShowOnlyInnerProperties))
 	FRigPhysicsSimulationSpaceSettings SimulationSpaceSettings;
 
+	// If we have to reset, then we can be told to track the input kinematically for a number of
+	// frames - this will be set > 0 and will be decremented each update. When it reaches 0, things
+	// can go to the desired simulation state.
+	int32 TrackInputCounter = 0;
+
 	virtual void Save(FArchive& Ar) override;
 	virtual void Load(FArchive& Ar) override;
 
diff --git a/Engine/Plugins/Experimental/PhysicsControl/Source/PhysicsControl/Public/PhysicsControlHelpers.h b/Engine/Plugins/Experimental/PhysicsControl/Source/PhysicsControl/Public/PhysicsControlHelpers.h
--- a/Engine/Plugins/Experimental/PhysicsControl/Source/PhysicsControl/Public/PhysicsControlHelpers.h
+++ b/Engine/Plugins/Experimental/PhysicsControl/Source/PhysicsControl/Public/PhysicsControlHelpers.h
@@ -4,6 +4,7 @@
 
 #include "UObject/NameTypes.h"
 #include "Math/Vector.h"
+#include "Math/Quat.h"
 
 struct FBodyInstance;
 struct FPhysicsControlData;
@@ -47,6 +48,17 @@ PHYSICSCONTROL_API FBodyInstance* GetBodyInstance(UPrimitiveComponent* Component
  */
 PHYSICSCONTROL_API FName GetPhysicalParentBone(USkeletalMeshComponent* SkeletalComponent, FName BoneName);
 
+FORCEINLINE FVector CalculateAngularVelocity(const FQuat& PrevQ, const FQuat& CurrentQ, float Dt)
+{
+	const FQuat DeltaQ = (CurrentQ * PrevQ.Inverse()).GetShortestArcWith(FQuat::Identity);
+	return DeltaQ.ToRotationVector() / Dt;
+}
+
+FORCEINLINE FVector CalculateLinearVelocity(const FVector& PrevP, const FVector& CurrentP, float Dt)
+{
+	return (CurrentP - PrevP) / Dt;
+}
+
 /**
  * Converts strength/damping ratio/extra damping into spring stiffness/damping.
  */
diff --git a/Engine/Plugins/Experimental/PhysicsControl/Source/PhysicsControl/Public/PhysicsControlObjectVersion.h b/Engine/Plugins/Experimental/PhysicsControl/Source/PhysicsControl/Public/PhysicsControlObjectVersion.h
--- a/Engine/Plugins/Experimental/PhysicsControl/Source/PhysicsControl/Public/PhysicsControlObjectVersion.h
+++ b/Engine/Plugins/Experimental/PhysicsControl/Source/PhysicsControl/Public/PhysicsControlObjectVersion.h
@@ -46,6 +46,10 @@ struct PHYSICSCONTROL_API FPhysicsControlObjectVersion
 
 		ControlRigSimSpaceDragMultipliers,
 
+		ControlRigAccelerationThresholdForReset,
+
+		ControlRigRemoveResetCooldownFrames,
+
 		// -----<new versions can be added above this line>-------------------------------------------------
 		VersionPlusOne,
 		LatestVersion = VersionPlusOne - 1
