Revert the caching of the FRigPhysicsBodyComponents - was causing crashes in control rig.

diff --git a/Engine/Plugins/Experimental/PhysicsControl/Source/ControlRigPhysics/Private/RigPhysicsSimulation.cpp b/Engine/Plugins/Experimental/PhysicsControl/Source/ControlRigPhysics/Private/RigPhysicsSimulation.cpp
--- a/Engine/Plugins/Experimental/PhysicsControl/Source/ControlRigPhysics/Private/RigPhysicsSimulation.cpp
+++ b/Engine/Plugins/Experimental/PhysicsControl/Source/ControlRigPhysics/Private/RigPhysicsSimulation.cpp
@@ -731,8 +731,6 @@ void FRigPhysicsSimulation::InstantiatePhysicsBodies(
 		if (const FRigPhysicsBodyComponent* PhysicsComponent = Cast<FRigPhysicsBodyComponent>(
 			Hierarchy->FindComponent(ComponentKey)))
 		{
-			Record.PhysicsComponent = PhysicsComponent;
-
 			const FRigPhysicsCollision& Collision = PhysicsComponent->Collision;
 			const FRigPhysicsDynamics& Dynamics = PhysicsComponent->Dynamics;
 			const FPhysicsControlModifierData& BodyData = PhysicsComponent->BodyData;
@@ -1265,7 +1263,8 @@ void FRigPhysicsSimulation::UpdatePrePhysics(
 	{
 		const FRigComponentKey& ComponentKey = BodyRecordPair.Key;
 		FRigBodyRecord& Record = BodyRecordPair.Value;
-		if (const FRigPhysicsBodyComponent* PhysicsComponent = Record.PhysicsComponent)
+		if (const FRigPhysicsBodyComponent* PhysicsComponent = Cast<FRigPhysicsBodyComponent>(
+			Hierarchy->FindComponent(ComponentKey)))
 		{
 			UpdateBodyRecordPrePhysics(SolverComponent, DeltaTime, Record, PhysicsComponent);
 		}
@@ -1285,7 +1284,8 @@ void FRigPhysicsSimulation::UpdatePrePhysics(
 		{
 			const FRigComponentKey& ComponentKey = BodyRecordPair.Key;
 			FRigBodyRecord& Record = BodyRecordPair.Value;
-			if (const FRigPhysicsBodyComponent* PhysicsComponent = Record.PhysicsComponent)
+			if (const FRigPhysicsBodyComponent* PhysicsComponent = Cast<FRigPhysicsBodyComponent>(
+				Hierarchy->FindComponent(ComponentKey)))
 			{
 				UpdateBodyPrePhysics(ExecuteContext, SolverComponent, Record, PhysicsComponent);
 			}
@@ -1512,7 +1512,8 @@ void FRigPhysicsSimulation::UpdatePostPhysics(
 			const FRigComponentKey& ComponentKey = BodyRecordPair.Key;
 			const FRigBodyRecord& Record = BodyRecordPair.Value;
 
-			if (const FRigPhysicsBodyComponent* PhysicsComponent = Record.PhysicsComponent)
+			if (const FRigPhysicsBodyComponent* PhysicsComponent = Cast<FRigPhysicsBodyComponent>(
+				Hierarchy->FindComponent(ComponentKey)))
 			{
 				if (Record.ActorHandle && !Record.ActorHandle->GetIsKinematic())
 				{
diff --git a/Engine/Plugins/Experimental/PhysicsControl/Source/ControlRigPhysics/Private/RigPhysicsSimulation.h b/Engine/Plugins/Experimental/PhysicsControl/Source/ControlRigPhysics/Private/RigPhysicsSimulation.h
--- a/Engine/Plugins/Experimental/PhysicsControl/Source/ControlRigPhysics/Private/RigPhysicsSimulation.h
+++ b/Engine/Plugins/Experimental/PhysicsControl/Source/ControlRigPhysics/Private/RigPhysicsSimulation.h
@@ -63,9 +63,6 @@ struct FRigBodyRecord
 {
 	GENERATED_BODY()
 
-	// The component that generated us;
-	const FRigPhysicsBodyComponent* PhysicsComponent = nullptr;
-
 	// Things that are set during instantiation
 	ImmediatePhysics::FActorHandle* ActorHandle = nullptr;
 
