Small change to the default acceleration threshold for control rig physics resets

diff --git a/Engine/Plugins/Experimental/PhysicsControl/Source/ControlRigPhysics/Public/RigPhysicsData.h b/Engine/Plugins/Experimental/PhysicsControl/Source/ControlRigPhysics/Public/RigPhysicsData.h
--- a/Engine/Plugins/Experimental/PhysicsControl/Source/ControlRigPhysics/Public/RigPhysicsData.h
+++ b/Engine/Plugins/Experimental/PhysicsControl/Source/ControlRigPhysics/Public/RigPhysicsData.h
@@ -533,7 +533,7 @@ struct CONTROLRIGPHYSICS_API FRigPhysicsSolverSettings
 	// velocities of all objects in the simulation.This is to detect problems - for example if the
 	// target animation has teleports etc. Disabled when the value is zero.
 	UPROPERTY(EditAnywhere, Category = Settings, meta = (ClampMin = "0.0"))
-	float KinematicAccelerationThresholdForReset = 25000.0f;
+	float KinematicAccelerationThresholdForReset = 40000.0f;
 };
 
 /**
