Fix initialization of control rig physics to avoid spurious controls

diff --git a/Engine/Plugins/Experimental/PhysicsControl/Source/ControlRigPhysics/Private/RigPhysicsSimulation.h b/Engine/Plugins/Experimental/PhysicsControl/Source/ControlRigPhysics/Private/RigPhysicsSimulation.h
--- a/Engine/Plugins/Experimental/PhysicsControl/Source/ControlRigPhysics/Private/RigPhysicsSimulation.h
+++ b/Engine/Plugins/Experimental/PhysicsControl/Source/ControlRigPhysics/Private/RigPhysicsSimulation.h
@@ -105,7 +105,7 @@ struct FRigJointRecord
 	// The drive works with velocities so we store the previous target transform, and when it was stored.
 	UE::PhysicsControl::FPosQuat    PreviousDriveTargetTM;
 	// This is stored from the main solver update counter, marking when the previous drive TM was valid.
-	uint64                          PreviousDriveTargetUpdateCounter = 0;
+	int64                           PreviousDriveTargetUpdateCounter = -999;
 };
 
 //======================================================================================================================
@@ -129,7 +129,7 @@ struct FRigControlRecord
 	// The control works with velocities so we store the previous target transform, and when it was stored.
 	UE::PhysicsControl::FPosQuat    PreviousTargetTM;
 	// This is stored from the main solver update counter, marking when the previous target TM was valid.
-	uint64                          PreviousTargetUpdateCounter = 0;
+	int64                           PreviousTargetUpdateCounter = -999;
 };
 
 
@@ -352,10 +352,10 @@ private:
 	// This is incremented at the end of each simulation step, used to identify when previously
 	// calculated values are valid
 	// Note that the universe will roll over before a uint64 does
-	uint64 UpdateCounter = 0;
+	int64 UpdateCounter = 0;
 
 	// The update counter when PrevSourceComponentSpaceTM etc were written. Check this before using them
-	uint64 PreviousUpdateCounter = 0;
+	int64 PreviousUpdateCounter = -999;
 
 	// This will be set to false after instantiation, which can happen manually during construction,
 	// or be postponed until there is a simulation step.
