Fix for calculation of physics control rig control targets (sim vs world vs bone space)

diff --git a/Engine/Plugins/Experimental/PhysicsControl/Source/ControlRigPhysics/Private/RigPhysicsSimulation.cpp b/Engine/Plugins/Experimental/PhysicsControl/Source/ControlRigPhysics/Private/RigPhysicsSimulation.cpp
--- a/Engine/Plugins/Experimental/PhysicsControl/Source/ControlRigPhysics/Private/RigPhysicsSimulation.cpp
+++ b/Engine/Plugins/Experimental/PhysicsControl/Source/ControlRigPhysics/Private/RigPhysicsSimulation.cpp
@@ -1319,7 +1319,7 @@ void FRigPhysicsSimulation::UpdatePrePhysics(
 			if (const FRigPhysicsControlComponent* ControlComponent = Cast<FRigPhysicsControlComponent>(
 				Hierarchy->FindComponent(ComponentKey)))
 			{
-				UpdateControlPrePhysics(ControlRecord, ControlComponent, Hierarchy, DeltaTime);
+				UpdateControlPrePhysics(ControlRecord, ControlComponent, SolverComponent, Hierarchy, DeltaTime);
 			}
 		}
 	}
@@ -1329,9 +1329,12 @@ void FRigPhysicsSimulation::UpdatePrePhysics(
 void FRigPhysicsSimulation::UpdateControlPrePhysics(
 	FRigControlRecord&                 ControlRecord,
 	const FRigPhysicsControlComponent* ControlComponent,
+	const FRigPhysicsSolverComponent*  SolverComponent,
 	const URigHierarchy*               Hierarchy,
 	const float                        DeltaTime)
 {
+	const FRigPhysicsSolverSettings& SolverSettings = SolverComponent->SolverSettings;
+
 	ImmediatePhysics::FJointHandle* JointHandle = ControlRecord.JointHandle;
 	if (!JointHandle)
 	{
@@ -1374,15 +1377,17 @@ void FRigPhysicsSimulation::UpdateControlPrePhysics(
 			Constraint->SetChildConnectorLocation(
 				ControlComponent->ControlData.GetControlPoint(ChildActorHandle));
 
-			UE::PhysicsControl::FPosQuat TargetTM(
+			FTransform TargetTM(
 				ControlComponent->ControlTarget.TargetOrientation,
 				ControlComponent->ControlTarget.TargetPosition);
 
 			if (ControlComponent->ControlData.bUseSkeletalAnimation)
 			{
-				UE::PhysicsControl::FPosQuat AnimTargetTM = CalculateTargetTM(
-					Hierarchy, JointSettings, ControlRecord);
-				TargetTM = AnimTargetTM * TargetTM;
+				const FTransform ComponentSpaceAnimTargetTM = CalculateTargetTM(
+					Hierarchy, JointSettings, ControlRecord).ToTransform();
+				const FTransform SimSpaceAnimTargetTM = ConvertComponentSpaceTransformToSimSpace(
+					SolverSettings, ComponentSpaceAnimTargetTM);
+				TargetTM = TargetTM * SimSpaceAnimTargetTM;
 			}
 
 			Constraint->SetLinearDrivePositionTarget(TargetTM.GetTranslation());
diff --git a/Engine/Plugins/Experimental/PhysicsControl/Source/ControlRigPhysics/Private/RigPhysicsSimulation.h b/Engine/Plugins/Experimental/PhysicsControl/Source/ControlRigPhysics/Private/RigPhysicsSimulation.h
--- a/Engine/Plugins/Experimental/PhysicsControl/Source/ControlRigPhysics/Private/RigPhysicsSimulation.h
+++ b/Engine/Plugins/Experimental/PhysicsControl/Source/ControlRigPhysics/Private/RigPhysicsSimulation.h
@@ -127,7 +127,7 @@ struct FRigControlRecord
 	// Things that are updated as the simulation progresses
 
 	// The control works with velocities so we store the previous target transform, and when it was stored.
-	UE::PhysicsControl::FPosQuat    PreviousTargetTM;
+	FTransform                      PreviousTargetTM;
 	// This is stored from the main solver update counter, marking when the previous target TM was valid.
 	int64                           PreviousTargetUpdateCounter = -999;
 };
@@ -273,6 +273,7 @@ private:
 	void UpdateControlPrePhysics(
 		FRigControlRecord&                 ControlRecord,
 		const FRigPhysicsControlComponent* ControlComponent,
+		const FRigPhysicsSolverComponent*  SolverComponent,
 		const URigHierarchy*               Hierarchy,
 		const float                        DeltaTime);
 
