Improved rig physics logging

diff --git a/Engine/Plugins/Experimental/PhysicsControl/Source/ControlRigPhysics/Private/RigPhysicsSimulation.cpp b/Engine/Plugins/Experimental/PhysicsControl/Source/ControlRigPhysics/Private/RigPhysicsSimulation.cpp
--- a/Engine/Plugins/Experimental/PhysicsControl/Source/ControlRigPhysics/Private/RigPhysicsSimulation.cpp
+++ b/Engine/Plugins/Experimental/PhysicsControl/Source/ControlRigPhysics/Private/RigPhysicsSimulation.cpp
@@ -1229,15 +1229,15 @@ void FRigPhysicsSimulation::CheckForResetsPrePhysics(
 	{
 		if (HighestSpeedSq > SpeedThresholdForResetSquared)
 		{
-			UE_LOG(LogRigPhysics, Warning, TEXT("Speed %f triggered reset in %s"),
+			UE_LOG(LogRigPhysics, Log, TEXT("Speed %f triggered reset in %s"),
 				FMath::Sqrt(HighestSpeedSq), *OwningControlRig->GetName());
 		}
 		if (HighestAccelerationSq > AccelerationThresholdForResetSquared)
 		{
-			UE_LOG(LogRigPhysics, Warning, TEXT("Acceleration %f triggered reset in %s"),
+			UE_LOG(LogRigPhysics, Log, TEXT("Acceleration %f triggered reset in %s"),
 				FMath::Sqrt(HighestAccelerationSq), *OwningControlRig->GetName());
 		}
-		SolverComponent->TrackInputCounter = FMath::Max(SolverComponent->TrackInputCounter, 2);
+		SolverComponent->TrackInputCounter = FMath::Max(SolverComponent->TrackInputCounter, 3);
 	}
 }
 
@@ -1273,10 +1273,10 @@ void FRigPhysicsSimulation::UpdatePrePhysics(
 	CheckForResetsPrePhysics(SolverComponent, DeltaTime);
 	if (SolverComponent->TrackInputCounter > 0)
 	{
-		UE_LOG(LogRigPhysics, Warning, TEXT("Forcing tracking of input for %s"), *OwningControlRig->GetName());
+		UE_LOG(LogRigPhysics, Log, TEXT("Forcing tracking (counter = %d) of input for %s"), 
+			SolverComponent->TrackInputCounter, *OwningControlRig->GetName());
 	}
 
-
 	{
 		DECLARE_SCOPE_HIERARCHICAL_COUNTER_FUNC()
 		QUICK_SCOPE_CYCLE_COUNTER(STAT_RigPhysics_UpdateBodiesPrePhysics);
@@ -1496,7 +1496,7 @@ void FRigPhysicsSimulation::UpdatePostPhysics(
 	{
 		if (HighestPosition > 0)
 		{
-			UE_LOG(LogRigPhysics, Warning, TEXT("Position %f triggered teleport in %s - resetting pose"),
+			UE_LOG(LogRigPhysics, Log, TEXT("Position %f triggered teleport in %s - resetting pose"),
 				HighestPosition, *OwningControlRig->GetName());
 		}
 		// Avoid cached transforms being used in controls by bumping the update counter. 
@@ -1510,7 +1510,7 @@ void FRigPhysicsSimulation::UpdatePostPhysics(
 	// and we don't write back to the hierarchy.
 	if (HighestPosition > 0)
 	{
-		UE_LOG(LogRigPhysics, Warning, TEXT("Resetting state to input pose in %s"), 
+		UE_LOG(LogRigPhysics, Log, TEXT("Resetting state to input pose in %s"), 
 			*OwningControlRig->GetName());
 		for (TPair<FRigComponentKey, FRigBodyRecord>& BodyRecordPair : BodyRecords)
 		{
