Add code support for updating the targets.

diff --git a/Engine/Plugins/Experimental/PhysicsControl/Source/ControlRigPhysics/Private/RigPhysicsControlExecution.cpp b/Engine/Plugins/Experimental/PhysicsControl/Source/ControlRigPhysics/Private/RigPhysicsControlExecution.cpp
--- a/Engine/Plugins/Experimental/PhysicsControl/Source/ControlRigPhysics/Private/RigPhysicsControlExecution.cpp
+++ b/Engine/Plugins/Experimental/PhysicsControl/Source/ControlRigPhysics/Private/RigPhysicsControlExecution.cpp
@@ -1,7 +1,9 @@
 // Copyright Epic Games, Inc. All Rights Reserved.
 
 #include "RigPhysicsControlExecution.h"
+
 #include "RigPhysicsControlComponent.h"
+#include "PhysicsControlHelpers.h"
 
 #include "Rigs/RigHierarchyController.h"
 #include "Units/RigUnitContext.h"
@@ -45,6 +47,20 @@ FRigUnit_AddPhysicsControl_Execute()
 	}
 }
 
+//======================================================================================================================
+FRigUnit_HierarchySetControlEnabled_Execute()
+{
+	DECLARE_SCOPE_HIERARCHICAL_COUNTER_RIGUNIT();
+	if (ExecuteContext.Hierarchy)
+	{
+		if (FRigPhysicsControlComponent* Component = Cast<FRigPhysicsControlComponent>(
+			ExecuteContext.Hierarchy->FindComponent(PhysicsControlComponentKey)))
+		{
+			Component->ControlData.bEnabled = bEnabled;
+		}
+	}
+}
+
 //======================================================================================================================
 FRigUnit_HierarchySetControlData_Execute()
 {
@@ -73,6 +89,25 @@ FRigUnit_HierarchySetControlTarget_Execute()
 	}
 }
 
+//======================================================================================================================
+FRigUnit_HierarchyUpdateControlTarget_Execute()
+{
+	DECLARE_SCOPE_HIERARCHICAL_COUNTER_RIGUNIT();
+	if (ExecuteContext.Hierarchy)
+	{
+		if (FRigPhysicsControlComponent* Component = Cast<FRigPhysicsControlComponent>(
+			ExecuteContext.Hierarchy->FindComponent(PhysicsControlComponentKey)))
+		{
+			Component->ControlTarget.TargetVelocity = UE::PhysicsControl::CalculateLinearVelocity(
+				TargetPosition, Component->ControlTarget.TargetPosition, DeltaTime);
+			Component->ControlTarget.TargetAngularVelocity = UE::PhysicsControl::CalculateAngularVelocity(
+				TargetOrientation.Quaternion(), Component->ControlTarget.TargetOrientation.Quaternion(), DeltaTime);
+			Component->ControlTarget.TargetPosition = TargetPosition;
+			Component->ControlTarget.TargetOrientation = TargetOrientation;
+		}
+	}
+}
+
 //======================================================================================================================
 FRigUnit_HierarchySetControlMultiplier_Execute()
 {
diff --git a/Engine/Plugins/Experimental/PhysicsControl/Source/ControlRigPhysics/Private/RigPhysicsSimulation.cpp b/Engine/Plugins/Experimental/PhysicsControl/Source/ControlRigPhysics/Private/RigPhysicsSimulation.cpp
--- a/Engine/Plugins/Experimental/PhysicsControl/Source/ControlRigPhysics/Private/RigPhysicsSimulation.cpp
+++ b/Engine/Plugins/Experimental/PhysicsControl/Source/ControlRigPhysics/Private/RigPhysicsSimulation.cpp
@@ -1319,7 +1319,11 @@ void FRigPhysicsSimulation::UpdatePrePhysics(
 			if (const FRigPhysicsControlComponent* ControlComponent = Cast<FRigPhysicsControlComponent>(
 				Hierarchy->FindComponent(ComponentKey)))
 			{
-				UpdateControlPrePhysics(ControlRecord, ControlComponent, SolverComponent, Hierarchy, DeltaTime);
+				if (ControlComponent->ControlData.bEnabled)
+				{
+					UpdateControlPrePhysics(ControlRecord, ControlComponent, SolverComponent, Hierarchy, DeltaTime);
+				}
+				SetPhysicsJointEnabled(ControlRecord.JointHandle, ControlComponent->ControlData.bEnabled);
 			}
 		}
 	}
diff --git a/Engine/Plugins/Experimental/PhysicsControl/Source/ControlRigPhysics/Public/RigPhysicsControlExecution.h b/Engine/Plugins/Experimental/PhysicsControl/Source/ControlRigPhysics/Public/RigPhysicsControlExecution.h
--- a/Engine/Plugins/Experimental/PhysicsControl/Source/ControlRigPhysics/Public/RigPhysicsControlExecution.h
+++ b/Engine/Plugins/Experimental/PhysicsControl/Source/ControlRigPhysics/Public/RigPhysicsControlExecution.h
@@ -58,6 +58,29 @@ struct FRigUnit_AddPhysicsControl : public FRigUnit_PhysicsBaseMutable
 	//FName Set;
 };
 
+// Sets whether a control is enabled
+USTRUCT(meta = (DisplayName = "Set Physics Control Enabled", Varying))
+struct CONTROLRIGPHYSICS_API FRigUnit_HierarchySetControlEnabled : public FRigUnit_PhysicsBaseMutable
+{
+	GENERATED_BODY()
+
+	FRigUnit_HierarchySetControlEnabled()
+	{
+		PhysicsControlComponentKey.ElementKey.Type = ERigElementType::Bone;
+		PhysicsControlComponentKey.Name = FRigPhysicsControlComponent::GetDefaultName();
+	}
+
+	RIGVM_METHOD()
+	virtual void Execute() override;
+
+	UPROPERTY(meta = (Input))
+	FRigComponentKey PhysicsControlComponentKey;
+
+	UPROPERTY(meta = (Input))
+	bool bEnabled = true;
+};
+
+
 // Sets the control data for a physics control
 USTRUCT(meta = (DisplayName = "Set Physics Control Data", Varying))
 struct FRigUnit_HierarchySetControlData : public FRigUnit_PhysicsBaseMutable
@@ -102,6 +125,38 @@ struct FRigUnit_HierarchySetControlTarget : public FRigUnit_PhysicsBaseMutable
 	FPhysicsControlTarget ControlTarget;
 };
 
+// Sets the target for a physics control and updates the target velocities based on the previews
+// targets (which will be overwritten)
+USTRUCT(meta = (DisplayName = "Update Physics Control Target", Varying))
+struct CONTROLRIGPHYSICS_API FRigUnit_HierarchyUpdateControlTarget : public FRigUnit_PhysicsBaseMutable
+{
+	GENERATED_BODY()
+
+	FRigUnit_HierarchyUpdateControlTarget()
+	{
+		PhysicsControlComponentKey.ElementKey.Type = ERigElementType::Bone;
+		PhysicsControlComponentKey.Name = FRigPhysicsControlComponent::GetDefaultName();
+	}
+
+	RIGVM_METHOD()
+	virtual void Execute() override;
+
+	UPROPERTY(meta = (Input))
+	FRigComponentKey PhysicsControlComponentKey;
+
+	/** The target position of the child body, relative to the parent body */
+	UPROPERTY(meta = (Input))
+	FVector TargetPosition = FVector::ZeroVector;
+
+	/** The target orientation of the child body, relative to the parent body */
+	UPROPERTY(meta = (Input))
+	FRotator TargetOrientation = FRotator::ZeroRotator;
+
+	// The delta time used to calculate the target velocity
+	UPROPERTY(meta = (Input))
+	float DeltaTime = 0.0f;
+};
+
 // Sets the multipliers for a physics control
 USTRUCT(meta = (DisplayName = "Set Physics Control Multiplier", Varying))
 struct FRigUnit_HierarchySetControlMultiplier : public FRigUnit_PhysicsBaseMutable
