Support setting simulation space settings, including specific functions for the external velocity.

diff --git a/Engine/Plugins/Experimental/PhysicsControl/Source/ControlRigPhysics/Private/RigPhysicsExecution.cpp b/Engine/Plugins/Experimental/PhysicsControl/Source/ControlRigPhysics/Private/RigPhysicsExecution.cpp
--- a/Engine/Plugins/Experimental/PhysicsControl/Source/ControlRigPhysics/Private/RigPhysicsExecution.cpp
+++ b/Engine/Plugins/Experimental/PhysicsControl/Source/ControlRigPhysics/Private/RigPhysicsExecution.cpp
@@ -129,6 +129,43 @@ FRigUnit_StepPhysicsSolver_Execute()
 	}
 }
 
+//======================================================================================================================
+FRigUnit_SetPhysicsSolverSimulationSpaceSettings_Execute()
+{
+	DECLARE_SCOPE_HIERARCHICAL_COUNTER_RIGUNIT();
+	if (!ExecuteContext.Hierarchy)
+	{
+		return;
+	}
+	if (UControlRig* ControlRig = Cast<UControlRig>(ExecuteContext.Hierarchy->GetOuter()))
+	{
+		if (FRigPhysicsSolverComponent* SolverComponent = Cast<FRigPhysicsSolverComponent>(
+			ExecuteContext.Hierarchy->FindComponent(PhysicsSolverComponentKey)))
+		{
+			SolverComponent->SimulationSpaceSettings = SimulationSpaceSettings;
+		}
+	}
+}
+
+//======================================================================================================================
+FRigUnit_SetPhysicsSolverExternalVelocity_Execute()
+{
+	DECLARE_SCOPE_HIERARCHICAL_COUNTER_RIGUNIT();
+	if (!ExecuteContext.Hierarchy)
+	{
+		return;
+	}
+	if (UControlRig* ControlRig = Cast<UControlRig>(ExecuteContext.Hierarchy->GetOuter()))
+	{
+		if (FRigPhysicsSolverComponent* SolverComponent = Cast<FRigPhysicsSolverComponent>(
+			ExecuteContext.Hierarchy->FindComponent(PhysicsSolverComponentKey)))
+		{
+			SolverComponent->SimulationSpaceSettings.ExternalLinearVelocity = ExternalLinearVelocity;
+			SolverComponent->SimulationSpaceSettings.ExternalAngularVelocity = ExternalAngularVelocity;
+		}
+	}
+}
+
 //======================================================================================================================
 FRigUnit_InstantiatePhysics_Execute()
 {
diff --git a/Engine/Plugins/Experimental/PhysicsControl/Source/ControlRigPhysics/Private/RigPhysicsSimulation.cpp b/Engine/Plugins/Experimental/PhysicsControl/Source/ControlRigPhysics/Private/RigPhysicsSimulation.cpp
--- a/Engine/Plugins/Experimental/PhysicsControl/Source/ControlRigPhysics/Private/RigPhysicsSimulation.cpp
+++ b/Engine/Plugins/Experimental/PhysicsControl/Source/ControlRigPhysics/Private/RigPhysicsSimulation.cpp
@@ -1393,6 +1393,10 @@ void FRigPhysicsSimulation::UpdateControlPrePhysics(
 					SolverSettings, ComponentSpaceAnimTargetTM);
 				TargetTM = TargetTM * SimSpaceAnimTargetTM;
 			}
+			else
+			{
+				TargetTM = ConvertComponentSpaceTransformToSimSpace(SolverSettings, TargetTM);
+			}
 
 			Constraint->SetLinearDrivePositionTarget(TargetTM.GetTranslation());
 			Constraint->SetAngularDrivePositionTarget(TargetTM.GetRotation());
diff --git a/Engine/Plugins/Experimental/PhysicsControl/Source/ControlRigPhysics/Private/RigPhysicsSimulation.h b/Engine/Plugins/Experimental/PhysicsControl/Source/ControlRigPhysics/Private/RigPhysicsSimulation.h
--- a/Engine/Plugins/Experimental/PhysicsControl/Source/ControlRigPhysics/Private/RigPhysicsSimulation.h
+++ b/Engine/Plugins/Experimental/PhysicsControl/Source/ControlRigPhysics/Private/RigPhysicsSimulation.h
@@ -171,8 +171,10 @@ public:
 	struct FSimulationSpaceData
 	{
 		FVector LinearVelocity;
+		// Angular velocity is in rad/s
 		FVector AngularVelocity;
 		FVector LinearAcceleration;
+		// Angular acceleration is in rad/s/s
 		FVector AngularAcceleration;
 		FVector Gravity;
 	};
diff --git a/Engine/Plugins/Experimental/PhysicsControl/Source/ControlRigPhysics/Private/RigPhysicsSimulationSpace.cpp b/Engine/Plugins/Experimental/PhysicsControl/Source/ControlRigPhysics/Private/RigPhysicsSimulationSpace.cpp
--- a/Engine/Plugins/Experimental/PhysicsControl/Source/ControlRigPhysics/Private/RigPhysicsSimulationSpace.cpp
+++ b/Engine/Plugins/Experimental/PhysicsControl/Source/ControlRigPhysics/Private/RigPhysicsSimulationSpace.cpp
@@ -179,7 +179,7 @@ FRigPhysicsSimulation::FSimulationSpaceData FRigPhysicsSimulation::UpdateSimulat
 		// Danny TODO This is probably redundant unless we support runtime switching of the space
 		InitSimulationSpace(SimulationSpaceState.ComponentTM, SimulationSpaceState.BoneRelComponentTM);
 		SimulationSpaceData.LinearVelocity = SimulationSpaceSettings.ExternalLinearVelocity;
-		SimulationSpaceData.AngularVelocity = SimulationSpaceSettings.ExternalAngularVelocity;
+		SimulationSpaceData.AngularVelocity = FMath::DegreesToRadians(SimulationSpaceSettings.ExternalAngularVelocity);
 		return SimulationSpaceData;
 	}
 
@@ -317,7 +317,7 @@ FRigPhysicsSimulation::FSimulationSpaceData FRigPhysicsSimulation::UpdateSimulat
 	}
 
 	SimulationSpaceData.LinearVelocity += SimulationSpaceSettings.ExternalLinearVelocity;
-	SimulationSpaceData.AngularVelocity += SimulationSpaceSettings.ExternalAngularVelocity;
+	SimulationSpaceData.AngularVelocity += FMath::DegreesToRadians(SimulationSpaceSettings.ExternalAngularVelocity);
 
 	// Transform world-space motion into simulation space Danny TODO note that this matches the code
 	// in RBAN, and is doing what the interface requires (i.e. movement of the space in the space of
diff --git a/Engine/Plugins/Experimental/PhysicsControl/Source/ControlRigPhysics/Public/RigPhysicsData.h b/Engine/Plugins/Experimental/PhysicsControl/Source/ControlRigPhysics/Public/RigPhysicsData.h
--- a/Engine/Plugins/Experimental/PhysicsControl/Source/ControlRigPhysics/Public/RigPhysicsData.h
+++ b/Engine/Plugins/Experimental/PhysicsControl/Source/ControlRigPhysics/Public/RigPhysicsData.h
@@ -390,7 +390,7 @@ struct FRigPhysicsSimulationSpaceSettings
 	// Additional angular velocity that is added to the component angular velocity. This can be used
 	// to make the simulation act as if the actor is rotating even when it is not. E.g., to apply
 	// physics to a character on a podium as the camera rotates around it, to emulate the podium
-	// itself rotating. Vector is in world space. Units are rad/s.
+	// itself rotating. Vector is in world space. Units are deg/s.
 	UPROPERTY(EditAnywhere, Category = Settings)
 	FVector ExternalAngularVelocity = FVector::ZeroVector;
 };
diff --git a/Engine/Plugins/Experimental/PhysicsControl/Source/ControlRigPhysics/Public/RigPhysicsExecution.h b/Engine/Plugins/Experimental/PhysicsControl/Source/ControlRigPhysics/Public/RigPhysicsExecution.h
--- a/Engine/Plugins/Experimental/PhysicsControl/Source/ControlRigPhysics/Public/RigPhysicsExecution.h
+++ b/Engine/Plugins/Experimental/PhysicsControl/Source/ControlRigPhysics/Public/RigPhysicsExecution.h
@@ -113,7 +113,7 @@ struct FRigUnit_StepPhysicsSolver : public FRigUnit_PhysicsBaseMutable
 	RIGVM_METHOD()
 	UE_API virtual void Execute() override;
 
-	// The solver to relate this new physics element to
+	// The physics solver
 	UPROPERTY(meta = (Input))
 	FRigComponentKey PhysicsSolverComponentKey;
 	
@@ -139,6 +139,63 @@ struct FRigUnit_StepPhysicsSolver : public FRigUnit_PhysicsBaseMutable
 	FRigPhysicsVisualizationSettings VisualizationSettings;
 };
 
+// Sets all the simulation space settings on the solver
+USTRUCT(meta = (DisplayName = "Set Physics Solver Simulation Space Settings"))
+struct FRigUnit_SetPhysicsSolverSimulationSpaceSettings : public FRigUnit_PhysicsBaseMutable
+{
+	GENERATED_BODY()
+
+	FRigUnit_SetPhysicsSolverSimulationSpaceSettings()
+	{
+		PhysicsSolverComponentKey.ElementKey.Type = ERigElementType::Bone;
+		PhysicsSolverComponentKey.Name = FRigPhysicsSolverComponent::GetDefaultName();
+	}
+
+	RIGVM_METHOD()
+	UE_API virtual void Execute() override;
+
+	// The physics solver
+	UPROPERTY(meta = (Input))
+	FRigComponentKey PhysicsSolverComponentKey;
+
+	// The new simulation space settings
+	UPROPERTY(meta = (Input))
+	FRigPhysicsSimulationSpaceSettings SimulationSpaceSettings;
+};
+
+// Sets the external velocity of the simulation - used for adding wind effects
+USTRUCT(meta = (DisplayName = "Set Physics Solver External Velocity"))
+struct FRigUnit_SetPhysicsSolverExternalVelocity : public FRigUnit_PhysicsBaseMutable
+{
+	GENERATED_BODY()
+
+	FRigUnit_SetPhysicsSolverExternalVelocity()
+	{
+		PhysicsSolverComponentKey.ElementKey.Type = ERigElementType::Bone;
+		PhysicsSolverComponentKey.Name = FRigPhysicsSolverComponent::GetDefaultName();
+	}
+
+	RIGVM_METHOD()
+	UE_API virtual void Execute() override;
+
+	// The physics solver
+	UPROPERTY(meta = (Input))
+	FRigComponentKey PhysicsSolverComponentKey;
+
+	// Additional velocity that is added to the component velocity so the simulation acts as if the
+	// actor is moving at speed, even when stationary. The vector is in world space. This could be
+	// used for wind effects etc. Typical values are similar to the velocity of the object or
+	// effect, and usually around or less than 1000 for characters/wind.
+	UPROPERTY(meta = (Input))
+	FVector ExternalLinearVelocity = FVector::ZeroVector;
+
+	// Additional angular velocity that is added to the component angular velocity. This can be used
+	// to make the simulation act as if the actor is rotating even when it is not. E.g., to apply
+	// physics to a character on a podium as the camera rotates around it, to emulate the podium
+	// itself rotating. Vector is in world space. Units are deg/s.
+	UPROPERTY(meta = (Input))
+	FVector ExternalAngularVelocity = FVector::ZeroVector;
+};
 
 /**
  * Forces tracking of the input animation (on all physics bodies) for the next N frames
