Control rig physics drive uses the specified target and has an option whether or not to use the animation target

diff --git a/Engine/Plugins/Experimental/PhysicsControl/Source/ControlRigPhysics/Private/RigPhysicsData.cpp b/Engine/Plugins/Experimental/PhysicsControl/Source/ControlRigPhysics/Private/RigPhysicsData.cpp
--- a/Engine/Plugins/Experimental/PhysicsControl/Source/ControlRigPhysics/Private/RigPhysicsData.cpp
+++ b/Engine/Plugins/Experimental/PhysicsControl/Source/ControlRigPhysics/Private/RigPhysicsData.cpp
@@ -362,6 +362,10 @@ FArchive& operator <<(FArchive& Ar, FRigPhysicsDriveData& Data)
 		}
 	}
 	Ar << Data.SkeletalAnimationVelocityMultiplier;
+	if (Ar.CustomVer(FPhysicsControlObjectVersion::GUID) >= FPhysicsControlObjectVersion::ControlRigDriveRelativeToAnimation)
+	{
+		Ar << Data.bUseSkeletalAnimation;
+	}
 	return Ar;
 }
 
diff --git a/Engine/Plugins/Experimental/PhysicsControl/Source/ControlRigPhysics/Private/RigPhysicsExecution.cpp b/Engine/Plugins/Experimental/PhysicsControl/Source/ControlRigPhysics/Private/RigPhysicsExecution.cpp
--- a/Engine/Plugins/Experimental/PhysicsControl/Source/ControlRigPhysics/Private/RigPhysicsExecution.cpp
+++ b/Engine/Plugins/Experimental/PhysicsControl/Source/ControlRigPhysics/Private/RigPhysicsExecution.cpp
@@ -304,15 +304,6 @@ static void SetUpJointData(
 	JointData.TwistConstraint = ConstraintProfileProperties.TwistLimit;
 }
 
-//======================================================================================================================
-static void SetUpDriveData(
-	FRigPhysicsDriveData&               DriveData, 
-	const FConstraintProfileProperties& ConstraintProfileProperties)
-{
-	DriveData.LinearDriveConstraint = ConstraintProfileProperties.LinearDrive;
-	DriveData.AngularDriveConstraint = ConstraintProfileProperties.AngularDrive;
-}
-
 //======================================================================================================================
 FRigUnit_HierarchyInstantiateFromPhysicsAsset_Execute()
 {
@@ -539,7 +530,10 @@ FRigUnit_HierarchyInstantiateFromPhysicsAsset_Execute()
 			FRigPhysicsDriveData DriveData;
 			if (bEnableDrives)
 			{
-				SetUpDriveData(DriveData, ConstraintProfileProperties);
+				DriveData.LinearDriveConstraint = ConstraintProfileProperties.LinearDrive;
+				DriveData.AngularDriveConstraint = ConstraintProfileProperties.AngularDrive;
+				// RBAN doesn't use skeletal animation, and by default we want to replicate that.
+				DriveData.bUseSkeletalAnimation = false;
 			}
 
 			// Add the component
diff --git a/Engine/Plugins/Experimental/PhysicsControl/Source/ControlRigPhysics/Private/RigPhysicsJointExecution.cpp b/Engine/Plugins/Experimental/PhysicsControl/Source/ControlRigPhysics/Private/RigPhysicsJointExecution.cpp
--- a/Engine/Plugins/Experimental/PhysicsControl/Source/ControlRigPhysics/Private/RigPhysicsJointExecution.cpp
+++ b/Engine/Plugins/Experimental/PhysicsControl/Source/ControlRigPhysics/Private/RigPhysicsJointExecution.cpp
@@ -82,6 +82,21 @@ FRigUnit_HierarchySetJointDriveData_Execute()
 	}
 }
 
+//======================================================================================================================
+FRigUnit_HierarchySetJointDriveUseSkeletalAnimation_Execute()
+{
+	DECLARE_SCOPE_HIERARCHICAL_COUNTER_RIGUNIT();
+	if (!ExecuteContext.Hierarchy)
+	{
+		return;
+	}
+	if (FRigPhysicsJointComponent* Component = Cast<FRigPhysicsJointComponent>(
+		ExecuteContext.Hierarchy->FindComponent(PhysicsJointComponentKey)))
+	{
+		Component->DriveData.bUseSkeletalAnimation = bUseSkeletalAnimation;
+	}
+}
+
 //======================================================================================================================
 FRigUnit_MakeArticulationJointData_Execute()
 {
diff --git a/Engine/Plugins/Experimental/PhysicsControl/Source/ControlRigPhysics/Private/RigPhysicsSimulation.cpp b/Engine/Plugins/Experimental/PhysicsControl/Source/ControlRigPhysics/Private/RigPhysicsSimulation.cpp
--- a/Engine/Plugins/Experimental/PhysicsControl/Source/ControlRigPhysics/Private/RigPhysicsSimulation.cpp
+++ b/Engine/Plugins/Experimental/PhysicsControl/Source/ControlRigPhysics/Private/RigPhysicsSimulation.cpp
@@ -1087,6 +1087,9 @@ void FRigPhysicsSimulation::UpdateJointPrePhysics(
 			Settings.LinearDriveDamping.SquaredLength() > 0.0f)
 		{
 			// Multiplier on the velocity calculated from the current and previous target
+			//
+			// Danny TODO surely this should always be taken from DriveData? Do we/can we
+			// distinguish between manual and animation velocities?
 			float TargetVelocityMultiplier = 1.0f;
 
 			FRigElementKey ChildSourceKey;
@@ -1101,37 +1104,45 @@ void FRigPhysicsSimulation::UpdateJointPrePhysics(
 				TargetVelocityMultiplier = DriveData.SkeletalAnimationVelocityMultiplier;
 			}
 
-			FRigElementKey ParentSourceKey;
-			if (const FRigPhysicsBodyComponent* ParentPhysicsComponent = Cast<FRigPhysicsBodyComponent>(
-				Hierarchy->FindComponent(Record.ParentBodyComponentKey)))
+			UE::PhysicsControl::FPosQuat DriveTargetTM;
+			if (DriveData.bUseSkeletalAnimation)
 			{
-				ParentSourceKey = ParentPhysicsComponent->BodySolverSettings.SourceBone;
-				if (!ParentSourceKey.IsValid())
+				FRigElementKey ParentSourceKey;
+				if (const FRigPhysicsBodyComponent* ParentPhysicsComponent = Cast<FRigPhysicsBodyComponent>(
+					Hierarchy->FindComponent(Record.ParentBodyComponentKey)))
 				{
-					ParentSourceKey = Record.ParentBodyComponentKey.ElementKey;
+					ParentSourceKey = ParentPhysicsComponent->BodySolverSettings.SourceBone;
+					if (!ParentSourceKey.IsValid())
+					{
+						ParentSourceKey = Record.ParentBodyComponentKey.ElementKey;
+					}
 				}
-			}
 
-			if (!ChildSourceKey.IsValid() || !ParentSourceKey.IsValid())
-			{
-				return;
-			}
+				if (!ChildSourceKey.IsValid() || !ParentSourceKey.IsValid())
+				{
+					return;
+				}
 
-			// Danny TODO now all transforms are being cached in the record, get them from there
-			// rather than the hierarchy
+				// Danny TODO now all transforms are being cached in the record, get them from there
+				// rather than the hierarchy
 
-			// Note that the drive operates between a parent and child part, so we
-			// don't need to worry about global/component (etc) space.
-			const FTransform ChildTM = Hierarchy->GetGlobalTransform(ChildSourceKey);
-			const FTransform ParentTM = Hierarchy->GetGlobalTransform(ParentSourceKey);
+				// Note that the drive operates between a parent and child part, so we
+				// don't need to worry about global/component (etc) space.
+				const FTransform ChildTM = Hierarchy->GetGlobalTransform(ChildSourceKey);
+				const FTransform ParentTM = Hierarchy->GetGlobalTransform(ParentSourceKey);
 
-			const UE::PhysicsControl::FPosQuat ComponentSpaceParentFrameTM = 
-				Settings.ConnectorTransforms[1] * ParentTM;
-			const UE::PhysicsControl::FPosQuat ComponentSpaceChildFrameTM = 
-				Settings.ConnectorTransforms[0] * ChildTM;
+				const UE::PhysicsControl::FPosQuat ComponentSpaceParentFrameTM = 
+					Settings.ConnectorTransforms[1] * ParentTM;
+				const UE::PhysicsControl::FPosQuat ComponentSpaceChildFrameTM = 
+					Settings.ConnectorTransforms[0] * ChildTM;
+
+				DriveTargetTM = ComponentSpaceParentFrameTM.Inverse() * ComponentSpaceChildFrameTM;
+			}
 
-			const UE::PhysicsControl::FPosQuat DriveTargetTM =
-				ComponentSpaceParentFrameTM.Inverse() * ComponentSpaceChildFrameTM;
+			// Apply the offset in the frame of the (potentially animation) target
+			DriveTargetTM = DriveTargetTM * UE::PhysicsControl::FPosQuat(
+				DriveData.LinearDriveConstraint.PositionTarget, 
+				DriveData.AngularDriveConstraint.OrientationTarget.Quaternion());
 
 			Constraint->SetLinearDrivePositionTarget(DriveTargetTM.GetTranslation());
 			Constraint->SetAngularDrivePositionTarget(DriveTargetTM.GetRotation());
diff --git a/Engine/Plugins/Experimental/PhysicsControl/Source/ControlRigPhysics/Public/RigPhysicsData.h b/Engine/Plugins/Experimental/PhysicsControl/Source/ControlRigPhysics/Public/RigPhysicsData.h
--- a/Engine/Plugins/Experimental/PhysicsControl/Source/ControlRigPhysics/Public/RigPhysicsData.h
+++ b/Engine/Plugins/Experimental/PhysicsControl/Source/ControlRigPhysics/Public/RigPhysicsData.h
@@ -592,6 +592,11 @@ struct FRigPhysicsDriveData
 	UPROPERTY(EditAnywhere, Category = AngularDrive)
 	FAngularDriveConstraint AngularDriveConstraint;
 
+	// If true, then targets in the linear and angular drives will be applied on top of the pose
+	// from animation.
+	UPROPERTY(EditAnywhere, Category = Drive)
+	bool bUseSkeletalAnimation = true;
+
 	// The amount of skeletal animation velocity to use in the targets
 	UPROPERTY(EditAnywhere, Category = Drive)
 	float SkeletalAnimationVelocityMultiplier = 1.0f;
diff --git a/Engine/Plugins/Experimental/PhysicsControl/Source/ControlRigPhysics/Public/RigPhysicsJointExecution.h b/Engine/Plugins/Experimental/PhysicsControl/Source/ControlRigPhysics/Public/RigPhysicsJointExecution.h
--- a/Engine/Plugins/Experimental/PhysicsControl/Source/ControlRigPhysics/Public/RigPhysicsJointExecution.h
+++ b/Engine/Plugins/Experimental/PhysicsControl/Source/ControlRigPhysics/Public/RigPhysicsJointExecution.h
@@ -102,6 +102,30 @@ struct FRigUnit_HierarchySetJointDriveData : public FRigUnit_PhysicsBaseMutable
 	FRigPhysicsDriveData DriveData;
 };
 
+
+// Sets the joint drive for a physics component body
+USTRUCT(meta = (DisplayName = "Set Physics Joint Drive Use Skeletal Animation", Varying))
+struct FRigUnit_HierarchySetJointDriveUseSkeletalAnimation : public FRigUnit_PhysicsBaseMutable
+{
+	GENERATED_BODY()
+
+	FRigUnit_HierarchySetJointDriveUseSkeletalAnimation()
+	{
+		PhysicsJointComponentKey.ElementKey.Type = ERigElementType::Bone;
+		PhysicsJointComponentKey.Name = FRigPhysicsJointComponent::GetDefaultName();
+	}
+
+	RIGVM_METHOD()
+	UE_API virtual void Execute() override;
+
+	UPROPERTY(meta = (Input))
+	FRigComponentKey PhysicsJointComponentKey;
+
+	UPROPERTY(meta = (Input))
+	bool bUseSkeletalAnimation = true;
+};
+
+
 // Helper to simplify creation of joint data
 USTRUCT(meta = (DisplayName = "Make Joint Data", Keywords = "Construction,Create,New,Body,Skeleton", Varying))
 struct FRigUnit_MakeArticulationJointData : public FRigUnit_PhysicsBase
diff --git a/Engine/Plugins/Experimental/PhysicsControl/Source/PhysicsControl/Public/PhysicsControlObjectVersion.h b/Engine/Plugins/Experimental/PhysicsControl/Source/PhysicsControl/Public/PhysicsControlObjectVersion.h
--- a/Engine/Plugins/Experimental/PhysicsControl/Source/PhysicsControl/Public/PhysicsControlObjectVersion.h
+++ b/Engine/Plugins/Experimental/PhysicsControl/Source/PhysicsControl/Public/PhysicsControlObjectVersion.h
@@ -52,6 +52,8 @@ struct FPhysicsControlObjectVersion
 
 		ControlRigRemoveResetCooldownFrames,
 
+		ControlRigDriveRelativeToAnimation,
+
 		// -----<new versions can be added above this line>-------------------------------------------------
 		VersionPlusOne,
 		LatestVersion = VersionPlusOne - 1
