Improve reins - makes the control rig adaptively adjust the grab point in order to maintain a desired amount of slack

diff --git a/Engine/Plugins/Experimental/PhysicsControl/Source/ControlRigPhysics/Private/RigPhysicsControlExecution.cpp b/Engine/Plugins/Experimental/PhysicsControl/Source/ControlRigPhysics/Private/RigPhysicsControlExecution.cpp
--- a/Engine/Plugins/Experimental/PhysicsControl/Source/ControlRigPhysics/Private/RigPhysicsControlExecution.cpp
+++ b/Engine/Plugins/Experimental/PhysicsControl/Source/ControlRigPhysics/Private/RigPhysicsControlExecution.cpp
@@ -61,6 +61,22 @@ FRigUnit_HierarchySetControlEnabled_Execute()
 	}
 }
 
+//======================================================================================================================
+FRigUnit_HierarchySetControlCustomControlPoint_Execute()
+{
+	DECLARE_SCOPE_HIERARCHICAL_COUNTER_RIGUNIT();
+	if (ExecuteContext.Hierarchy)
+	{
+		if (FRigPhysicsControlComponent* Component = Cast<FRigPhysicsControlComponent>(
+			ExecuteContext.Hierarchy->FindComponent(PhysicsControlComponentKey)))
+		{
+			Component->ControlData.CustomControlPoint = CustomControlPoint;
+			Component->ControlData.bUseCustomControlPoint = bUseCustomControlPoint;
+		}
+	}
+}
+
+
 //======================================================================================================================
 FRigUnit_HierarchySetControlData_Execute()
 {
diff --git a/Engine/Plugins/Experimental/PhysicsControl/Source/ControlRigPhysics/Public/RigPhysicsControlExecution.h b/Engine/Plugins/Experimental/PhysicsControl/Source/ControlRigPhysics/Public/RigPhysicsControlExecution.h
--- a/Engine/Plugins/Experimental/PhysicsControl/Source/ControlRigPhysics/Public/RigPhysicsControlExecution.h
+++ b/Engine/Plugins/Experimental/PhysicsControl/Source/ControlRigPhysics/Public/RigPhysicsControlExecution.h
@@ -80,6 +80,33 @@ struct CONTROLRIGPHYSICS_API FRigUnit_HierarchySetControlEnabled : public FRigUn
 	bool bEnabled = true;
 };
 
+// Sets the custom control point on a control
+USTRUCT(meta = (DisplayName = "Set Physics Control Custom Control Point", Varying))
+struct CONTROLRIGPHYSICS_API FRigUnit_HierarchySetControlCustomControlPoint : public FRigUnit_PhysicsBaseMutable
+{
+	GENERATED_BODY()
+
+	FRigUnit_HierarchySetControlCustomControlPoint()
+	{
+		PhysicsControlComponentKey.ElementKey.Type = ERigElementType::Bone;
+		PhysicsControlComponentKey.Name = FRigPhysicsControlComponent::GetDefaultName();
+	}
+
+	RIGVM_METHOD()
+	virtual void Execute() override;
+
+	UPROPERTY(meta = (Input))
+	FRigComponentKey PhysicsControlComponentKey;
+
+	// The position of the control point relative to the child mesh, when using a custom control point.
+	UPROPERTY(meta = (Input))
+	FVector CustomControlPoint = FVector::ZeroVector;
+
+	// Whether or not to use the custom control point position
+	UPROPERTY(meta = (Input))
+	bool bUseCustomControlPoint = true;
+};
+
 
 // Sets the control data for a physics control
 USTRUCT(meta = (DisplayName = "Set Physics Control Data", Varying))
diff --git a/Engine/Plugins/Experimental/PhysicsControl/Source/PhysicsControl/Public/PhysicsControlData.h b/Engine/Plugins/Experimental/PhysicsControl/Source/PhysicsControl/Public/PhysicsControlData.h
--- a/Engine/Plugins/Experimental/PhysicsControl/Source/PhysicsControl/Public/PhysicsControlData.h
+++ b/Engine/Plugins/Experimental/PhysicsControl/Source/PhysicsControl/Public/PhysicsControlData.h
@@ -551,7 +551,8 @@ struct FPhysicsControlData
 	GENERATED_BODY()
 
 	FPhysicsControlData()
-		: LinearStrength(0.0f)
+		: bEnabled(true)
+		, LinearStrength(0.0f)
 		, LinearDampingRatio(1.0f)
 		, LinearExtraDamping(0.0f)
 		, MaxForce(0.0f)
@@ -563,7 +564,6 @@ struct FPhysicsControlData
 		, AngularTargetVelocityMultiplier(1.0f)
 		, SkeletalAnimationVelocityMultiplier(1.0f)
 		, CustomControlPoint(FVector::ZeroVector)
-		, bEnabled(true)
 		, bUseCustomControlPoint(false)
 		, bUseSkeletalAnimation(true)
 		, bDisableCollision(false)
@@ -598,6 +598,10 @@ struct FPhysicsControlData
 	/** Applies the values that have been flagged as enabled from the sparse data */
 	UE_API void UpdateFromSparseData(const FPhysicsControlSparseData& SparseData);
 
+	/** Whether this control should be enabled */
+	UPROPERTY(EditAnywhere, BlueprintReadWrite, Category = PhysicsControl)
+	bool bEnabled;
+
 	/** The strength used to drive linear motion */
 	UPROPERTY(EditAnywhere, BlueprintReadWrite, Category = PhysicsControl, meta = (ClampMin = "0.0"))
 	float LinearStrength;
@@ -673,10 +677,6 @@ struct FPhysicsControlData
 	UPROPERTY(EditAnywhere, BlueprintReadWrite, Category = PhysicsControl)
 	FVector CustomControlPoint;
 
-	/** Whether this control should be enabled */
-	UPROPERTY(EditAnywhere, BlueprintReadWrite, Category = PhysicsControl)
-	bool bEnabled;
-
 	/** Whether or not to use the custom control point position */
 	UPROPERTY(EditAnywhere, BlueprintReadWrite, Category = PhysicsControl)
 	bool bUseCustomControlPoint;
