Add support for turbulence in control rig physics - avoids potential/likely time accuracy problems

diff --git a/Engine/Plugins/Experimental/PhysicsControl/Source/ControlRigPhysics/Private/RigPhysicsData.cpp b/Engine/Plugins/Experimental/PhysicsControl/Source/ControlRigPhysics/Private/RigPhysicsData.cpp
--- a/Engine/Plugins/Experimental/PhysicsControl/Source/ControlRigPhysics/Private/RigPhysicsData.cpp
+++ b/Engine/Plugins/Experimental/PhysicsControl/Source/ControlRigPhysics/Private/RigPhysicsData.cpp
@@ -215,6 +215,10 @@ FArchive& operator <<(FArchive& Ar, FRigPhysicsSimulationSpaceSettings& Data)
 	Ar << Data.ExternalLinearDrag;
 	Ar << Data.ExternalLinearVelocity;
 	Ar << Data.ExternalAngularVelocity;
+	if (Ar.CustomVer(FPhysicsControlObjectVersion::GUID) >= FPhysicsControlObjectVersion::ControlRigDetectExternalVelocityTurbulence)
+	{
+		Ar << Data.ExternalTurbulenceVelocity;
+	}
 	return Ar;
 }
 
diff --git a/Engine/Plugins/Experimental/PhysicsControl/Source/ControlRigPhysics/Private/RigPhysicsExecution.cpp b/Engine/Plugins/Experimental/PhysicsControl/Source/ControlRigPhysics/Private/RigPhysicsExecution.cpp
--- a/Engine/Plugins/Experimental/PhysicsControl/Source/ControlRigPhysics/Private/RigPhysicsExecution.cpp
+++ b/Engine/Plugins/Experimental/PhysicsControl/Source/ControlRigPhysics/Private/RigPhysicsExecution.cpp
@@ -108,7 +108,7 @@ FRigUnit_StepPhysicsSolver_Execute()
 
 				if (ControlRig->RegisterPhysicsSimulation(NewPhysicsSimulation, PhysicsSolverComponentKey))
 				{
-					NewPhysicsSimulation->Instantiate(SolverComponent);
+					NewPhysicsSimulation->Instantiate(ExecuteContext, SolverComponent);
 				}
 				PhysicsSimulation = NewPhysicsSimulation.Get();
 			}
@@ -162,6 +162,7 @@ FRigUnit_SetPhysicsSolverExternalVelocity_Execute()
 		{
 			SolverComponent->SimulationSpaceSettings.ExternalLinearVelocity = ExternalLinearVelocity;
 			SolverComponent->SimulationSpaceSettings.ExternalAngularVelocity = ExternalAngularVelocity;
+			SolverComponent->SimulationSpaceSettings.ExternalTurbulenceVelocity = ExternalTurbulenceVelocity;
 		}
 	}
 }
@@ -183,7 +184,7 @@ FRigUnit_InstantiatePhysics_Execute()
 
 			if (ControlRig->RegisterPhysicsSimulation(PhysicsSimulation, PhysicsSolverComponentKey))
 			{
-				PhysicsSimulation->Instantiate(SolverComponent);
+				PhysicsSimulation->Instantiate(ExecuteContext, SolverComponent);
 			}
 		}
 	}
diff --git a/Engine/Plugins/Experimental/PhysicsControl/Source/ControlRigPhysics/Private/RigPhysicsSimulation.cpp b/Engine/Plugins/Experimental/PhysicsControl/Source/ControlRigPhysics/Private/RigPhysicsSimulation.cpp
--- a/Engine/Plugins/Experimental/PhysicsControl/Source/ControlRigPhysics/Private/RigPhysicsSimulation.cpp
+++ b/Engine/Plugins/Experimental/PhysicsControl/Source/ControlRigPhysics/Private/RigPhysicsSimulation.cpp
@@ -625,7 +625,8 @@ ImmediatePhysics::FActorHandle* FRigPhysicsSimulation::CreateBody(
 }
 
 //======================================================================================================================
-void FRigPhysicsSimulation::Instantiate(const FRigPhysicsSolverComponent* SolverComponent)
+void FRigPhysicsSimulation::Instantiate(
+	const FRigVMExecuteContext& ExecuteContext, const FRigPhysicsSolverComponent* SolverComponent)
 {
 	DECLARE_SCOPE_HIERARCHICAL_COUNTER_FUNC()
 	QUICK_SCOPE_CYCLE_COUNTER(STAT_RigPhysics_Instantiate);
@@ -638,7 +639,7 @@ void FRigPhysicsSimulation::Instantiate(const FRigPhysicsSolverComponent* Solver
 	// We need the simulation space in order to instantiate properly. This is not ideal, as we may
 	// end up updating the simulation space data twice (thus inserting data into the history). This
 	// shouldn't really matter as it will only be on the first step.
-	UpdateSimulationSpaceStateAndCalculateData(SolverComponent, 0.0f);
+	UpdateSimulationSpaceStateAndCalculateData(ExecuteContext, SolverComponent, 0.0f);
 
 	InitialiseSimulation(SolverComponent);
 
@@ -1609,12 +1610,13 @@ void FRigPhysicsSimulation::StepSimulation(
 	}
 
 	// We need to know about the simulation space etc before we can instantiate anything into the right place
-	SimulationSpaceData = UpdateSimulationSpaceStateAndCalculateData(SolverComponent, PhysicsSimulationSpaceDeltaTime);
+	SimulationSpaceData = UpdateSimulationSpaceStateAndCalculateData(
+		ExecuteContext, SolverComponent, PhysicsSimulationSpaceDeltaTime);
 
 	// We instantiate when we do the first simulation - this makes sure any changes applied by
 	// the user have been made. It also means there is no overhead if physics is never stepped.
 	// However, there may be a hitch due to the creation, so it may also happen during construction.
-	Instantiate(SolverComponent); // There is an early out if it's already been done
+	Instantiate(ExecuteContext, SolverComponent); // There is an early out if it's already been done
 
 	if (!Simulation)
 	{
diff --git a/Engine/Plugins/Experimental/PhysicsControl/Source/ControlRigPhysics/Private/RigPhysicsSimulation.h b/Engine/Plugins/Experimental/PhysicsControl/Source/ControlRigPhysics/Private/RigPhysicsSimulation.h
--- a/Engine/Plugins/Experimental/PhysicsControl/Source/ControlRigPhysics/Private/RigPhysicsSimulation.h
+++ b/Engine/Plugins/Experimental/PhysicsControl/Source/ControlRigPhysics/Private/RigPhysicsSimulation.h
@@ -146,7 +146,7 @@ public:
 		TObjectPtr<UControlRig> InOwningControlRig = nullptr);
 
 	// This will initialise/create the simulation and then create everything we need in it. 
-	void Instantiate(const FRigPhysicsSolverComponent* SolverComponent);
+	void Instantiate(const FRigVMExecuteContext& ExecuteContext, const FRigPhysicsSolverComponent* SolverComponent);
 
 	// Integrates the simulation forwards.
 	// If DeltaTimeOverride is +ve, then that value is used.
@@ -289,7 +289,9 @@ private:
 	void CheckForResetsPrePhysics(FRigPhysicsSolverComponent* SolverComponent, const float DeltaTime);
 
 	FSimulationSpaceData UpdateSimulationSpaceStateAndCalculateData(
-		const FRigPhysicsSolverComponent* SolverComponent, const float Dt);
+		const FRigVMExecuteContext&       ExecuteContext,
+		const FRigPhysicsSolverComponent* SolverComponent, 
+		const float                       Dt);
 
 	// Returns the simulation space transform, in world space
 	FTransform GetSimulationSpaceTransform(const FRigPhysicsSolverSettings& SolverSettings) const;
diff --git a/Engine/Plugins/Experimental/PhysicsControl/Source/ControlRigPhysics/Private/RigPhysicsSimulationSpace.cpp b/Engine/Plugins/Experimental/PhysicsControl/Source/ControlRigPhysics/Private/RigPhysicsSimulationSpace.cpp
--- a/Engine/Plugins/Experimental/PhysicsControl/Source/ControlRigPhysics/Private/RigPhysicsSimulationSpace.cpp
+++ b/Engine/Plugins/Experimental/PhysicsControl/Source/ControlRigPhysics/Private/RigPhysicsSimulationSpace.cpp
@@ -5,10 +5,13 @@
 
 #include "Rigs/RigHierarchy.h"
 #include "Rigs/RigHierarchyController.h"
+#include "RigVMCore/RigVMExecuteContext.h"
 #include "ControlRig.h"
 
 #include "PhysicsControlHelpers.h"
 
+#include "Engine/Engine.h"
+
 #include "Physics/ImmediatePhysics/ImmediatePhysicsActorHandle.h"
 #include "Physics/ImmediatePhysics/ImmediatePhysicsJointHandle.h"
 #include "Physics/ImmediatePhysics/ImmediatePhysicsSimulation.h"
@@ -25,6 +28,10 @@
 #include "PhysicsEngine/ConstraintInstance.h"
 #include "PhysicsEngine/BodySetup.h"
 
+TAutoConsoleVariable<int> CVarControlRigPhysicsShowSimulationSpaceInfo(
+	TEXT("ControlRig.Physics.ShowSimulationSpaceInfo"), 0,
+	TEXT("Shows information associated with the simulation space used in control rig physics."));
+
 //======================================================================================================================
 FORCEINLINE_DEBUGGABLE FTransform GetSpaceTransform(
 	ERigPhysicsSimulationSpace Space, const FTransform& ComponentTM, const FTransform& BoneTM)
@@ -132,10 +139,69 @@ void FRigPhysicsSimulation::InitSimulationSpace(
 	SimulationSpaceState.BoneRelComponentTM = BoneRelComponentTM;
 }
 
+
+//======================================================================================================================
+// Danny TODO We need support for double precision Perlin inputs - for now duplicate.
+namespace FPerlinHelpers
+{
+	// random permutation of 256 numbers, repeated 2x
+	static const int32 Permutation[512] = {
+		63, 9, 212, 205, 31, 128, 72, 59, 137, 203, 195, 170, 181, 115, 165, 40, 116, 139, 175, 225, 132, 99, 222, 2, 41, 15, 197, 93, 169, 90, 228, 43, 221, 38, 206, 204, 73, 17, 97, 10, 96, 47, 32, 138, 136, 30, 219,
+		78, 224, 13, 193, 88, 134, 211, 7, 112, 176, 19, 106, 83, 75, 217, 85, 0, 98, 140, 229, 80, 118, 151, 117, 251, 103, 242, 81, 238, 172, 82, 110, 4, 227, 77, 243, 46, 12, 189, 34, 188, 200, 161, 68, 76, 171, 194,
+		57, 48, 247, 233, 51, 105, 5, 23, 42, 50, 216, 45, 239, 148, 249, 84, 70, 125, 108, 241, 62, 66, 64, 240, 173, 185, 250, 49, 6, 37, 26, 21, 244, 60, 223, 255, 16, 145, 27, 109, 58, 102, 142, 253, 120, 149, 160,
+		124, 156, 79, 186, 135, 127, 14, 121, 22, 65, 54, 153, 91, 213, 174, 24, 252, 131, 192, 190, 202, 208, 35, 94, 231, 56, 95, 183, 163, 111, 147, 25, 67, 36, 92, 236, 71, 166, 1, 187, 100, 130, 143, 237, 178, 158,
+		104, 184, 159, 177, 52, 214, 230, 119, 87, 114, 201, 179, 198, 3, 248, 182, 39, 11, 152, 196, 113, 20, 232, 69, 141, 207, 234, 53, 86, 180, 226, 74, 150, 218, 29, 133, 8, 44, 123, 28, 146, 89, 101, 154, 220, 126,
+		155, 122, 210, 168, 254, 162, 129, 33, 18, 209, 61, 191, 199, 157, 245, 55, 164, 167, 215, 246, 144, 107, 235,
+
+		63, 9, 212, 205, 31, 128, 72, 59, 137, 203, 195, 170, 181, 115, 165, 40, 116, 139, 175, 225, 132, 99, 222, 2, 41, 15, 197, 93, 169, 90, 228, 43, 221, 38, 206, 204, 73, 17, 97, 10, 96, 47, 32, 138, 136, 30, 219,
+		78, 224, 13, 193, 88, 134, 211, 7, 112, 176, 19, 106, 83, 75, 217, 85, 0, 98, 140, 229, 80, 118, 151, 117, 251, 103, 242, 81, 238, 172, 82, 110, 4, 227, 77, 243, 46, 12, 189, 34, 188, 200, 161, 68, 76, 171, 194,
+		57, 48, 247, 233, 51, 105, 5, 23, 42, 50, 216, 45, 239, 148, 249, 84, 70, 125, 108, 241, 62, 66, 64, 240, 173, 185, 250, 49, 6, 37, 26, 21, 244, 60, 223, 255, 16, 145, 27, 109, 58, 102, 142, 253, 120, 149, 160,
+		124, 156, 79, 186, 135, 127, 14, 121, 22, 65, 54, 153, 91, 213, 174, 24, 252, 131, 192, 190, 202, 208, 35, 94, 231, 56, 95, 183, 163, 111, 147, 25, 67, 36, 92, 236, 71, 166, 1, 187, 100, 130, 143, 237, 178, 158,
+		104, 184, 159, 177, 52, 214, 230, 119, 87, 114, 201, 179, 198, 3, 248, 182, 39, 11, 152, 196, 113, 20, 232, 69, 141, 207, 234, 53, 86, 180, 226, 74, 150, 218, 29, 133, 8, 44, 123, 28, 146, 89, 101, 154, 220, 126,
+		155, 122, 210, 168, 254, 162, 129, 33, 18, 209, 61, 191, 199, 157, 245, 55, 164, 167, 215, 246, 144, 107, 235
+	};
+
+	// Gradient functions for 1D, 2D and 3D Perlin noise
+
+	FORCEINLINE float Grad1(int32 Hash, double X)
+	{
+		// Slicing Perlin's 3D improved noise would give us only scales of -1, 0 and 1; this looks pretty bad so let's use a different sampling
+		static const double Grad1Scales[16] = { -8 / 8, -7 / 8., -6 / 8., -5 / 8., -4 / 8., -3 / 8., -2 / 8., -1 / 8., 1 / 8., 2 / 8., 3 / 8., 4 / 8., 5 / 8., 6 / 8., 7 / 8., 8 / 8 };
+		return Grad1Scales[Hash & 15] * X;
+	}
+
+
+	// Curve w/ second derivative vanishing at 0 and 1, from Perlin's improved noise paper
+	FORCEINLINE double SmoothCurve(double X)
+	{
+		return X * X * X * (X * (X * 6.0 - 15.0) + 10.0);
+	}
+
+	FORCEINLINE float PerlinNoise1D(double X)
+	{
+		const double Xfl = FMath::FloorToFloat(X);
+		const int64 Xi = (int64)(Xfl) & 255;
+		X -= Xfl;
+		const double Xm1 = X - 1.0;
+
+		const int32 A = Permutation[Xi];
+		const int32 B = Permutation[Xi + 1];
+
+		const double U = SmoothCurve(X);
+
+		// 2.0 factor to ensure (-1, 1) range
+		return 2.0f * FMath::Lerp(Grad1(A, X), Grad1(B, Xm1), U);
+	}
+}
+
+
+
 //======================================================================================================================
 // Note - don't use the space conversion functions here as the state won't have been set yet.
 FRigPhysicsSimulation::FSimulationSpaceData FRigPhysicsSimulation::UpdateSimulationSpaceStateAndCalculateData(
-	const FRigPhysicsSolverComponent* SolverComponent, const float Dt)
+	const FRigVMExecuteContext&       ExecuteContext,
+	const FRigPhysicsSolverComponent* SolverComponent, 
+	const float                       Dt)
 {
 	const FRigPhysicsSolverSettings& SolverSettings = SolverComponent->SolverSettings;
 	const FRigPhysicsSimulationSpaceSettings& SimulationSpaceSettings = SolverComponent->SimulationSpaceSettings;
@@ -319,6 +385,25 @@ FRigPhysicsSimulation::FSimulationSpaceData FRigPhysicsSimulation::UpdateSimulat
 	SimulationSpaceData.LinearVelocity += SimulationSpaceSettings.ExternalLinearVelocity;
 	SimulationSpaceData.AngularVelocity += FMath::DegreesToRadians(SimulationSpaceSettings.ExternalAngularVelocity);
 
+	if (!SimulationSpaceSettings.ExternalTurbulenceVelocity.IsNearlyZero())
+	{
+		double GlobalTime = ExecuteContext.GetAbsoluteTime();
+		FVector T(
+			FPerlinHelpers::PerlinNoise1D(GlobalTime),
+			FPerlinHelpers::PerlinNoise1D(GlobalTime + 10.0),
+			FPerlinHelpers::PerlinNoise1D(GlobalTime + 20.0));
+
+		FVector Turbulence = T * SimulationSpaceSettings.ExternalTurbulenceVelocity;
+
+		SimulationSpaceData.LinearVelocity += Turbulence;
+
+		if (GEngine && CVarControlRigPhysicsShowSimulationSpaceInfo.GetValueOnAnyThread())
+		{
+			GEngine->AddOnScreenDebugMessage(INDEX_NONE, 0.f, FColor::Yellow,
+				FString::Printf(TEXT("Sim %s Turbulence %s"), *OwningControlRig->GetName(), *Turbulence.ToString()));
+		}
+	}
+
 	// Transform world-space motion into simulation space Danny TODO note that this matches the code
 	// in RBAN, and is doing what the interface requires (i.e. movement of the space in the space of
 	// the space!). 
@@ -331,5 +416,14 @@ FRigPhysicsSimulation::FSimulationSpaceData FRigPhysicsSimulation::UpdateSimulat
 	SimulationSpaceData.AngularAcceleration =
 		SimulationSpaceState.SimulationSpaceTM.InverseTransformVector(SimulationSpaceData.AngularAcceleration);
 
+	if (GEngine && CVarControlRigPhysicsShowSimulationSpaceInfo.GetValueOnAnyThread())
+	{
+		GEngine->AddOnScreenDebugMessage(INDEX_NONE, 0.f, FColor::Yellow,
+			FString::Printf(TEXT("Sim %s dt %.2fms Space speed %.3f accel %.3f"),
+				*OwningControlRig->GetName(), Dt * 1000.0f,
+				SimulationSpaceData.LinearVelocity.Length(),
+				SimulationSpaceData.LinearAcceleration.Length()));
+	}
+
 	return SimulationSpaceData;
 }
diff --git a/Engine/Plugins/Experimental/PhysicsControl/Source/ControlRigPhysics/Public/RigPhysicsData.h b/Engine/Plugins/Experimental/PhysicsControl/Source/ControlRigPhysics/Public/RigPhysicsData.h
--- a/Engine/Plugins/Experimental/PhysicsControl/Source/ControlRigPhysics/Public/RigPhysicsData.h
+++ b/Engine/Plugins/Experimental/PhysicsControl/Source/ControlRigPhysics/Public/RigPhysicsData.h
@@ -393,6 +393,11 @@ struct FRigPhysicsSimulationSpaceSettings
 	// itself rotating. Vector is in world space. Units are deg/s.
 	UPROPERTY(EditAnywhere, Category = Settings)
 	FVector ExternalAngularVelocity = FVector::ZeroVector;
+
+	// This will treat the external velocity like a wind field and add turbulence to it. Units are
+	// the same as velocity, so this is the approximate magnitude of the turbulence.
+	UPROPERTY(EditAnywhere, Category = Settings)
+	FVector ExternalTurbulenceVelocity = FVector::ZeroVector;
 };
 
 /**
diff --git a/Engine/Plugins/Experimental/PhysicsControl/Source/ControlRigPhysics/Public/RigPhysicsExecution.h b/Engine/Plugins/Experimental/PhysicsControl/Source/ControlRigPhysics/Public/RigPhysicsExecution.h
--- a/Engine/Plugins/Experimental/PhysicsControl/Source/ControlRigPhysics/Public/RigPhysicsExecution.h
+++ b/Engine/Plugins/Experimental/PhysicsControl/Source/ControlRigPhysics/Public/RigPhysicsExecution.h
@@ -195,6 +195,11 @@ struct FRigUnit_SetPhysicsSolverExternalVelocity : public FRigUnit_PhysicsBaseMu
 	// itself rotating. Vector is in world space. Units are deg/s.
 	UPROPERTY(meta = (Input))
 	FVector ExternalAngularVelocity = FVector::ZeroVector;
+
+	// This will treat the external velocity like a wind field and add turbulence to it. Units are
+	// the same as velocity, so this is the approximate magnitude of the turbulence.
+	UPROPERTY(meta = (Input))
+	FVector ExternalTurbulenceVelocity = FVector::ZeroVector;
 };
 
 /**
diff --git a/Engine/Plugins/Experimental/PhysicsControl/Source/PhysicsControl/Public/PhysicsControlObjectVersion.h b/Engine/Plugins/Experimental/PhysicsControl/Source/PhysicsControl/Public/PhysicsControlObjectVersion.h
--- a/Engine/Plugins/Experimental/PhysicsControl/Source/PhysicsControl/Public/PhysicsControlObjectVersion.h
+++ b/Engine/Plugins/Experimental/PhysicsControl/Source/PhysicsControl/Public/PhysicsControlObjectVersion.h
@@ -54,6 +54,8 @@ struct FPhysicsControlObjectVersion
 
 		ControlRigDriveRelativeToAnimation,
 
+		ControlRigDetectExternalVelocityTurbulence,
+
 		// -----<new versions can be added above this line>-------------------------------------------------
 		VersionPlusOne,
 		LatestVersion = VersionPlusOne - 1
