Add the ability to import bones and collision from a physics asset into the control rig

diff --git a/Engine/Plugins/Experimental/PhysicsControl/Source/ControlRigPhysics/Private/RigPhysicsExecution.cpp b/Engine/Plugins/Experimental/PhysicsControl/Source/ControlRigPhysics/Private/RigPhysicsExecution.cpp
--- a/Engine/Plugins/Experimental/PhysicsControl/Source/ControlRigPhysics/Private/RigPhysicsExecution.cpp
+++ b/Engine/Plugins/Experimental/PhysicsControl/Source/ControlRigPhysics/Private/RigPhysicsExecution.cpp
@@ -412,7 +412,8 @@ FRigUnit_HierarchyInstantiateFromPhysicsAsset_Execute()
 		if (FRigPhysicsBodyComponent* Component = Cast<FRigPhysicsBodyComponent>(
 			ExecuteContext.Hierarchy->FindComponent(PhysicsBodyComponentKey)))
 		{
-			Component->BodySolverSettings = Solver;
+			Component->BodySolverSettings.PhysicsSolverComponentKey = PhysicsBodyComponentKey;
+			Component->BodySolverSettings.bUseAutomaticSolver = bUseAutomaticSolver;
 			Component->Dynamics = Dynamics;
 			Component->BodyData = BodyData;
 			Component->Collision = Collision;
@@ -557,6 +558,127 @@ FRigUnit_HierarchyInstantiateFromPhysicsAsset_Execute()
 	}
 }
 
+//======================================================================================================================
+FRigUnit_HierarchyImportCollisionFromPhysicsAsset_Execute()
+{
+	DECLARE_SCOPE_HIERARCHICAL_COUNTER_RIGUNIT();
+	if (!ExecuteContext.Hierarchy)
+	{
+		return;
+	}
+	if (ExecuteContext.GetEventName() != FRigUnit_PrepareForExecution::EventName)
+	{
+		UE_CONTROLRIG_RIGUNIT_REPORT_ERROR(TEXT("InstantiateFromPhysicsAsset can only be used during Setup"));
+		return;
+	}
+	if (!PhysicsAsset)
+	{
+		UE_CONTROLRIG_RIGUNIT_REPORT_ERROR(TEXT("InstantiateFromPhysicsAsset needs a valid physics asset"));
+		return;
+	}
+
+	URigHierarchyController* Controller = ExecuteContext.Hierarchy->GetController();
+	if (!Controller)
+	{
+		return;
+	}
+
+	FRigHierarchyControllerInstructionBracket InstructionBracket(Controller, ExecuteContext.GetInstructionIndex());
+
+	for (TObjectPtr<USkeletalBodySetup> SkeletalBodySetup : PhysicsAsset->SkeletalBodySetups)
+	{
+		if (!SkeletalBodySetup)
+		{
+			continue;
+		}
+
+		// Check we can use this based on the bone name
+		FName PhysicsBoneName = SkeletalBodySetup->BoneName;
+
+		if (!BonesToUse.IsEmpty() && BonesToUse.Find(PhysicsBoneName) == INDEX_NONE)
+		{
+			continue;
+		}
+
+		FName BoneName(NameSpace.ToString() + PhysicsBoneName.ToString());
+		FRigElementKey BoneKey = Controller->AddBone(BoneName, Owner, FTransform(), true, ERigBoneType::Imported);
+
+		if (BoneKey.IsValid())
+		{
+			// Collision
+			FRigPhysicsCollision Collision;
+			const FKAggregateGeom& AggGeom = SkeletalBodySetup->AggGeom;
+			for (const FKBoxElem& BoxElem : AggGeom.BoxElems)
+			{
+				FTransform TM(BoxElem.Rotation, BoxElem.Center);
+				FRigPhysicsCollisionBox Box(TM, FVector(BoxElem.X, BoxElem.Y, BoxElem.Z));
+				SetCommonProperties(Box, BoxElem);
+				Collision.Boxes.Add(Box);
+			}
+			for (const FKSphereElem& SphereElem : AggGeom.SphereElems)
+			{
+				FTransform TM(SphereElem.Center);
+				FRigPhysicsCollisionSphere Sphere(TM, SphereElem.Radius);
+				SetCommonProperties(Sphere, SphereElem);
+				Collision.Spheres.Add(Sphere);
+			}
+			for (const FKSphylElem& SphylElem : AggGeom.SphylElems)
+			{
+				FTransform TM(SphylElem.Rotation, SphylElem.Center);
+				FRigPhysicsCollisionCapsule Capsule(TM, SphylElem.Radius, SphylElem.Length);
+				SetCommonProperties(Capsule, SphylElem);
+				Collision.Capsules.Add(Capsule);
+			}
+
+			const UPhysicalMaterial* Material = SkeletalBodySetup->GetPhysMaterial();
+			Material = Material ? Material : GetDefault<UPhysicalMaterial>();
+			if (Material)
+			{
+				Collision.Material.Friction = FMath::Max(Material->Friction, Material->StaticFriction);
+				Collision.Material.Restitution = Material->Restitution;
+				Collision.Material.FrictionCombineMode = (ERigPhysicsCombineMode)
+					Material->FrictionCombineMode.GetValue();
+				Collision.Material.RestitutionCombineMode = (ERigPhysicsCombineMode)
+					Material->RestitutionCombineMode.GetValue();
+			}
+
+			// Dynamics
+			FRigPhysicsDynamics Dynamics;
+			Dynamics.MassOverride = SkeletalBodySetup->CalculateMass();
+			Dynamics.LinearDamping = SkeletalBodySetup->DefaultInstance.LinearDamping;
+			Dynamics.AngularDamping = SkeletalBodySetup->DefaultInstance.AngularDamping;
+
+			// Body Data
+			FPhysicsControlModifierData BodyData;
+			BodyData.MovementType = EPhysicsMovementType::Kinematic; // Always kinematic
+			BodyData.CollisionType =
+				SkeletalBodySetup->CollisionReponse == EBodyCollisionResponse::BodyCollision_Enabled ?
+				ECollisionEnabled::PhysicsOnly : ECollisionEnabled::NoCollision;
+
+			// Add the component
+			FRigComponentKey PhysicsBodyComponentKey = Controller->AddComponent(
+				FRigPhysicsBodyComponent::StaticStruct(), TEXT("PhysicsBody"), BoneKey);
+			if (FRigPhysicsBodyComponent* Component = Cast<FRigPhysicsBodyComponent>(
+				ExecuteContext.Hierarchy->FindComponent(PhysicsBodyComponentKey)))
+			{
+				Component->BodySolverSettings.PhysicsSolverComponentKey = PhysicsBodyComponentKey;
+				Component->BodySolverSettings.bUseAutomaticSolver = bUseAutomaticSolver;
+				Component->Dynamics = Dynamics;
+				Component->BodyData = BodyData;
+				Component->Collision = Collision;
+
+				PhysicsBodyComponentKeys.Add(PhysicsBodyComponentKey);
+				BoneKeys.Add(BoneKey);
+			}
+			else
+			{
+				Controller->RemoveElement(BoneKey);
+			}
+		}
+	}
+
+}
+
 //======================================================================================================================
 FRigUnit_GetPhysicsSolverSpaceData_Execute()
 {
diff --git a/Engine/Plugins/Experimental/PhysicsControl/Source/ControlRigPhysics/Private/RigPhysicsSimulationVisualization.cpp b/Engine/Plugins/Experimental/PhysicsControl/Source/ControlRigPhysics/Private/RigPhysicsSimulationVisualization.cpp
--- a/Engine/Plugins/Experimental/PhysicsControl/Source/ControlRigPhysics/Private/RigPhysicsSimulationVisualization.cpp
+++ b/Engine/Plugins/Experimental/PhysicsControl/Source/ControlRigPhysics/Private/RigPhysicsSimulationVisualization.cpp
@@ -27,12 +27,13 @@
 
 namespace RigPhysicsSolverDraw
 {
-	FLinearColor DynamicColor = FColor::Yellow;
-	FLinearColor KinematicColor = FColor::Blue;
+	FLinearColor DynamicCollisionColor = FColor::Yellow;
+	FLinearColor DynamicNoCollisionColor = FColor::Orange;
+	FLinearColor KinematicCollisionColor = FColor::Blue;
+	FLinearColor KinematicNoCollisionColor = FColor::Purple;
 
 	FLinearColor ActiveContactColor = FColor::Red;
 	FLinearColor InactiveContactColor = FColor::Silver;
-
 }
 
 TAutoConsoleVariable<int> CVarControlRigPhysicsShowActiveContactsOverride(
@@ -43,11 +44,27 @@ TAutoConsoleVariable<int> CVarControlRigPhysicsShowInactiveContactsOverride(
 	TEXT("ControlRig.Physics.ShowInactiveContactsOveride"), -1,
 	TEXT("Whether to draw inactive contacts (requires visualization to be enabled). -1 uses the visualization setting, 0 forces drawing to be disabled, 1 forces it to be enabled."));
 
+//======================================================================================================================
+static FLinearColor GetColor(bool bIsKinematic, bool bCollides)
+{
+	if (bIsKinematic)
+	{
+		return bCollides 
+			? RigPhysicsSolverDraw::KinematicCollisionColor : RigPhysicsSolverDraw::KinematicNoCollisionColor;
+	}
+	else
+	{
+		return bCollides 
+			? RigPhysicsSolverDraw::DynamicCollisionColor : RigPhysicsSolverDraw::DynamicNoCollisionColor;
+	}
+}
+
 //======================================================================================================================
 static void DrawShapes(
 	FRigVMDrawInterface*                    DI,
 	const FRigPhysicsVisualizationSettings& VisualizationSettings,
 	const bool                              bIsKinematic,
+	const bool                              bCollides,
 	const Chaos::FGeometryParticleHandle*   Particle,
 	const FTransform&                       ShapeTransform, 
 	const Chaos::FImplicitObject*           ImplicitObject,
@@ -79,7 +96,7 @@ static void DrawShapes(
 			ShapeTransform.GetRotation() * Transformed->GetTransform().GetRotation(),
 			ShapeTransform.TransformPosition(Transformed->GetTransform().GetLocation()));
 		DrawShapes(
-			DI, VisualizationSettings, bIsKinematic, Particle, 
+			DI, VisualizationSettings, bIsKinematic, bCollides, Particle, 
 			TransformedTransform, Transformed->GetTransformedObject(), Shape);
 		return;
 	}
@@ -87,7 +104,7 @@ static void DrawShapes(
 	{
 		const Chaos::FSphere* Sphere = ImplicitObject->GetObject<Chaos::FSphere>();
 		DI->DrawSphere(ShapeTransform, FTransform(FVec3(Sphere->GetCenterOfMass())), Sphere->GetRadiusf(),
-			bIsKinematic ? RigPhysicsSolverDraw::KinematicColor : RigPhysicsSolverDraw::DynamicColor,
+			GetColor(bIsKinematic, bCollides),
 			VisualizationSettings.LineThickness, VisualizationSettings.ShapeDetail);
 		break;
 	}
@@ -95,7 +112,7 @@ static void DrawShapes(
 	{
 		const TBox<FReal, 3>* Box = ImplicitObject->GetObject<TBox<FReal, 3>>();
 		DI->DrawBox(ShapeTransform, FTransform(FQuat::Identity, Box->GetCenter(), Box->Extents()),
-			bIsKinematic ? RigPhysicsSolverDraw::KinematicColor : RigPhysicsSolverDraw::DynamicColor,
+			GetColor(bIsKinematic, bCollides),
 			VisualizationSettings.LineThickness);
 		break;
 	}
@@ -105,7 +122,7 @@ static void DrawShapes(
 		const FQuat Q = FRotationMatrix::MakeFromZ(Capsule->GetAxis()).ToQuat();
 		DI->DrawCapsule(ShapeTransform, FTransform(Q, Capsule->GetCenterOfMass()),
 			Capsule->GetRadiusf(), Capsule->GetHeightf(),
-			bIsKinematic ? RigPhysicsSolverDraw::KinematicColor : RigPhysicsSolverDraw::DynamicColor,
+			GetColor(bIsKinematic, bCollides),
 			VisualizationSettings.LineThickness);
 		break;
 	}
@@ -113,7 +130,6 @@ static void DrawShapes(
 		// If we don't know what it is, don't draw it.
 		break;
 	}
-
 }
 
 //======================================================================================================================
@@ -124,13 +140,15 @@ static void DrawActor(
 	const ImmediatePhysics::FActorHandle&   ActorHandle)
 {
 	bool bIsKinematic = ActorHandle.GetIsKinematic();
+	bool bCollides = ActorHandle.GetHasCollision();
+
 	const Chaos::FGeometryParticleHandle* GeometryParticleHandle = ActorHandle.GetParticle();
 	const FTransform ParticleTransform = GeometryParticleHandle->GetTransformXR() * SpaceTransform;
 	for (const Chaos::FShapeInstancePtr& ShapeInstance : GeometryParticleHandle->ShapeInstances())
 	{
 		const Chaos::FImplicitObject* ImplicitObject = ShapeInstance->GetGeometry();
 		DrawShapes(
-			DI, VisualizationSettings, bIsKinematic, GeometryParticleHandle,
+			DI, VisualizationSettings, bIsKinematic, bCollides, GeometryParticleHandle,
 			ParticleTransform, ImplicitObject, ShapeInstance.Get());
 	}
 }
diff --git a/Engine/Plugins/Experimental/PhysicsControl/Source/ControlRigPhysics/Public/RigPhysicsBodyExecution.h b/Engine/Plugins/Experimental/PhysicsControl/Source/ControlRigPhysics/Public/RigPhysicsBodyExecution.h
--- a/Engine/Plugins/Experimental/PhysicsControl/Source/ControlRigPhysics/Public/RigPhysicsBodyExecution.h
+++ b/Engine/Plugins/Experimental/PhysicsControl/Source/ControlRigPhysics/Public/RigPhysicsBodyExecution.h
@@ -18,7 +18,7 @@
  * Adds a new physics body as a component on the owner element.
  * Note: This node only runs as part of the construction event.
  */
-USTRUCT(meta = (DisplayName = "Add Physics Body Component", Keywords = "Construction,Create,New,Body,Skeleton", Varying))
+USTRUCT(meta = (DisplayName = "Add Physics Body Component", Keywords = "Spawn,Construction,Create,New,Body,Skeleton", Varying))
 struct FRigUnit_AddPhysicsBody : public FRigUnit_PhysicsBaseMutable
 {
 	GENERATED_BODY()
diff --git a/Engine/Plugins/Experimental/PhysicsControl/Source/ControlRigPhysics/Public/RigPhysicsControlExecution.h b/Engine/Plugins/Experimental/PhysicsControl/Source/ControlRigPhysics/Public/RigPhysicsControlExecution.h
--- a/Engine/Plugins/Experimental/PhysicsControl/Source/ControlRigPhysics/Public/RigPhysicsControlExecution.h
+++ b/Engine/Plugins/Experimental/PhysicsControl/Source/ControlRigPhysics/Public/RigPhysicsControlExecution.h
@@ -16,7 +16,7 @@
  * Adds a new physics control as a component on the owner element.
  * Note: This node only runs as part of the construction event.
  */
-USTRUCT(meta = (Category = "RigPhysics", NodeColor = "1.0 0.6 0.3", DisplayName = "Add Physics Control Component", Keywords = "Construction,Create,New,Joint,Control,Physics", Varying))
+USTRUCT(meta = (Category = "RigPhysics", NodeColor = "1.0 0.6 0.3", DisplayName = "Add Physics Control Component", Keywords = "Spawn,Construction,Create,New,Joint,Control,Physics", Varying))
 struct FRigUnit_AddPhysicsControl : public FRigUnit_PhysicsBaseMutable
 {
 	GENERATED_BODY()
diff --git a/Engine/Plugins/Experimental/PhysicsControl/Source/ControlRigPhysics/Public/RigPhysicsExecution.h b/Engine/Plugins/Experimental/PhysicsControl/Source/ControlRigPhysics/Public/RigPhysicsExecution.h
--- a/Engine/Plugins/Experimental/PhysicsControl/Source/ControlRigPhysics/Public/RigPhysicsExecution.h
+++ b/Engine/Plugins/Experimental/PhysicsControl/Source/ControlRigPhysics/Public/RigPhysicsExecution.h
@@ -35,7 +35,7 @@ struct FRigUnit_PhysicsBase : public FRigUnit
  * Adds a new physics solver as a component on the owner element.
  * Note: This node only runs as part of the construction event.
  */
-USTRUCT(meta=(DisplayName="Add Physics Solver Component", Keywords="Construction,Create,New,Simulation", Varying))
+USTRUCT(meta=(DisplayName="Add Physics Solver Component", Keywords="Spawn,Construction,Create,New,Simulation", Varying))
 struct FRigUnit_AddPhysicsSolver : public FRigUnit_PhysicsBaseMutable
 {
 	GENERATED_BODY()
@@ -235,7 +235,7 @@ struct FRigUnit_TrackInputPose : public FRigUnit_PhysicsBaseMutable
 
 
 // Adds a set of physics components including the body, joint and controls
-USTRUCT(meta = (DisplayName = "Add Physics Components", Keywords = "Construction,Create,New", Varying))
+USTRUCT(meta = (DisplayName = "Add Physics Components", Keywords = "Spawn,Construction,Create,New", Varying))
 struct FRigUnit_AddPhysicsComponents : public FRigUnit_PhysicsBaseMutable
 {
 	GENERATED_BODY()
@@ -244,7 +244,7 @@ struct FRigUnit_AddPhysicsComponents : public FRigUnit_PhysicsBaseMutable
 	{
 		Owner.Type = ERigElementType::Bone;
 		Solver.PhysicsSolverComponentKey.ElementKey.Type = ERigElementType::Bone;
-		Solver.PhysicsSolverComponentKey.Name = FRigPhysicsBodyComponent::GetDefaultName();
+		Solver.PhysicsSolverComponentKey.Name = FRigPhysicsSolverComponent::GetDefaultName();
 	}
 
 	RIGVM_METHOD()
@@ -310,6 +310,60 @@ struct FRigUnit_AddPhysicsComponents : public FRigUnit_PhysicsBaseMutable
 	FPhysicsControlData ParentSpaceControlData;
 };
 
+// Imports/creates bones from the physics asset and creates collision for them.
+// The bones will lose their hierarchy and be placed under the specified parent - ready to be moved around.
+USTRUCT(meta = (DisplayName = "Import Collision From Physics Asset", Keywords = "Construction,Create,New", Varying))
+struct FRigUnit_HierarchyImportCollisionFromPhysicsAsset : public FRigUnit_PhysicsBaseMutable
+{
+	GENERATED_BODY()
+
+	FRigUnit_HierarchyImportCollisionFromPhysicsAsset()
+	{
+		PhysicsSolverComponentKey.ElementKey.Type = ERigElementType::Bone;
+		PhysicsSolverComponentKey.Name = FRigPhysicsSolverComponent::GetDefaultName();
+	}
+
+	RIGVM_METHOD()
+	UE_API virtual void Execute() override;
+
+	// Note that setting the solver component, if known, has the benefit of avoiding the need to
+	// search for an automatic solver.
+	UPROPERTY(meta = (Input))
+	FRigComponentKey PhysicsSolverComponentKey;
+
+	// If true (and the physics solver is not explicitly set), then this component will be added to
+	// any physics solver that exists above it in the hierarchy, if that solver allows automatically
+	// adding physics components.
+	UPROPERTY(meta = (Input))
+	bool bUseAutomaticSolver = true;
+
+	UPROPERTY(meta = (Input))
+	TObjectPtr<UPhysicsAsset> PhysicsAsset;
+
+	// Name of the constraint profile to use. If empty (or invalid), the default profile will be used
+	UPROPERTY(meta = (Input))
+	FName ConstraintProfileName;
+
+	// If this is empty, then all bones with bodies in the physics asset will be created. Otherwise
+	// only bodies that relate to the specified bones will be created.
+	UPROPERTY(meta = (Input))
+	TArray<FName> BonesToUse;
+
+	// Prefix to the bone names
+	UPROPERTY(meta = (Input))
+	FName NameSpace = "Physics_";
+
+	// Parent/owner for all the new bones
+	UPROPERTY(meta = (Input))
+	FRigElementKey Owner;
+
+	UPROPERTY(meta = (Output))
+	TArray<FRigElementKey> BoneKeys;
+
+	UPROPERTY(meta = (Output))
+	TArray<FRigComponentKey> PhysicsBodyComponentKeys;
+};
+
 // Creates multiple physics components based on the supplied physics asset.
 // Note that the resulting simulation bodies may not precisely match the physics asset.
 USTRUCT(meta = (DisplayName = "Instantiate From Physics Asset", Keywords = "Construction,Create,New", Varying))
@@ -319,13 +373,23 @@ struct FRigUnit_HierarchyInstantiateFromPhysicsAsset : public FRigUnit_PhysicsBa
 
 	FRigUnit_HierarchyInstantiateFromPhysicsAsset()
 	{
+		PhysicsSolverComponentKey.ElementKey.Type = ERigElementType::Bone;
+		PhysicsSolverComponentKey.Name = FRigPhysicsSolverComponent::GetDefaultName();
 	}
 
 	RIGVM_METHOD()
 	UE_API virtual void Execute() override;
 
+	// Note that setting the solver component, if known, has the benefit of avoiding the need to
+	// search for an automatic solver.
 	UPROPERTY(meta = (Input))
-	FRigPhysicsBodySolverSettings Solver;
+	FRigComponentKey PhysicsSolverComponentKey;
+
+	// If true (and the physics solver is not explicitly set), then this component will be added to
+	// any physics solver that exists above it in the hierarchy, if that solver allows automatically
+	// adding physics components.
+	UPROPERTY(meta = (Input))
+	bool bUseAutomaticSolver = true;
 
 	UPROPERTY(meta = (Input))
 	TObjectPtr<UPhysicsAsset> PhysicsAsset;
diff --git a/Engine/Plugins/Experimental/PhysicsControl/Source/ControlRigPhysics/Public/RigPhysicsJointExecution.h b/Engine/Plugins/Experimental/PhysicsControl/Source/ControlRigPhysics/Public/RigPhysicsJointExecution.h
--- a/Engine/Plugins/Experimental/PhysicsControl/Source/ControlRigPhysics/Public/RigPhysicsJointExecution.h
+++ b/Engine/Plugins/Experimental/PhysicsControl/Source/ControlRigPhysics/Public/RigPhysicsJointExecution.h
@@ -18,7 +18,7 @@
  * Adds a new physics body as a component on the owner element.
  * Note: This node only runs as part of the construction event.
  */
-USTRUCT(meta = (DisplayName = "Add Physics Joint Component", Keywords = "Construction,Create,New,Body,Skeleton", Varying))
+USTRUCT(meta = (DisplayName = "Add Physics Joint Component", Keywords = "Spawn,Construction,Create,New,Joint,Articulation,Skeleton", Varying))
 struct FRigUnit_AddPhysicsJoint : public FRigUnit_PhysicsBaseMutable
 {
 	GENERATED_BODY()
