Fix bogus resets from my last submit.

diff --git a/Engine/Plugins/Experimental/PhysicsControl/Source/ControlRigPhysics/Private/RigPhysicsBodyExecution.cpp b/Engine/Plugins/Experimental/PhysicsControl/Source/ControlRigPhysics/Private/RigPhysicsBodyExecution.cpp
--- a/Engine/Plugins/Experimental/PhysicsControl/Source/ControlRigPhysics/Private/RigPhysicsBodyExecution.cpp
+++ b/Engine/Plugins/Experimental/PhysicsControl/Source/ControlRigPhysics/Private/RigPhysicsBodyExecution.cpp
@@ -208,6 +208,21 @@ FRigUnit_HierarchySetPhysicsBodyCollisionType_Execute()
 	}
 }
 
+//======================================================================================================================
+FRigUnit_HierarchySetPhysicsBodyIncludeInChecksForReset_Execute()
+{
+	DECLARE_SCOPE_HIERARCHICAL_COUNTER_RIGUNIT();
+	if (!ExecuteContext.Hierarchy)
+	{
+		return;
+	}
+	if (FRigPhysicsBodyComponent* Component = Cast<FRigPhysicsBodyComponent>(
+		ExecuteContext.Hierarchy->FindComponent(PhysicsBodyComponentKey)))
+	{
+		Component->BodySolverSettings.bIncludeInChecksForReset = bInclude;
+	}
+}
+
 //======================================================================================================================
 FRigUnit_HierarchySetPhysicsBodyGravityMultiplier_Execute()
 {
diff --git a/Engine/Plugins/Experimental/PhysicsControl/Source/ControlRigPhysics/Private/RigPhysicsData.cpp b/Engine/Plugins/Experimental/PhysicsControl/Source/ControlRigPhysics/Private/RigPhysicsData.cpp
--- a/Engine/Plugins/Experimental/PhysicsControl/Source/ControlRigPhysics/Private/RigPhysicsData.cpp
+++ b/Engine/Plugins/Experimental/PhysicsControl/Source/ControlRigPhysics/Private/RigPhysicsData.cpp
@@ -383,6 +383,10 @@ FArchive& operator <<(FArchive& Ar, FRigPhysicsBodySolverSettings& Data)
 	{
 		Ar << Data.bUseAutomaticSolver;
 	}
+	if (Ar.CustomVer(FPhysicsControlObjectVersion::GUID) >= FPhysicsControlObjectVersion::ControlRigBodyIncludeInChecksForReset)
+	{
+		Ar << Data.bIncludeInChecksForReset;
+	}
 	return Ar;
 }
 
diff --git a/Engine/Plugins/Experimental/PhysicsControl/Source/ControlRigPhysics/Private/RigPhysicsSimulation.cpp b/Engine/Plugins/Experimental/PhysicsControl/Source/ControlRigPhysics/Private/RigPhysicsSimulation.cpp
--- a/Engine/Plugins/Experimental/PhysicsControl/Source/ControlRigPhysics/Private/RigPhysicsSimulation.cpp
+++ b/Engine/Plugins/Experimental/PhysicsControl/Source/ControlRigPhysics/Private/RigPhysicsSimulation.cpp
@@ -1201,6 +1201,9 @@ void FRigPhysicsSimulation::CheckForResetsPrePhysics(
 {
 	DECLARE_SCOPE_HIERARCHICAL_COUNTER_FUNC()
 	QUICK_SCOPE_CYCLE_COUNTER(STAT_RigPhysics_CheckForResetsPrePhysics);
+
+	URigHierarchy* Hierarchy = OwningControlRig->GetHierarchy();
+
 	const FRigPhysicsSolverSettings& SolverSettings = SolverComponent->SolverSettings;
 
 	double SpeedThresholdForResetSquared = SolverSettings.KinematicSpeedThresholdForReset > 0 ?
@@ -1209,6 +1212,12 @@ void FRigPhysicsSimulation::CheckForResetsPrePhysics(
 	double AccelerationThresholdForResetSquared = SolverSettings.KinematicAccelerationThresholdForReset > 0 ?
 		FMath::Square(SolverSettings.KinematicAccelerationThresholdForReset) : TNumericLimits<double>::Max();
 
+	if (SpeedThresholdForResetSquared == TNumericLimits<double>::Max() &&
+		AccelerationThresholdForResetSquared == TNumericLimits<double>::Max())
+	{
+		return;
+	}
+
 	double HighestSpeedSq = -1.0;
 	double HighestAccelerationSq = -1.0;
 
@@ -1219,17 +1228,23 @@ void FRigPhysicsSimulation::CheckForResetsPrePhysics(
 
 		if (Record.ActorHandle && Record.ActorHandle->GetIsKinematic())
 		{
-			const FVector Velocity = Record.SourceComponentSpaceVelocity;
-			const FVector Acceleration =
-				DeltaTime > SMALL_NUMBER && (UpdateCounter == PreviousUpdateCounter + 1)
-				? (Record.SourceComponentSpaceVelocity - Record.PrevSourceComponentSpaceVelocity) / DeltaTime
-				: FVector::ZeroVector;
-
-			HighestSpeedSq = FMath::Max(HighestSpeedSq, Velocity.SquaredLength());
-			HighestAccelerationSq = FMath::Max(HighestAccelerationSq, Acceleration.SquaredLength());
+			if (const FRigPhysicsBodyComponent* PhysicsComponent = Cast<FRigPhysicsBodyComponent>(
+				Hierarchy->FindComponent(ComponentKey)))
+			{
+				if (PhysicsComponent->BodySolverSettings.bIncludeInChecksForReset)
+				{
+					const FVector Velocity = Record.SourceComponentSpaceVelocity;
+					const FVector Acceleration =
+						DeltaTime > SMALL_NUMBER && (UpdateCounter == PreviousUpdateCounter + 1)
+						? (Record.SourceComponentSpaceVelocity - Record.PrevSourceComponentSpaceVelocity) / DeltaTime
+						: FVector::ZeroVector;
+
+					HighestSpeedSq = FMath::Max(HighestSpeedSq, Velocity.SquaredLength());
+					HighestAccelerationSq = FMath::Max(HighestAccelerationSq, Acceleration.SquaredLength());
+				}
+			}
 		}
 	}
-
 	if (HighestSpeedSq > SpeedThresholdForResetSquared || HighestAccelerationSq > AccelerationThresholdForResetSquared)
 	{
 		if (HighestSpeedSq > SpeedThresholdForResetSquared)
@@ -1682,7 +1697,7 @@ void FRigPhysicsSimulation::StepSimulation(
 	}
 
 	// Always do a read-back, even for zero Dt
-	UpdatePostPhysics(SolverComponent, Alpha, PhysicsDeltaTime);
+	UpdatePostPhysics(SolverComponent, FMath::Clamp(Alpha, 0.0f, 1.0f), PhysicsDeltaTime);
 
 	if (SolverComponent->TrackInputCounter > 0)
 	{
diff --git a/Engine/Plugins/Experimental/PhysicsControl/Source/ControlRigPhysics/Public/RigPhysicsBodyExecution.h b/Engine/Plugins/Experimental/PhysicsControl/Source/ControlRigPhysics/Public/RigPhysicsBodyExecution.h
--- a/Engine/Plugins/Experimental/PhysicsControl/Source/ControlRigPhysics/Public/RigPhysicsBodyExecution.h
+++ b/Engine/Plugins/Experimental/PhysicsControl/Source/ControlRigPhysics/Public/RigPhysicsBodyExecution.h
@@ -314,6 +314,29 @@ struct FRigUnit_HierarchySetPhysicsBodyCollisionType : public FRigUnit_PhysicsBa
 	TEnumAsByte<ECollisionEnabled::Type> CollisionType = ECollisionEnabled::QueryAndPhysics;
 };
 
+
+// Sets whether this body should be included in checks for resetting physics on the whole rig
+USTRUCT(meta = (DisplayName = "Set Physics Body Include In Checks For Reset", Varying))
+struct FRigUnit_HierarchySetPhysicsBodyIncludeInChecksForReset : public FRigUnit_PhysicsBaseMutable
+{
+	GENERATED_BODY()
+
+	FRigUnit_HierarchySetPhysicsBodyIncludeInChecksForReset()
+	{
+		PhysicsBodyComponentKey.ElementKey.Type = ERigElementType::Bone;
+		PhysicsBodyComponentKey.Name = FRigPhysicsBodyComponent::GetDefaultName();
+	}
+
+	RIGVM_METHOD()
+	UE_API virtual void Execute() override;
+
+	UPROPERTY(meta = (Input))
+	FRigComponentKey PhysicsBodyComponentKey;
+
+	UPROPERTY(meta = (Input))
+	bool bInclude = true;
+};
+
 /**
  * Sets the multiplier on gravity that should be applied to the body.
  */
diff --git a/Engine/Plugins/Experimental/PhysicsControl/Source/ControlRigPhysics/Public/RigPhysicsData.h b/Engine/Plugins/Experimental/PhysicsControl/Source/ControlRigPhysics/Public/RigPhysicsData.h
--- a/Engine/Plugins/Experimental/PhysicsControl/Source/ControlRigPhysics/Public/RigPhysicsData.h
+++ b/Engine/Plugins/Experimental/PhysicsControl/Source/ControlRigPhysics/Public/RigPhysicsData.h
@@ -710,6 +710,11 @@ struct FRigPhysicsBodySolverSettings
 	UPROPERTY(EditAnywhere, Category = Solver)
 	FRigElementKey TargetBone;
 
+	// Whether to include this bone in checks for whether we should reset physics on the whole rig,
+	// using thresholds in the solver. Set to false for bodies you know may jump around safely when kinematic.
+	UPROPERTY(EditAnywhere, Category = Solver)
+	bool bIncludeInChecksForReset = true;
+
 	UE_API void OnRigHierarchyKeyChanged(const FRigHierarchyKey& InOldKey, const FRigHierarchyKey& InNewKey);
 };
 
diff --git a/Engine/Plugins/Experimental/PhysicsControl/Source/PhysicsControl/Public/PhysicsControlObjectVersion.h b/Engine/Plugins/Experimental/PhysicsControl/Source/PhysicsControl/Public/PhysicsControlObjectVersion.h
--- a/Engine/Plugins/Experimental/PhysicsControl/Source/PhysicsControl/Public/PhysicsControlObjectVersion.h
+++ b/Engine/Plugins/Experimental/PhysicsControl/Source/PhysicsControl/Public/PhysicsControlObjectVersion.h
@@ -56,6 +56,8 @@ struct FPhysicsControlObjectVersion
 
 		ControlRigDetectExternalVelocityTurbulence,
 
+		ControlRigBodyIncludeInChecksForReset,
+
 		// -----<new versions can be added above this line>-------------------------------------------------
 		VersionPlusOne,
 		LatestVersion = VersionPlusOne - 1
