Add cvar to override control rig physics max step (i.e. to disable substepping) - mainly for profiling. Trivial & no change unless the cvar is used.

diff --git a/Engine/Plugins/Experimental/PhysicsControl/Source/ControlRigPhysics/Private/RigPhysicsSimulation.cpp b/Engine/Plugins/Experimental/PhysicsControl/Source/ControlRigPhysics/Private/RigPhysicsSimulation.cpp
--- a/Engine/Plugins/Experimental/PhysicsControl/Source/ControlRigPhysics/Private/RigPhysicsSimulation.cpp
+++ b/Engine/Plugins/Experimental/PhysicsControl/Source/ControlRigPhysics/Private/RigPhysicsSimulation.cpp
@@ -57,6 +57,10 @@ TAutoConsoleVariable<int> CVarControlRigPhysicsMaxTimeStepsOverride(
 	TEXT("ControlRig.Physics.MaxTimeStepsOverride"), -1,
 	TEXT("-1 disables the override, so the max timesteps authored in the simulation settings will be used. A +ve value is used to specify the maximum number of timesteps."));
 
+TAutoConsoleVariable<float> CVarControlRigPhysicsMaxDeltaTimeOverride(
+	TEXT("ControlRig.Physics.MaxDeltaTimeOverride"), -1,
+	TEXT("-1 disables the override, so the max delta time authored in the simulation settings will be used. A +ve value is used to specify the maximum delta time."));
+
 constexpr int32 ConstraintChildIndex = 0;
 constexpr int32 ConstraintParentIndex = 1;
 
@@ -1638,6 +1642,8 @@ void FRigPhysicsSimulation::StepSimulation(
 		? SolverSettings.FixedTimeStep : CVarControlRigPhysicsFixedTimeStepOverride.GetValueOnAnyThread();
 	int MaxTimeSteps = CVarControlRigPhysicsMaxTimeStepsOverride.GetValueOnAnyThread() < 0
 		? SolverSettings.MaxTimeSteps : CVarControlRigPhysicsMaxTimeStepsOverride.GetValueOnAnyThread();
+	float MaxDeltaTime = CVarControlRigPhysicsMaxDeltaTimeOverride.GetValueOnAnyThread() < 0
+		? SolverSettings.MaxDeltaTime : CVarControlRigPhysicsMaxDeltaTimeOverride.GetValueOnAnyThread();
 
 	// Set settings that might change
 	Simulation->SetSolverSettings(
@@ -1690,7 +1696,7 @@ void FRigPhysicsSimulation::StepSimulation(
 		UpdatePrePhysics(ExecuteContext, SolverComponent, PhysicsDeltaTime);
 
 		Simulation->Simulate(
-			PhysicsDeltaTime, SolverSettings.MaxDeltaTime, MaxTimeSteps,
+			PhysicsDeltaTime, MaxDeltaTime, MaxTimeSteps,
 			SimulationSpaceData.Gravity, &ChaosJointSolverSettings);
 
 		PreviousUpdateCounter = UpdateCounter;
