Trivial changes to improve debugging of controls in control rig. No change in runtime functionality.

diff --git a/Engine/Plugins/Experimental/PhysicsControl/Source/ControlRigPhysics/Private/RigPhysicsSimulationVisualization.cpp b/Engine/Plugins/Experimental/PhysicsControl/Source/ControlRigPhysics/Private/RigPhysicsSimulationVisualization.cpp
--- a/Engine/Plugins/Experimental/PhysicsControl/Source/ControlRigPhysics/Private/RigPhysicsSimulationVisualization.cpp
+++ b/Engine/Plugins/Experimental/PhysicsControl/Source/ControlRigPhysics/Private/RigPhysicsSimulationVisualization.cpp
@@ -176,7 +176,8 @@ static void DrawJoint(
 	FRigVMDrawInterface*                    DI,
 	const FRigPhysicsVisualizationSettings& VisualizationSettings,
 	const FTransform&                       SpaceTransform,
-	const ImmediatePhysics::FJointHandle&   JointHandle)
+	const ImmediatePhysics::FJointHandle&   JointHandle,
+	const bool                              bDrawAsAxes)
 {
 	const Chaos::FPBDJointConstraintHandle* ConstraintHandle = JointHandle.GetConstraint();
 	if (!ConstraintHandle)
@@ -201,10 +202,11 @@ static void DrawJoint(
 	FTransform DialFrame = JointSettings.ConnectorTransforms[0] * ChildActorTM;
 	FTransform LimitFrame = JointSettings.ConnectorTransforms[1] * ParentActorTM;
 
-	static bool bDrawAsAxes = false;
 	if (bDrawAsAxes)
 	{
-		DI->DrawAxes(LimitFrame, FTransform(), Size, VisualizationSettings.LineThickness * 2);
+		DI->DrawAxes(
+			FTransform(JointSettings.AngularDrivePositionTarget, JointSettings.LinearDrivePositionTarget) * LimitFrame, 
+			FTransform(), Size, VisualizationSettings.LineThickness * 2);
 		DI->DrawAxes(DialFrame, FTransform(), Size, VisualizationSettings.LineThickness);
 	}
 
@@ -328,6 +330,35 @@ static void DrawJoint(
 	}
 }
 
+//======================================================================================================================
+static bool IsConstraintDriveActive(const Chaos::FPBDJointConstraintHandle* Constraint)
+{
+	const Chaos::FPBDJointSettings& Settings = Constraint->GetSettings();
+	if (Settings.bAngularSLerpPositionDriveEnabled || Settings.bAngularSwingPositionDriveEnabled || Settings.bAngularTwistPositionDriveEnabled)
+	{
+		if (Settings.AngularDriveStiffness.SizeSquared() > 0.0)
+		{
+			return true;
+		}
+	}
+	if (Settings.bAngularSLerpVelocityDriveEnabled || Settings.bAngularSwingVelocityDriveEnabled || Settings.bAngularTwistVelocityDriveEnabled)
+	{
+		if (Settings.AngularDriveDamping.SizeSquared() > 0.0)
+		{
+			return true;
+		}
+	}
+	if (Settings.bLinearPositionDriveEnabled.Max() && Settings.LinearDriveStiffness.SizeSquared() > 0.0)
+	{
+		return true;
+	}
+	if (Settings.bLinearVelocityDriveEnabled.Max() && Settings.LinearDriveDamping.SizeSquared() > 0.0)
+	{
+		return true;
+	}
+	return false;
+}
+
 //======================================================================================================================
 void FRigPhysicsSimulation::Draw(
 	FRigVMDrawInterface*                    DI,
@@ -357,25 +388,50 @@ void FRigPhysicsSimulation::Draw(
 	// component space.
 	const FTransform SpaceTransform = ConvertSimSpaceTransformToComponentSpace(SolverSettings, FTransform());
 
-	if (CollisionActorHandle)
+	if (VisualizationSettings.bShowBodies)
 	{
-		DrawActor(DI, VisualizationSettings, SpaceTransform, *CollisionActorHandle);
+		if (CollisionActorHandle)
+		{
+			DrawActor(DI, VisualizationSettings, SpaceTransform, *CollisionActorHandle);
+		}
+
+		for (const TPair<FRigComponentKey, FRigBodyRecord>& BodyRecordPair : BodyRecords)
+		{
+			const FRigBodyRecord& Record = BodyRecordPair.Value;
+			if (const ImmediatePhysics::FActorHandle* ActorHandle = Record.ActorHandle)
+			{
+				DrawActor(DI, VisualizationSettings, SpaceTransform, *ActorHandle);
+			}
+		}
 	}
 
-	for (const TPair<FRigComponentKey, FRigBodyRecord>& BodyRecordPair : BodyRecords)
+	if (VisualizationSettings.bShowJoints)
 	{
-		const FRigBodyRecord& Record = BodyRecordPair.Value;
-		if (const ImmediatePhysics::FActorHandle* ActorHandle = Record.ActorHandle)
+		for (const TPair<FRigComponentKey, FRigJointRecord>& JointRecordPair : JointRecords)
 		{
-			DrawActor(DI, VisualizationSettings, SpaceTransform, *ActorHandle);
+			const FRigJointRecord& Record = JointRecordPair.Value;
+			if (const ImmediatePhysics::FJointHandle* JointHandle = Record.JointHandle)
+			{
+				DrawJoint(DI, VisualizationSettings, SpaceTransform, *JointHandle, false);
+			}
 		}
 	}
-	for (const TPair<FRigComponentKey, FRigJointRecord>& JointRecordPair : JointRecords)
+
+	if (VisualizationSettings.bShowControls)
 	{
-		const FRigJointRecord& Record = JointRecordPair.Value;
-		if (const ImmediatePhysics::FJointHandle* JointHandle = Record.JointHandle)
+		for (const TPair<FRigComponentKey, FRigControlRecord>& ControlRecordPair : ControlRecords)
 		{
-			DrawJoint(DI, VisualizationSettings, SpaceTransform, *JointHandle);
+			const FRigControlRecord& Record = ControlRecordPair.Value;
+			if (const ImmediatePhysics::FJointHandle* JointHandle = Record.JointHandle)
+			{
+				if (const Chaos::FPBDJointConstraintHandle* Constraint = JointHandle->GetConstraint())
+				{
+					if (IsConstraintDriveActive(Constraint))
+					{
+						DrawJoint(DI, VisualizationSettings, SpaceTransform, *JointHandle, true);
+					}
+				}
+			}
 		}
 	}
 
diff --git a/Engine/Plugins/Experimental/PhysicsControl/Source/ControlRigPhysics/Public/RigPhysicsData.h b/Engine/Plugins/Experimental/PhysicsControl/Source/ControlRigPhysics/Public/RigPhysicsData.h
--- a/Engine/Plugins/Experimental/PhysicsControl/Source/ControlRigPhysics/Public/RigPhysicsData.h
+++ b/Engine/Plugins/Experimental/PhysicsControl/Source/ControlRigPhysics/Public/RigPhysicsData.h
@@ -42,6 +42,15 @@ struct FRigPhysicsVisualizationSettings
 	UPROPERTY(EditAnywhere, Category = Visualization, meta = (ClampMin = "0.0"))
 	int32 ShapeDetail = 16;
 
+	UPROPERTY(EditAnywhere, Category = Visualization)
+	bool bShowBodies = true;
+
+	UPROPERTY(EditAnywhere, Category = Visualization)
+	bool bShowJoints = true;
+
+	UPROPERTY(EditAnywhere, Category = Visualization)
+	bool bShowControls = false;
+
 	// If visualization is enabled, show active contacts. Note that this can be overridden using
 	// CVar ControlRig.Physics.DrawActiveContactsOveride
 	UPROPERTY(EditAnywhere, Category = Visualization)
