Add trivial changes to allow setting material properties in control rig physics, and a cvar to override stepping the physics solvers in control rig for testing: ControlRig.Physics.EnableStepSolver

diff --git a/Engine/Plugins/Experimental/PhysicsControl/Source/ControlRigPhysics/Private/RigPhysicsBodyExecution.cpp b/Engine/Plugins/Experimental/PhysicsControl/Source/ControlRigPhysics/Private/RigPhysicsBodyExecution.cpp
--- a/Engine/Plugins/Experimental/PhysicsControl/Source/ControlRigPhysics/Private/RigPhysicsBodyExecution.cpp
+++ b/Engine/Plugins/Experimental/PhysicsControl/Source/ControlRigPhysics/Private/RigPhysicsBodyExecution.cpp
@@ -223,6 +223,27 @@ FRigUnit_HierarchySetPhysicsBodyIncludeInChecksForReset_Execute()
 	}
 }
 
+//======================================================================================================================
+FRigUnit_HierarchySetPhysicsBodyMaterial_Execute()
+{
+	DECLARE_SCOPE_HIERARCHICAL_COUNTER_RIGUNIT();
+	if (!ExecuteContext.Hierarchy)
+	{
+		return;
+	}
+	if (ExecuteContext.GetEventName() != FRigUnit_PrepareForExecution::EventName)
+	{
+		UE_CONTROLRIG_RIGUNIT_REPORT_ERROR(TEXT("SetPhysicsBodyMaterial can only be used during Setup"));
+	}
+
+	if (FRigPhysicsBodyComponent* Component = Cast<FRigPhysicsBodyComponent>(
+		ExecuteContext.Hierarchy->FindComponent(PhysicsBodyComponentKey)))
+	{
+		Component->Collision.Material = Material;
+	}
+}
+
+
 //======================================================================================================================
 FRigUnit_HierarchySetPhysicsBodyGravityMultiplier_Execute()
 {
diff --git a/Engine/Plugins/Experimental/PhysicsControl/Source/ControlRigPhysics/Private/RigPhysicsExecution.cpp b/Engine/Plugins/Experimental/PhysicsControl/Source/ControlRigPhysics/Private/RigPhysicsExecution.cpp
--- a/Engine/Plugins/Experimental/PhysicsControl/Source/ControlRigPhysics/Private/RigPhysicsExecution.cpp
+++ b/Engine/Plugins/Experimental/PhysicsControl/Source/ControlRigPhysics/Private/RigPhysicsExecution.cpp
@@ -29,6 +29,11 @@ TAutoConsoleVariable<int> CVarControlRigPhysicsVisualizationOverride(
 	TEXT("ControlRig.Physics.VisualizationOveride"), -1,
 	TEXT("-1 uses the Step Physics Solver setting, 0 forces visualization to be disabled, 1 forces it to be enabled."));
 
+TAutoConsoleVariable<int> CVarControlRigPhysicsEnableStepSolver(
+	TEXT("ControlRig.Physics.EnableStepSolver"), 1,
+	TEXT("Setting to zero disables stepping the solver"));
+
+
 //======================================================================================================================
 FRigUnit_AddPhysicsSolver_Execute()
 {
@@ -93,6 +98,10 @@ FRigUnit_StepPhysicsSolver_Execute()
 	{
 		return;
 	}
+	if (CVarControlRigPhysicsEnableStepSolver.GetValueOnAnyThread() <= 0)
+	{
+		return;
+	}
 	if (UControlRig* ControlRig = Cast<UControlRig>(ExecuteContext.Hierarchy->GetOuter()))
 	{
 		if (FRigPhysicsSolverComponent* SolverComponent = Cast<FRigPhysicsSolverComponent>(
diff --git a/Engine/Plugins/Experimental/PhysicsControl/Source/ControlRigPhysics/Public/RigPhysicsBodyExecution.h b/Engine/Plugins/Experimental/PhysicsControl/Source/ControlRigPhysics/Public/RigPhysicsBodyExecution.h
--- a/Engine/Plugins/Experimental/PhysicsControl/Source/ControlRigPhysics/Public/RigPhysicsBodyExecution.h
+++ b/Engine/Plugins/Experimental/PhysicsControl/Source/ControlRigPhysics/Public/RigPhysicsBodyExecution.h
@@ -337,6 +337,28 @@ struct FRigUnit_HierarchySetPhysicsBodyIncludeInChecksForReset : public FRigUnit
 	bool bInclude = true;
 };
 
+// Applies the material settings to all the collision shapes
+USTRUCT(meta = (DisplayName = "Set Physics Body Material", Varying))
+struct FRigUnit_HierarchySetPhysicsBodyMaterial : public FRigUnit_PhysicsBaseMutable
+{
+	GENERATED_BODY()
+
+	FRigUnit_HierarchySetPhysicsBodyMaterial()
+	{
+		PhysicsBodyComponentKey.ElementKey.Type = ERigElementType::Bone;
+		PhysicsBodyComponentKey.Name = FRigPhysicsBodyComponent::GetDefaultName();
+	}
+
+	RIGVM_METHOD()
+	UE_API virtual void Execute() override;
+
+	UPROPERTY(meta = (Input))
+	FRigComponentKey PhysicsBodyComponentKey;
+
+	UPROPERTY(meta = (Input))
+	FRigPhysicsMaterial Material;
+};
+
 /**
  * Sets the multiplier on gravity that should be applied to the body.
  */
