Fix control rig physics log message (degrees vs radians).
Adds simple control rig physics node to set control data and multipliers in one call

diff --git a/Engine/Plugins/Experimental/PhysicsControl/Source/ControlRigPhysics/Private/RigPhysicsControlExecution.cpp b/Engine/Plugins/Experimental/PhysicsControl/Source/ControlRigPhysics/Private/RigPhysicsControlExecution.cpp
--- a/Engine/Plugins/Experimental/PhysicsControl/Source/ControlRigPhysics/Private/RigPhysicsControlExecution.cpp
+++ b/Engine/Plugins/Experimental/PhysicsControl/Source/ControlRigPhysics/Private/RigPhysicsControlExecution.cpp
@@ -92,7 +92,7 @@ FRigUnit_HierarchySetControlData_Execute()
 }
 
 //======================================================================================================================
-FRigUnit_HierarchySetControlTarget_Execute()
+FRigUnit_HierarchySetControlMultiplier_Execute()
 {
 	DECLARE_SCOPE_HIERARCHICAL_COUNTER_RIGUNIT();
 	if (ExecuteContext.Hierarchy)
@@ -100,13 +100,13 @@ FRigUnit_HierarchySetControlTarget_Execute()
 		if (FRigPhysicsControlComponent* Component = Cast<FRigPhysicsControlComponent>(
 			ExecuteContext.Hierarchy->FindComponent(PhysicsControlComponentKey)))
 		{
-			Component->ControlTarget = ControlTarget;
+			Component->ControlMultiplier = ControlMultiplier;
 		}
 	}
 }
 
 //======================================================================================================================
-FRigUnit_HierarchyUpdateControlTarget_Execute()
+FRigUnit_HierarchySetControlDataAndMultiplier_Execute()
 {
 	DECLARE_SCOPE_HIERARCHICAL_COUNTER_RIGUNIT();
 	if (ExecuteContext.Hierarchy)
@@ -114,18 +114,14 @@ FRigUnit_HierarchyUpdateControlTarget_Execute()
 		if (FRigPhysicsControlComponent* Component = Cast<FRigPhysicsControlComponent>(
 			ExecuteContext.Hierarchy->FindComponent(PhysicsControlComponentKey)))
 		{
-			Component->ControlTarget.TargetVelocity = UE::PhysicsControl::CalculateLinearVelocity(
-				TargetPosition, Component->ControlTarget.TargetPosition, DeltaTime);
-			Component->ControlTarget.TargetAngularVelocity = UE::PhysicsControl::CalculateAngularVelocity(
-				TargetOrientation.Quaternion(), Component->ControlTarget.TargetOrientation.Quaternion(), DeltaTime);
-			Component->ControlTarget.TargetPosition = TargetPosition;
-			Component->ControlTarget.TargetOrientation = TargetOrientation;
+			Component->ControlData = ControlData;
+			Component->ControlMultiplier = ControlMultiplier;
 		}
 	}
 }
 
 //======================================================================================================================
-FRigUnit_HierarchySetControlMultiplier_Execute()
+FRigUnit_HierarchySetControlTarget_Execute()
 {
 	DECLARE_SCOPE_HIERARCHICAL_COUNTER_RIGUNIT();
 	if (ExecuteContext.Hierarchy)
@@ -133,7 +129,26 @@ FRigUnit_HierarchySetControlMultiplier_Execute()
 		if (FRigPhysicsControlComponent* Component = Cast<FRigPhysicsControlComponent>(
 			ExecuteContext.Hierarchy->FindComponent(PhysicsControlComponentKey)))
 		{
-			Component->ControlMultiplier = ControlMultiplier;
+			Component->ControlTarget = ControlTarget;
+		}
+	}
+}
+
+//======================================================================================================================
+FRigUnit_HierarchyUpdateControlTarget_Execute()
+{
+	DECLARE_SCOPE_HIERARCHICAL_COUNTER_RIGUNIT();
+	if (ExecuteContext.Hierarchy)
+	{
+		if (FRigPhysicsControlComponent* Component = Cast<FRigPhysicsControlComponent>(
+			ExecuteContext.Hierarchy->FindComponent(PhysicsControlComponentKey)))
+		{
+			Component->ControlTarget.TargetVelocity = UE::PhysicsControl::CalculateLinearVelocity(
+				TargetPosition, Component->ControlTarget.TargetPosition, DeltaTime);
+			Component->ControlTarget.TargetAngularVelocity = UE::PhysicsControl::CalculateAngularVelocity(
+				TargetOrientation.Quaternion(), Component->ControlTarget.TargetOrientation.Quaternion(), DeltaTime);
+			Component->ControlTarget.TargetPosition = TargetPosition;
+			Component->ControlTarget.TargetOrientation = TargetOrientation;
 		}
 	}
 }
diff --git a/Engine/Plugins/Experimental/PhysicsControl/Source/ControlRigPhysics/Private/RigPhysicsSimulationSpace.cpp b/Engine/Plugins/Experimental/PhysicsControl/Source/ControlRigPhysics/Private/RigPhysicsSimulationSpace.cpp
--- a/Engine/Plugins/Experimental/PhysicsControl/Source/ControlRigPhysics/Private/RigPhysicsSimulationSpace.cpp
+++ b/Engine/Plugins/Experimental/PhysicsControl/Source/ControlRigPhysics/Private/RigPhysicsSimulationSpace.cpp
@@ -324,7 +324,7 @@ FRigPhysicsSimulation::FSimulationSpaceData FRigPhysicsSimulation::UpdateSimulat
 		if (bAngularAccelerationTrigger)
 		{
 			UE_LOG(LogRigPhysics, Log, TEXT("Detected angular Acceleration (%f > %f) teleport in %s"),
-				SimulationSpaceData.AngularAcceleration.Length(),
+				FMath::RadiansToDegrees(SimulationSpaceData.AngularAcceleration.Length()),
 				SimulationSpaceSettings.AngularAccelerationThresholdForTeleport,
 				*OwningControlRig->GetName());
 		}
diff --git a/Engine/Plugins/Experimental/PhysicsControl/Source/ControlRigPhysics/Public/RigPhysicsControlExecution.h b/Engine/Plugins/Experimental/PhysicsControl/Source/ControlRigPhysics/Public/RigPhysicsControlExecution.h
--- a/Engine/Plugins/Experimental/PhysicsControl/Source/ControlRigPhysics/Public/RigPhysicsControlExecution.h
+++ b/Engine/Plugins/Experimental/PhysicsControl/Source/ControlRigPhysics/Public/RigPhysicsControlExecution.h
@@ -130,6 +130,53 @@ struct FRigUnit_HierarchySetControlData : public FRigUnit_PhysicsBaseMutable
 	FPhysicsControlData ControlData;
 };
 
+// Sets the multipliers for a physics control
+USTRUCT(meta = (DisplayName = "Set Physics Control Multiplier", Varying))
+struct FRigUnit_HierarchySetControlMultiplier : public FRigUnit_PhysicsBaseMutable
+{
+	GENERATED_BODY()
+
+	FRigUnit_HierarchySetControlMultiplier()
+	{
+		PhysicsControlComponentKey.ElementKey.Type = ERigElementType::Bone;
+		PhysicsControlComponentKey.Name = FRigPhysicsControlComponent::GetDefaultName();
+	}
+
+	RIGVM_METHOD()
+		UE_API virtual void Execute() override;
+
+	UPROPERTY(meta = (Input))
+	FRigComponentKey PhysicsControlComponentKey;
+
+	UPROPERTY(meta = (Input))
+	FPhysicsControlMultiplier ControlMultiplier;
+};
+
+// Sets the control data and multiplier for a physics control
+USTRUCT(meta = (DisplayName = "Set Physics Control Data And Multiplier", Varying))
+struct FRigUnit_HierarchySetControlDataAndMultiplier : public FRigUnit_PhysicsBaseMutable
+{
+	GENERATED_BODY()
+
+	FRigUnit_HierarchySetControlDataAndMultiplier()
+	{
+		PhysicsControlComponentKey.ElementKey.Type = ERigElementType::Bone;
+		PhysicsControlComponentKey.Name = FRigPhysicsControlComponent::GetDefaultName();
+	}
+
+	RIGVM_METHOD()
+		UE_API virtual void Execute() override;
+
+	UPROPERTY(meta = (Input))
+	FRigComponentKey PhysicsControlComponentKey;
+
+	UPROPERTY(meta = (Input))
+	FPhysicsControlData ControlData;
+
+	UPROPERTY(meta = (Input))
+	FPhysicsControlMultiplier ControlMultiplier;
+};
+
 // Sets the target for a physics control
 USTRUCT(meta = (DisplayName = "Set Physics Control Target", Varying))
 struct FRigUnit_HierarchySetControlTarget : public FRigUnit_PhysicsBaseMutable
@@ -184,26 +231,5 @@ struct CONTROLRIGPHYSICS_API FRigUnit_HierarchyUpdateControlTarget : public FRig
 	float DeltaTime = 0.0f;
 };
 
-// Sets the multipliers for a physics control
-USTRUCT(meta = (DisplayName = "Set Physics Control Multiplier", Varying))
-struct FRigUnit_HierarchySetControlMultiplier : public FRigUnit_PhysicsBaseMutable
-{
-	GENERATED_BODY()
-
-	FRigUnit_HierarchySetControlMultiplier()
-	{
-		PhysicsControlComponentKey.ElementKey.Type = ERigElementType::Bone;
-		PhysicsControlComponentKey.Name = FRigPhysicsControlComponent::GetDefaultName();
-	}
-
-	RIGVM_METHOD()
-	UE_API virtual void Execute() override;
-
-	UPROPERTY(meta = (Input))
-	FRigComponentKey PhysicsControlComponentKey;
-
-	UPROPERTY(meta = (Input))
-	FPhysicsControlMultiplier ControlMultiplier;
-};
 
 #undef UE_API
