Control Rig Physics: Correct post-physics kinematic body velocities when a fixed timestep is being used. Tweaks the defaults for detection of resets/teleports.

diff --git a/Engine/Plugins/Experimental/PhysicsControl/Source/ControlRigPhysics/Private/RigPhysicsSimulation.cpp b/Engine/Plugins/Experimental/PhysicsControl/Source/ControlRigPhysics/Private/RigPhysicsSimulation.cpp
--- a/Engine/Plugins/Experimental/PhysicsControl/Source/ControlRigPhysics/Private/RigPhysicsSimulation.cpp
+++ b/Engine/Plugins/Experimental/PhysicsControl/Source/ControlRigPhysics/Private/RigPhysicsSimulation.cpp
@@ -969,8 +969,13 @@ void FRigPhysicsSimulation::UpdateBodyPrePhysics(
 
 	if (SolverComponent->TrackInputCounter > 0)
 	{
-		BodyData.MovementType = EPhysicsMovementType::Kinematic;
-		KinematicTargetSpace = ERigPhysicsKinematicTargetSpace::IgnoreTarget;
+		// If the body is already kinematic, then we don't want to make it start ignoring a target.
+		// However, dynamic objects should avoid using a target if they are just tracking due to a reset.
+		if (BodyData.MovementType != EPhysicsMovementType::Kinematic)
+		{
+			BodyData.MovementType = EPhysicsMovementType::Kinematic;
+			KinematicTargetSpace = ERigPhysicsKinematicTargetSpace::IgnoreTarget;
+		}
 	}
 
 	UpdateBodyFromModifierData(
@@ -982,7 +987,7 @@ void FRigPhysicsSimulation::UpdateBodyPrePhysics(
 
 		// If the target is already in component space, then that's all we need
 		FTransform KinematicTargetCS;
-		switch (PhysicsComponent->KinematicTargetSpace)
+		switch (KinematicTargetSpace)
 		{
 		case ERigPhysicsKinematicTargetSpace::Component:
 		{
@@ -1007,7 +1012,7 @@ void FRigPhysicsSimulation::UpdateBodyPrePhysics(
 			}
 			if (SourceKey.IsValid())
 			{
-				switch (PhysicsComponent->KinematicTargetSpace)
+				switch (KinematicTargetSpace)
 				{
 				case ERigPhysicsKinematicTargetSpace::OffsetInBoneSpace:
 				{
@@ -1471,7 +1476,10 @@ void FRigPhysicsSimulation::UpdateControlPrePhysics(
 // Note that we read back into a target bone, which may have been specified explicitly, or will
 // otherwise default to the physics element parent.
 void FRigPhysicsSimulation::UpdatePostPhysics(
-	FRigPhysicsSolverComponent* SolverComponent, const float Alpha, const float DeltaTime)
+	FRigPhysicsSolverComponent* SolverComponent, 
+	const float                 Alpha, 
+	const float                 DeltaTime,
+	const bool                  bUsingFixedStep)
 {
 	DECLARE_SCOPE_HIERARCHICAL_COUNTER_FUNC()
 	QUICK_SCOPE_CYCLE_COUNTER(STAT_RigPhysics_UpdatePostPhysics);
@@ -1579,13 +1587,30 @@ void FRigPhysicsSimulation::UpdatePostPhysics(
 		for (const FRigComponentKey& ComponentKey : SortedBodyComponentKeys)
 		{
 			const FRigBodyRecord& Record = *BodyRecords.Find(ComponentKey);
-			if (Record.ActorHandle && Record.TargetElementKey.IsValid())
+			if (Record.ActorHandle)
 			{
-				// Note that we set bAffectChildren = true (i.e. don't counter-animate
-				// children), so that attached animation bones will follow physics, but we
-				// rely on our bodies being sorted so we work out from the leaf nodes so as
-				// not to disturb previously set bodies.
-				Hierarchy->SetGlobalTransform(Record.TargetElementKey, Record.TargetComponentSpaceTM, false, true);
+				// Kinematic velocities will be incorrect when using a fixed timestep, because of
+				// the rewind. Overwrite them here, since we know exactly what they should be, in
+				// case they are subsequently made dynamic. This is particularly important when
+				// warming up velocities following a reset.
+				if (bUsingFixedStep && Record.ActorHandle->GetIsKinematic())
+				{
+					FVector V = ConvertComponentSpaceVectorToSimSpace(
+						SolverSettings, Record.SourceComponentSpaceVelocity);
+					FVector W = ConvertComponentSpaceVectorToSimSpace(
+						SolverSettings, Record.SourceComponentSpaceAngularVelocity);
+					Record.ActorHandle->SetLinearVelocity(V);
+					Record.ActorHandle->SetAngularVelocity(W);
+				}
+
+				if (Record.TargetElementKey.IsValid())
+				{
+					// Note that we set bAffectChildren = true (i.e. don't counter-animate
+					// children), so that attached animation bones will follow physics, but we
+					// rely on our bodies being sorted so we work out from the leaf nodes so as
+					// not to disturb previously set bodies.
+					Hierarchy->SetGlobalTransform(Record.TargetElementKey, Record.TargetComponentSpaceTM, false, true);
+				}
 			}
 		}
 	}
@@ -1703,7 +1728,7 @@ void FRigPhysicsSimulation::StepSimulation(
 	}
 
 	// Always do a read-back, even for zero Dt
-	UpdatePostPhysics(SolverComponent, FMath::Clamp(Alpha, 0.0f, 1.0f), PhysicsDeltaTime);
+	UpdatePostPhysics(SolverComponent, FMath::Clamp(Alpha, 0.0f, 1.0f), PhysicsDeltaTime, FixedTimeStep > 0);
 
 	if (SolverComponent->TrackInputCounter > 0)
 	{
diff --git a/Engine/Plugins/Experimental/PhysicsControl/Source/ControlRigPhysics/Private/RigPhysicsSimulation.h b/Engine/Plugins/Experimental/PhysicsControl/Source/ControlRigPhysics/Private/RigPhysicsSimulation.h
--- a/Engine/Plugins/Experimental/PhysicsControl/Source/ControlRigPhysics/Private/RigPhysicsSimulation.h
+++ b/Engine/Plugins/Experimental/PhysicsControl/Source/ControlRigPhysics/Private/RigPhysicsSimulation.h
@@ -282,7 +282,8 @@ private:
 	void UpdatePostPhysics(
 		FRigPhysicsSolverComponent* SolverComponent, 
 		const float                 Alpha, 
-		const float                 DeltaTime);
+		const float                 DeltaTime,
+		const bool                  bUsingFixedStep);
 
 	void InitSimulationSpace(const FTransform& ComponentToWorld, const FTransform& BoneToComponent);
 
@@ -301,6 +302,9 @@ private:
 	FTransform ConvertComponentSpaceTransformToSimSpace(
 		const FRigPhysicsSolverSettings& SolverSettings, const FTransform& TM) const;
 
+	FVector ConvertComponentSpaceVectorToSimSpace(
+		const FRigPhysicsSolverSettings& SolverSettings, const FVector& V) const;
+
 	// Converts a transform from the simulation space to component space (e.g. for writing back to
 	// the owning control rig)
 	FTransform ConvertSimSpaceTransformToComponentSpace(
diff --git a/Engine/Plugins/Experimental/PhysicsControl/Source/ControlRigPhysics/Private/RigPhysicsSimulationSpace.cpp b/Engine/Plugins/Experimental/PhysicsControl/Source/ControlRigPhysics/Private/RigPhysicsSimulationSpace.cpp
--- a/Engine/Plugins/Experimental/PhysicsControl/Source/ControlRigPhysics/Private/RigPhysicsSimulationSpace.cpp
+++ b/Engine/Plugins/Experimental/PhysicsControl/Source/ControlRigPhysics/Private/RigPhysicsSimulationSpace.cpp
@@ -70,6 +70,19 @@ FTransform FRigPhysicsSimulation::ConvertComponentSpaceTransformToSimSpace(
 	}
 }
 
+//======================================================================================================================
+FVector FRigPhysicsSimulation::ConvertComponentSpaceVectorToSimSpace(
+	const FRigPhysicsSolverSettings& SolverSettings, const FVector& V) const
+{
+	switch (SolverSettings.SimulationSpace)
+	{
+	case ERigPhysicsSimulationSpace::Component: return V;
+	case ERigPhysicsSimulationSpace::World:  return SimulationSpaceState.ComponentTM.TransformVector(V);
+	case ERigPhysicsSimulationSpace::SpaceBone: return SimulationSpaceState.BoneRelComponentTM.InverseTransformVector(V);
+	default: ensureMsgf(false, TEXT("Unsupported Simulation Space")); return V;
+	}
+}
+
 //======================================================================================================================
 FTransform FRigPhysicsSimulation::ConvertSimSpaceTransformToComponentSpace(
 	const FRigPhysicsSolverSettings& SolverSettings, const FTransform& TM) const
diff --git a/Engine/Plugins/Experimental/PhysicsControl/Source/ControlRigPhysics/Public/RigPhysicsData.h b/Engine/Plugins/Experimental/PhysicsControl/Source/ControlRigPhysics/Public/RigPhysicsData.h
--- a/Engine/Plugins/Experimental/PhysicsControl/Source/ControlRigPhysics/Public/RigPhysicsData.h
+++ b/Engine/Plugins/Experimental/PhysicsControl/Source/ControlRigPhysics/Public/RigPhysicsData.h
@@ -345,14 +345,14 @@ struct FRigPhysicsSimulationSpaceSettings
 	// simulation state.
 	// A value of zero will disable this detection
 	UPROPERTY(EditAnywhere, Category = Settings, meta = (ClampMin = "0.0"))
-	float LinearAccelerationThresholdForTeleport = 10000;
+	float LinearAccelerationThresholdForTeleport = 100000;
 
 	// If the angular acceleration (deg/s/s) is above this threshold, the movement is treated as a teleport.
 	// The calculated accelerations and velocities will be set to zero, without resetting the
 	// simulation state.
 	// A value of zero will disable this detection
 	UPROPERTY(EditAnywhere, Category = Settings, meta = (ClampMin = "0.0"))
-	float AngularAccelerationThresholdForTeleport = 10000;
+	float AngularAccelerationThresholdForTeleport = 100000;
 
 	// If the linear position change is above this threshold, the movement is treated as a teleport.
 	// The calculated accelerations and velocities will be set to zero, without resetting the
@@ -547,9 +547,10 @@ struct FRigPhysicsSolverSettings
 
 	// If any kinematic object in the simulation exceeds this acceleration, we will reset the
 	// velocities of all objects in the simulation.This is to detect problems - for example if the
-	// target animation has teleports etc. Disabled when the value is zero.
+	// target animation has teleports etc. Disabled when the value is zero. A value of around 100000
+	// will typically be a good place to start.
 	UPROPERTY(EditAnywhere, Category = Settings, meta = (ClampMin = "0.0"))
-	float KinematicAccelerationThresholdForReset = 40000.0f;
+	float KinematicAccelerationThresholdForReset = 0.0f;
 };
 
 /**
