Minor Control Rig Physics changes/improvements

diff --git a/Engine/Plugins/Experimental/PhysicsControl/Source/ControlRigPhysics/Private/RigPhysicsExecution.cpp b/Engine/Plugins/Experimental/PhysicsControl/Source/ControlRigPhysics/Private/RigPhysicsExecution.cpp
--- a/Engine/Plugins/Experimental/PhysicsControl/Source/ControlRigPhysics/Private/RigPhysicsExecution.cpp
+++ b/Engine/Plugins/Experimental/PhysicsControl/Source/ControlRigPhysics/Private/RigPhysicsExecution.cpp
@@ -312,6 +312,9 @@ static void SetUpJointData(
 	JointData.LinearConstraint = ConstraintProfileProperties.LinearLimit;
 	JointData.ConeConstraint = ConstraintProfileProperties.ConeLimit;
 	JointData.TwistConstraint = ConstraintProfileProperties.TwistLimit;
+
+	JointData.LinearProjectionAmount = ConstraintProfileProperties.ProjectionLinearAlpha;
+	JointData.AngularProjectionAmount = ConstraintProfileProperties.ProjectionAngularAlpha;
 }
 
 //======================================================================================================================
diff --git a/Engine/Plugins/Experimental/PhysicsControl/Source/ControlRigPhysics/Private/RigPhysicsJointExecution.cpp b/Engine/Plugins/Experimental/PhysicsControl/Source/ControlRigPhysics/Private/RigPhysicsJointExecution.cpp
--- a/Engine/Plugins/Experimental/PhysicsControl/Source/ControlRigPhysics/Private/RigPhysicsJointExecution.cpp
+++ b/Engine/Plugins/Experimental/PhysicsControl/Source/ControlRigPhysics/Private/RigPhysicsJointExecution.cpp
@@ -67,6 +67,21 @@ FRigUnit_HierarchySetJointData_Execute()
 	}
 }
 
+//======================================================================================================================
+FRigUnit_HierarchyGetJointData_Execute()
+{
+	DECLARE_SCOPE_HIERARCHICAL_COUNTER_RIGUNIT();
+	if (!ExecuteContext.Hierarchy)
+	{
+		return;
+	}
+	if (FRigPhysicsJointComponent* Component = Cast<FRigPhysicsJointComponent>(
+		ExecuteContext.Hierarchy->FindComponent(PhysicsJointComponentKey)))
+	{
+		JointData = Component->JointData;
+	}
+}
+
 //======================================================================================================================
 FRigUnit_HierarchySetJointDriveData_Execute()
 {
@@ -82,6 +97,21 @@ FRigUnit_HierarchySetJointDriveData_Execute()
 	}
 }
 
+//======================================================================================================================
+FRigUnit_HierarchyGetJointDriveData_Execute()
+{
+	DECLARE_SCOPE_HIERARCHICAL_COUNTER_RIGUNIT();
+	if (!ExecuteContext.Hierarchy)
+	{
+		return;
+	}
+	if (FRigPhysicsJointComponent* Component = Cast<FRigPhysicsJointComponent>(
+		ExecuteContext.Hierarchy->FindComponent(PhysicsJointComponentKey)))
+	{
+		DriveData = Component->DriveData;
+	}
+}
+
 //======================================================================================================================
 FRigUnit_HierarchySetJointDriveUseSkeletalAnimation_Execute()
 {
diff --git a/Engine/Plugins/Experimental/PhysicsControl/Source/ControlRigPhysics/Private/RigPhysicsSimulation.cpp b/Engine/Plugins/Experimental/PhysicsControl/Source/ControlRigPhysics/Private/RigPhysicsSimulation.cpp
--- a/Engine/Plugins/Experimental/PhysicsControl/Source/ControlRigPhysics/Private/RigPhysicsSimulation.cpp
+++ b/Engine/Plugins/Experimental/PhysicsControl/Source/ControlRigPhysics/Private/RigPhysicsSimulation.cpp
@@ -1513,7 +1513,7 @@ void FRigPhysicsSimulation::UpdatePostPhysics(
 
 			// Calculate the target TM even if we're going to reset - it's likely useful for
 			// debugging (and this should be rare!)
-			Record.TargetComponentSpaceTM = ConvertSimSpaceTransformToComponentSpace(
+			Record.FinalComponentSpaceTM = ConvertSimSpaceTransformToComponentSpace(
 				SolverSettings, SimSpaceTM);
 			if (Alpha < 0.999f)
 			{
@@ -1523,11 +1523,11 @@ void FRigPhysicsSimulation::UpdatePostPhysics(
 				// and position.
 				FTransform CurrentTM = Record.SourceComponentSpaceTM.ToTransform();
 				FQuat TargetQ = FQuat::Slerp(
-					CurrentTM.GetRotation(), Record.TargetComponentSpaceTM.GetRotation(), Alpha);
+					CurrentTM.GetRotation(), Record.FinalComponentSpaceTM.GetRotation(), Alpha);
 				FVector TargetT = FMath::Lerp(
-					CurrentTM.GetTranslation(), Record.TargetComponentSpaceTM.GetTranslation(), Alpha);
-				Record.TargetComponentSpaceTM.SetRotation(TargetQ);
-				Record.TargetComponentSpaceTM.SetTranslation(TargetT);
+					CurrentTM.GetTranslation(), Record.FinalComponentSpaceTM.GetTranslation(), Alpha);
+				Record.FinalComponentSpaceTM.SetRotation(TargetQ);
+				Record.FinalComponentSpaceTM.SetTranslation(TargetT);
 			}
 		}
 	}
@@ -1609,7 +1609,7 @@ void FRigPhysicsSimulation::UpdatePostPhysics(
 					// children), so that attached animation bones will follow physics, but we
 					// rely on our bodies being sorted so we work out from the leaf nodes so as
 					// not to disturb previously set bodies.
-					Hierarchy->SetGlobalTransform(Record.TargetElementKey, Record.TargetComponentSpaceTM, false, true);
+					Hierarchy->SetGlobalTransform(Record.TargetElementKey, Record.FinalComponentSpaceTM, false, true);
 				}
 			}
 		}
diff --git a/Engine/Plugins/Experimental/PhysicsControl/Source/ControlRigPhysics/Private/RigPhysicsSimulation.h b/Engine/Plugins/Experimental/PhysicsControl/Source/ControlRigPhysics/Private/RigPhysicsSimulation.h
--- a/Engine/Plugins/Experimental/PhysicsControl/Source/ControlRigPhysics/Private/RigPhysicsSimulation.h
+++ b/Engine/Plugins/Experimental/PhysicsControl/Source/ControlRigPhysics/Private/RigPhysicsSimulation.h
@@ -71,7 +71,7 @@ struct FRigBodyRecord
 
 	// The final/simulated TM is stored before writing it into the output, so we can avoid
 	// corrupting the output if anything is bad and we need to reset.
-	FTransform TargetComponentSpaceTM;
+	FTransform FinalComponentSpaceTM;
 
 	// These are updated for all records in UpdatePrePhysics. The times/validity are determined by
 	// the CurrentDeltaTime, PrevDeltaTime and update counters in the simulation itself.
diff --git a/Engine/Plugins/Experimental/PhysicsControl/Source/ControlRigPhysics/Public/RigPhysicsJointExecution.h b/Engine/Plugins/Experimental/PhysicsControl/Source/ControlRigPhysics/Public/RigPhysicsJointExecution.h
--- a/Engine/Plugins/Experimental/PhysicsControl/Source/ControlRigPhysics/Public/RigPhysicsJointExecution.h
+++ b/Engine/Plugins/Experimental/PhysicsControl/Source/ControlRigPhysics/Public/RigPhysicsJointExecution.h
@@ -80,6 +80,28 @@ struct FRigUnit_HierarchySetJointData : public FRigUnit_PhysicsBaseMutable
 	FRigPhysicsJointData JointData;
 };
 
+// Gets the joint for a physics joint component
+USTRUCT(meta = (DisplayName = "Get Physics Joint Properties"))
+struct FRigUnit_HierarchyGetJointData : public FRigUnit_PhysicsBase
+{
+	GENERATED_BODY()
+
+	FRigUnit_HierarchyGetJointData()
+	{
+		PhysicsJointComponentKey.ElementKey.Type = ERigElementType::Bone;
+		PhysicsJointComponentKey.Name = FRigPhysicsJointComponent::GetDefaultName();
+	}
+
+	RIGVM_METHOD()
+	UE_API virtual void Execute() override;
+
+	UPROPERTY(meta = (Input))
+	FRigComponentKey PhysicsJointComponentKey;
+
+	UPROPERTY(meta = (Output))
+	FRigPhysicsJointData JointData;
+};
+
 // Sets the joint drive for a physics component body
 USTRUCT(meta = (DisplayName = "Set Physics Joint Drive Properties", Varying))
 struct FRigUnit_HierarchySetJointDriveData : public FRigUnit_PhysicsBaseMutable
@@ -102,6 +124,27 @@ struct FRigUnit_HierarchySetJointDriveData : public FRigUnit_PhysicsBaseMutable
 	FRigPhysicsDriveData DriveData;
 };
 
+// Gets the joint drive for a physics joint component
+USTRUCT(meta = (DisplayName = "Get Physics Joint Drive Properties"))
+struct FRigUnit_HierarchyGetJointDriveData : public FRigUnit_PhysicsBase
+{
+	GENERATED_BODY()
+
+	FRigUnit_HierarchyGetJointDriveData()
+	{
+		PhysicsJointComponentKey.ElementKey.Type = ERigElementType::Bone;
+		PhysicsJointComponentKey.Name = FRigPhysicsJointComponent::GetDefaultName();
+	}
+
+	RIGVM_METHOD()
+	UE_API virtual void Execute() override;
+
+	UPROPERTY(meta = (Input))
+	FRigComponentKey PhysicsJointComponentKey;
+
+	UPROPERTY(meta = (Output))
+	FRigPhysicsDriveData DriveData;
+};
 
 // Sets the joint drive for a physics component body
 USTRUCT(meta = (DisplayName = "Set Physics Joint Drive Use Skeletal Animation", Varying))
@@ -126,8 +169,9 @@ struct FRigUnit_HierarchySetJointDriveUseSkeletalAnimation : public FRigUnit_Phy
 };
 
 
-// Helper to simplify creation of joint data
-USTRUCT(meta = (DisplayName = "Make Joint Data", Keywords = "Construction,Create,New,Body,Skeleton", Varying))
+// Helper to simplify creation of joint data for a typical articulation (i.e. locked linear motion,
+// but with an angular limit)
+USTRUCT(meta = (DisplayName = "Make Articulation Joint Data", Keywords = "Construction,Create,New,Body,Skeleton", Varying))
 struct FRigUnit_MakeArticulationJointData : public FRigUnit_PhysicsBase
 {
 	GENERATED_BODY()
@@ -153,7 +197,8 @@ struct FRigUnit_MakeArticulationJointData : public FRigUnit_PhysicsBase
 	FRigPhysicsJointData JointData;
 };
 
-// Helper to simplify creation of articulation drive data
+// Helper to simplify creation of drive data for a typical articulation (i.e. no linear drive, but
+// with an angular drive)
 USTRUCT(meta = (DisplayName = "Make Articulation Drive Data", Keywords = "Construction,Create,New,Body,Skeleton", Varying))
 struct FRigUnit_MakeArticulationDriveData : public FRigUnit_PhysicsBase
 {
