Make the physics control asset editor a little more robust

diff --git a/Engine/Plugins/Experimental/PhysicsControl/Source/PhysicsControl/Public/PhysicsControlAsset.h b/Engine/Plugins/Experimental/PhysicsControl/Source/PhysicsControl/Public/PhysicsControlAsset.h
--- a/Engine/Plugins/Experimental/PhysicsControl/Source/PhysicsControl/Public/PhysicsControlAsset.h
+++ b/Engine/Plugins/Experimental/PhysicsControl/Source/PhysicsControl/Public/PhysicsControlAsset.h
@@ -99,7 +99,9 @@ public:
 	UPROPERTY(EditAnywhere, Category = Inheritance)
 	TArray<TSoftObjectPtr<UPhysicsControlAsset>> AdditionalProfileAssets;
 
-	/** The PhysicsAsset that this control asset is targeting. Can also get the preview mesh from this */
+	// The PhysicsAsset that this control asset is targeting. This will also determine the preview
+	// mesh. You will need to close/re-open the Physics Control Asset editor in order to see the
+	// effect of changing it.
 	UPROPERTY(AssetRegistrySearchable, EditAnywhere, Category = PreviewMesh)
 	TSoftObjectPtr<UPhysicsAsset> PhysicsAsset;
 
diff --git a/Engine/Plugins/Experimental/PhysicsControl/Source/PhysicsControlEditor/Private/PhysicsControlAssetEditor.cpp b/Engine/Plugins/Experimental/PhysicsControl/Source/PhysicsControlEditor/Private/PhysicsControlAssetEditor.cpp
--- a/Engine/Plugins/Experimental/PhysicsControl/Source/PhysicsControlEditor/Private/PhysicsControlAssetEditor.cpp
+++ b/Engine/Plugins/Experimental/PhysicsControl/Source/PhysicsControlEditor/Private/PhysicsControlAssetEditor.cpp
@@ -115,13 +115,14 @@ void FPhysicsControlAssetEditor::InitAssetEditor(
 					MakeShared<FPhysicsControlAssetEditorSkeletonTreeBuilder>(PhysicsAsset);
 			}
 		}
-		SkeletonTreeArgs.ContextName = GetToolkitFName();
 
-		GetMutableDefault<UPersonaOptions>()->bFlattenSkeletonHierarchyWhenFiltering = false;
-		GetMutableDefault<UPersonaOptions>()->bHideParentsWhenFiltering = true;
-		if (PersonaToolkit->GetSkeleton())
+		if (PersonaToolkit->GetSkeleton() && SkeletonTreeArgs.Builder.IsValid())
 		{
-			ISkeletonEditorModule& SkeletonEditorModule = 
+			SkeletonTreeArgs.ContextName = GetToolkitFName();
+			GetMutableDefault<UPersonaOptions>()->bFlattenSkeletonHierarchyWhenFiltering = false;
+			GetMutableDefault<UPersonaOptions>()->bHideParentsWhenFiltering = true;
+
+			ISkeletonEditorModule& SkeletonEditorModule =
 				FModuleManager::LoadModuleChecked<ISkeletonEditorModule>("SkeletonEditor");
 			SkeletonTree = SkeletonEditorModule.CreateSkeletonTree(PersonaToolkit->GetSkeleton(), SkeletonTreeArgs);
 		}
@@ -974,6 +975,11 @@ void FPhysicsControlAssetEditor::HandleSelectionChanged(
 //======================================================================================================================
 void FPhysicsControlAssetEditor::HandleGetFilterLabel(TArray<FText>& InOutItems) const
 {
+	if (!SkeletonTreeBuilder.IsValid())
+	{
+		return;
+	}
+
 	if (SkeletonTreeBuilder->bShowBodies)
 	{
 		InOutItems.Add(LOCTEXT("BodiesFilterLabel", "Bodies"));
