Physics Control support for automatically re-enabling controls if the relevant bodies have been disabled (set to No Collision) and then re-enabled

diff --git a/Engine/Plugins/Experimental/PhysicsControl/Source/PhysicsControl/Private/PhysicsControlComponent.cpp b/Engine/Plugins/Experimental/PhysicsControl/Source/PhysicsControl/Private/PhysicsControlComponent.cpp
--- a/Engine/Plugins/Experimental/PhysicsControl/Source/PhysicsControl/Private/PhysicsControlComponent.cpp
+++ b/Engine/Plugins/Experimental/PhysicsControl/Source/PhysicsControl/Private/PhysicsControlComponent.cpp
@@ -103,8 +103,15 @@ void UPhysicsControlComponent::UpdateControls(float DeltaTime)
 		FPhysicsControlRecord& Record = RecordPair.Value;
 		if (!Record.ConstraintInstance)
 		{
-			Record.InitConstraint(this, ControlName);
+			Record.InitConstraint(this, ControlName, bWarnAboutInvalidNames);
 		}
+		else if (bAttemptToRecreateDisabledControls && !Record.ConstraintInstance->GetPhysicsScene())
+		{
+			// If bodies have been removed from the simulation, then the constraint is removed too,
+			// so attempt to quietly reinitialize the constraint.
+			Record.InitConstraint(this, ControlName, false);
+		}
+
 		ApplyControl(Record);
 	}
 
@@ -736,7 +743,7 @@ bool UPhysicsControlComponent::SetControlParent(
 
 		Record->ParentComponent = ParentComponent;
 		Record->PhysicsControl.ParentBoneName = ParentBoneName;
-		return Record->InitConstraint(this, Name);
+		return Record->InitConstraint(this, Name, bWarnAboutInvalidNames);
 	}
 	if (bWarnAboutInvalidNames)
 	{
diff --git a/Engine/Plugins/Experimental/PhysicsControl/Source/PhysicsControl/Private/PhysicsControlRecord.cpp b/Engine/Plugins/Experimental/PhysicsControl/Source/PhysicsControl/Private/PhysicsControlRecord.cpp
--- a/Engine/Plugins/Experimental/PhysicsControl/Source/PhysicsControl/Private/PhysicsControlRecord.cpp
+++ b/Engine/Plugins/Experimental/PhysicsControl/Source/PhysicsControl/Private/PhysicsControlRecord.cpp
@@ -32,7 +32,8 @@ FVector FPhysicsControlRecord::GetControlPoint() const
 }
 
 //======================================================================================================================
-bool FPhysicsControlRecord::InitConstraint(UObject* ConstraintDebugOwner, FName ControlName)
+bool FPhysicsControlRecord::InitConstraint(
+	UObject* ConstraintDebugOwner, FName ControlName, bool bWarnAboutInvalidNames)
 {
 	if (!ConstraintInstance.IsValid())
 	{
@@ -47,18 +48,24 @@ bool FPhysicsControlRecord::InitConstraint(UObject* ConstraintDebugOwner, FName
 
 	if (ParentComponent.IsValid() && !PhysicsControl.ParentBoneName.IsNone() && !ParentBody)
 	{
-		UE_LOG(LogPhysicsControl, Warning,
-			TEXT("Failed to find expected parent body %s when making constraint for control %s"),
-			*PhysicsControl.ParentBoneName.ToString(),
-			*ControlName.ToString());
+		if (bWarnAboutInvalidNames)
+		{
+			UE_LOG(LogPhysicsControl, Warning,
+				TEXT("Failed to find expected parent body %s when making constraint for control %s"),
+				*PhysicsControl.ParentBoneName.ToString(),
+				*ControlName.ToString());
+		}
 		return false;
 	}
 	if (ChildComponent.IsValid() && !PhysicsControl.ChildBoneName.IsNone() && !ChildBody)
 	{
-		UE_LOG(LogPhysicsControl, Warning,
-			TEXT("Failed to find expected child body %s when making constraint for control %s"),
-			*PhysicsControl.ChildBoneName.ToString(),
-			*ControlName.ToString());
+		if (bWarnAboutInvalidNames)
+		{
+			UE_LOG(LogPhysicsControl, Warning,
+				TEXT("Failed to find expected child body %s when making constraint for control %s"),
+				*PhysicsControl.ChildBoneName.ToString(),
+				*ControlName.ToString());
+		}
 		return false;
 	}
 
@@ -78,7 +85,6 @@ bool FPhysicsControlRecord::InitConstraint(UObject* ConstraintDebugOwner, FName
 	ConstraintInstance->SetLinearPositionDrive(true, true, true);
 	ConstraintInstance->SetLinearVelocityDrive(true, true, true);
 
-	// Ensure the control point is set
 	UpdateConstraintControlPoint();
 
 	return true;
diff --git a/Engine/Plugins/Experimental/PhysicsControl/Source/PhysicsControl/Public/PhysicsControlComponent.h b/Engine/Plugins/Experimental/PhysicsControl/Source/PhysicsControl/Public/PhysicsControlComponent.h
--- a/Engine/Plugins/Experimental/PhysicsControl/Source/PhysicsControl/Public/PhysicsControlComponent.h
+++ b/Engine/Plugins/Experimental/PhysicsControl/Source/PhysicsControl/Public/PhysicsControlComponent.h
@@ -1663,23 +1663,26 @@ public:
 	float VelocityPredictionTime = 0.2f;
 
 	/**
-	 * Upper limit on the number of controls or modifiers that will be created using the same name (which
-	 * will get a numerical postfix). When this limit is reached a warning will be issued  and the control 
-	 * or modifier won't be created. This is to avoid problems if controls or modifiers are being created 
-	 * dynamically, and can generally be a "moderately large" number, depending on how many controls or 
-	 * modifiers you expect to create.
-	 */
-	UPROPERTY(EditAnywhere, BlueprintReadWrite, Category = Settings)
-	int32 MaxNumControlsOrModifiersPerName = 256;
-
-	/** 
-	 * Warn if an an invalid control or body modifier name is used. This can happen quite easily since 
+	 * Warn if an an invalid control or body modifier name is used. This can happen quite easily since
 	 * they're only referenced through names, which are likely auto-generated. However, it may happen for
 	 * valid reasons, in which case you'll want to disable this.
 	 */
 	UPROPERTY(EditAnywhere, BlueprintReadWrite, Category = Debug)
 	bool bWarnAboutInvalidNames = true;
 
+	// Upper limit on the number of controls or modifiers that will be created using the same name (which
+	// will get a numerical postfix). When this limit is reached a warning will be issued  and the control 
+	// or modifier won't be created. This is to avoid problems if controls or modifiers are being created 
+	// dynamically, and can generally be a "moderately large" number, depending on how many controls or 
+	// modifiers you expect to create.
+	UPROPERTY(EditAnywhere, BlueprintReadWrite, Category = Settings)
+	int32 MaxNumControlsOrModifiersPerName = 256;
+
+	// If set then controls that have been disabled due to the bodies being set to NoCollision (for
+	// example) will be reinitialized, if possible, when those bodies are returned to the simulation.
+	UPROPERTY(EditAnywhere, BlueprintReadWrite, Category = Settings)
+	bool bAttemptToRecreateDisabledControls = true;
+
 public:
 	//Begin ActorComponent interface
 	UE_API virtual void OnRegister() override;
diff --git a/Engine/Plugins/Experimental/PhysicsControl/Source/PhysicsControl/Public/PhysicsControlRecord.h b/Engine/Plugins/Experimental/PhysicsControl/Source/PhysicsControl/Public/PhysicsControlRecord.h
--- a/Engine/Plugins/Experimental/PhysicsControl/Source/PhysicsControl/Public/PhysicsControlRecord.h
+++ b/Engine/Plugins/Experimental/PhysicsControl/Source/PhysicsControl/Public/PhysicsControlRecord.h
@@ -11,9 +11,7 @@ struct FConstraintInstance;
 struct FBodyInstance;
 class UPrimitiveComponent;
 
-/**
- * There will be a PhysicsControlRecord created at runtime for every Control that has been created
- */
+// There will be a PhysicsControlRecord created at runtime for every Control that has been created
 struct FPhysicsControlRecord
 {
 	FPhysicsControlRecord(
@@ -27,31 +25,28 @@ struct FPhysicsControlRecord
 		, ChildComponent(InChildComponent)
 	{}
 
-	/** Removes any constraint and resets the state */
+	// Removes any constraint and resets the state
 	void ResetConstraint();
 
-	/** Returns the control point, which may be custom or automatic (centre of mass) */
+	// Returns the control point, which may be custom or automatic (centre of mass)
 	FVector GetControlPoint() const;
 
-	/**
-	 * Creates the constraint if necessary and stores it. Then initializes the constraint with the bodies.
-	 * Returns true/false on success/failure */
-	bool InitConstraint(UObject* ConstraintDebugOwner, FName ControlName);
+	// Creates the constraint if necessary and stores it. Then initializes the constraint with the
+	// bodies. 
+	bool InitConstraint(UObject* ConstraintDebugOwner, FName ControlName, bool bWarnAboutInvalidNames);
 
-	/** Ensures the constraint frame matches the control point in the record. */
+	// Ensures the constraint frame matches the control point in the record.
 	void UpdateConstraintControlPoint();
 
-	/** Sets the control point to the center of mass of the child mesh (or to zero if that fails). */
+	// Sets the control point to the center of mass of the child mesh (or to zero if that fails).
 	void ResetControlPoint();
 
-	/** The configuration data */
+	// The configuration data
 	FPhysicsControl PhysicsControl;
 
-	/**
-	 * The position/orientation etc targets for the controls. These are procedural/explicit control targets -
-	 * skeletal meshes have the option to use skeletal animation as well, in which case these targets are 
-	 * expressed as relative to that animation.
-	 */
+	// The position/orientation etc targets for the controls. These are procedural/explicit control targets -
+	// skeletal meshes have the option to use skeletal animation as well, in which case these targets are 
+	// expressed as relative to that animation.
 	FPhysicsControlTarget ControlTarget;
 
 	// The previous skeletal control target. This will have been set at the end of a previous update
@@ -64,19 +59,17 @@ struct FPhysicsControlRecord
 	// expected update counter will be set when the control/previous target TM has just been updated.
 	FGraphTraversalCounter ExpectedUpdateCounter;
 
-	/**  The mesh that will be doing the driving. Blank/non-existent means it will happen in world space */
+	// The mesh that will be doing the driving. Blank/non-existent means it will happen in world space
 	TWeakObjectPtr<UPrimitiveComponent> ParentComponent;
 
-	/** The mesh that the control will be driving. */
+	// The mesh that the control will be driving.
 	TWeakObjectPtr<UPrimitiveComponent> ChildComponent;
 
-	/** The underlying constraint used to implement the control. */
+	// The underlying constraint used to implement the control.
 	TSharedPtr<FConstraintInstance> ConstraintInstance;
 };
 
-/**
- * There will be a PhysicsBodyModifier created at runtime for every BodyInstance involved in the component
- */
+// There will be a PhysicsBodyModifier created at runtime for every BodyInstance involved in the component
 struct FPhysicsBodyModifierRecord
 {
 	FPhysicsBodyModifierRecord(
@@ -90,35 +83,27 @@ struct FPhysicsBodyModifierRecord
 		, bResetToCachedTarget(false)
 	{}
 
-	/**  The mesh that will be modified. */
+	// The mesh that will be modified.
 	TWeakObjectPtr<UPrimitiveComponent> Component;
 
 	// The core data
 	FPhysicsBodyModifier BodyModifier;
 
-	/** 
-	 * The target position when kinematic. Note that this is applied on top of any animation 
-	 * target if bUseSkeletalAnimation is set. 
-	 */
+	// The target position when kinematic. Note that this is applied on top of any animation 
+	// target if bUseSkeletalAnimation is set. 
 	FVector KinematicTargetPosition = FVector::ZeroVector;
 
-	/**
-	 * The target orientation when kinematic. Note that this is applied on top of any animation
-	 * target if bUseSkeletalAnimation is set.
-	 */
+	// The target orientation when kinematic. Note that this is applied on top of any animation
+	// target if bUseSkeletalAnimation is set.
 	FQuat KinematicTargetOrientation = FQuat::Identity;
 
-	/** 
-	 * If true then the body will be set to the transform/velocity stored in any cached target (if that
-	 * exists), and then this flag will be cleared.
-	 */
+	// If true then the body will be set to the transform/velocity stored in any cached target (if that
+	// exists), and then this flag will be cleared.
 	uint8 bResetToCachedTarget:1;
 };
 
-/**
- * Used internally/only at runtime to track when a SkeletalMeshComponent is being controlled through
- * a modifier, and to restore settings when that stops.
- */
+// Used internally/only at runtime to track when a SkeletalMeshComponent is being controlled through
+// a modifier, and to restore settings when that stops.
 struct FModifiedSkeletalMeshData
 {
 public: 
@@ -126,15 +111,13 @@ public:
 
 public:
 
-	/** The original setting for restoration when we're deleted */
+	// The original setting for restoration when we're deleted
 	uint8 bOriginalUpdateMeshWhenKinematic : 1;
 
-	/** The original setting for restoration when we're deleted */
+	// The original setting for restoration when we're deleted
 	EKinematicBonesUpdateToPhysics::Type OriginalKinematicBonesUpdateType;
 
-	/**
-	 * Track when skeletal meshes are going to be used so this entry can be removed
-	 */
+	// Track when skeletal meshes are going to be used so this entry can be removed
 	int32 ReferenceCount;
 };
 
