Physics Control: This fixes duplication of the Linear/AngularTargetVelocityMultiplier and TargetVelocityMultiplier.

diff --git a/Engine/Plugins/Experimental/PhysicsControl/Source/PhysicsControl/Private/PhysicsControlComponent.cpp b/Engine/Plugins/Experimental/PhysicsControl/Source/PhysicsControl/Private/PhysicsControlComponent.cpp
--- a/Engine/Plugins/Experimental/PhysicsControl/Source/PhysicsControl/Private/PhysicsControlComponent.cpp
+++ b/Engine/Plugins/Experimental/PhysicsControl/Source/PhysicsControl/Private/PhysicsControlComponent.cpp
@@ -337,7 +337,8 @@ TArray<FName> UPhysicsControlComponent::CreateControlsFromSkeletalMeshAndConstra
 				FPhysicsControlData ControlData;
 				// This is to match the skeletal mesh component velocity drive, which does not use the
 				// target animation velocity.
-				ControlData.SkeletalAnimationVelocityMultiplier = 0;
+				ControlData.AngularTargetVelocityMultiplier = 0;
+				ControlData.LinearTargetVelocityMultiplier = 0;
 				FConstraintProfileProperties ProfileProperties;
 				if (!SkeletalMeshComponent->GetConstraintProfilePropertiesOrDefault(
 					ProfileProperties, ChildBoneName, ConstraintProfile))
@@ -450,7 +451,8 @@ TArray<FName> UPhysicsControlComponent::CreateControlsFromSkeletalMeshAndConstra
 		FPhysicsControlData ControlData;
 		// This is to match the skeletal mesh component velocity drive, which does not use the
 		// target animation velocity.
-		ControlData.SkeletalAnimationVelocityMultiplier = 0;
+		ControlData.AngularTargetVelocityMultiplier = 0;
+		ControlData.LinearTargetVelocityMultiplier = 0;
 		FConstraintProfileProperties ProfileProperties;
 		if (!SkeletalMeshComponent->GetConstraintProfilePropertiesOrDefault(
 			ProfileProperties, ChildBoneName, ConstraintProfile))
@@ -624,7 +626,8 @@ TMap<FName, FPhysicsControlNames> UPhysicsControlComponent::CreateControlsFromLi
 			FPhysicsControlData ControlData;
 			// This is to match the skeletal mesh component velocity drive, which does not use the
 			// target animation velocity.
-			ControlData.SkeletalAnimationVelocityMultiplier = 0;
+			ControlData.AngularTargetVelocityMultiplier = 0;
+			ControlData.LinearTargetVelocityMultiplier = 0;
 
 			FConstraintProfileProperties ProfileProperties;
 			if (!SkeletalMeshComponent->GetConstraintProfilePropertiesOrDefault(
@@ -1400,13 +1403,15 @@ bool UPhysicsControlComponent::SetControlTargetPoses(
 bool UPhysicsControlComponent::SetControlUseSkeletalAnimation(
 	const FName Name,
 	const bool  bUseSkeletalAnimation,
-	const float SkeletalAnimationVelocityMultiplier)
+	const float AngularTargetVelocityMultiplier,
+	const float LinearTargetVelocityMultiplier)
 {
 	FPhysicsControl* PhysicsControl = FindControl(Name);
 	if (PhysicsControl)
 	{
 		PhysicsControl->ControlData.bUseSkeletalAnimation = bUseSkeletalAnimation;
-		PhysicsControl->ControlData.SkeletalAnimationVelocityMultiplier = SkeletalAnimationVelocityMultiplier;
+		PhysicsControl->ControlData.AngularTargetVelocityMultiplier = AngularTargetVelocityMultiplier;
+		PhysicsControl->ControlData.LinearTargetVelocityMultiplier = LinearTargetVelocityMultiplier;
 		return true;
 	}
 	if (bWarnAboutInvalidNames)
@@ -1421,11 +1426,13 @@ bool UPhysicsControlComponent::SetControlUseSkeletalAnimation(
 void UPhysicsControlComponent::SetControlsUseSkeletalAnimation(
 	const TArray<FName>& Names,
 	const bool           bUseSkeletalAnimation,
-	const float          SkeletalAnimationVelocityMultiplier)
+	const float          AngularTargetVelocityMultiplier,
+	const float          LinearTargetVelocityMultiplier)
 {
 	for (FName Name : Names)
 	{
-		SetControlUseSkeletalAnimation(Name, bUseSkeletalAnimation, SkeletalAnimationVelocityMultiplier);
+		SetControlUseSkeletalAnimation(
+			Name, bUseSkeletalAnimation, AngularTargetVelocityMultiplier, LinearTargetVelocityMultiplier);
 	}
 }
 
@@ -1433,10 +1440,12 @@ void UPhysicsControlComponent::SetControlsUseSkeletalAnimation(
 void UPhysicsControlComponent::SetControlsInSetUseSkeletalAnimation(
 	const FName SetName,
 	const bool  bUseSkeletalAnimation,
-	const float SkeletalAnimationVelocityMultiplier)
+	const float AngularTargetVelocityMultiplier,
+	const float LinearTargetVelocityMultiplier)
 {
 	SetControlsUseSkeletalAnimation(
-		GetControlNamesInSet(SetName), bUseSkeletalAnimation, SkeletalAnimationVelocityMultiplier);
+		GetControlNamesInSet(SetName), 
+		bUseSkeletalAnimation, AngularTargetVelocityMultiplier, LinearTargetVelocityMultiplier);
 }
 
 //======================================================================================================================
diff --git a/Engine/Plugins/Experimental/PhysicsControl/Source/PhysicsControl/Private/PhysicsControlData.cpp b/Engine/Plugins/Experimental/PhysicsControl/Source/PhysicsControl/Private/PhysicsControlData.cpp
--- a/Engine/Plugins/Experimental/PhysicsControl/Source/PhysicsControl/Private/PhysicsControlData.cpp
+++ b/Engine/Plugins/Experimental/PhysicsControl/Source/PhysicsControl/Private/PhysicsControlData.cpp
@@ -1,6 +1,7 @@
 // Copyright Epic Games, Inc. All Rights Reserved.
 
 #include "PhysicsControlData.h"
+#include "PhysicsControlObjectVersion.h"
 #include "Physics/ImmediatePhysics/ImmediatePhysicsActorHandle.h"
 
 #include UE_INLINE_GENERATED_CPP_BY_NAME(PhysicsControlData)
@@ -62,7 +63,6 @@ FPhysicsControlData Interpolate(
 	INTERPOLATE_PARAM(AngularExtraDamping);
 	INTERPOLATE_PARAM(LinearTargetVelocityMultiplier);
 	INTERPOLATE_PARAM(AngularTargetVelocityMultiplier);
-	INTERPOLATE_PARAM(SkeletalAnimationVelocityMultiplier);
 	INTERPOLATE_PARAM(CustomControlPoint);
 	SELECT_PARAM(bEnabled);
 	SELECT_PARAM(bUseCustomControlPoint);
@@ -85,7 +85,6 @@ FPhysicsControlSparseData Interpolate(
 	INTERPOLATE_PARAM(AngularExtraDamping);
 	INTERPOLATE_PARAM(LinearTargetVelocityMultiplier);
 	INTERPOLATE_PARAM(AngularTargetVelocityMultiplier);
-	INTERPOLATE_PARAM(SkeletalAnimationVelocityMultiplier);
 	INTERPOLATE_PARAM(CustomControlPoint);
 	SELECT_PARAM(bEnabled);
 	SELECT_PARAM(bUseCustomControlPoint);
@@ -101,7 +100,6 @@ FPhysicsControlSparseData Interpolate(
 	SET_ENABLED_PARAM(AngularExtraDamping);
 	SET_ENABLED_PARAM(LinearTargetVelocityMultiplier);
 	SET_ENABLED_PARAM(AngularTargetVelocityMultiplier);
-	SET_ENABLED_PARAM(SkeletalAnimationVelocityMultiplier);
 	SET_ENABLED_PARAM(CustomControlPoint);
 	SET_ENABLED_PARAM(bEnabled);
 	SET_ENABLED_PARAM(bUseCustomControlPoint);
@@ -160,7 +158,6 @@ void FPhysicsControlData::UpdateFromSparseData(const FPhysicsControlSparseData&
 	SET_SPARSE_DATA(MaxTorque);
 	SET_SPARSE_DATA(LinearTargetVelocityMultiplier);
 	SET_SPARSE_DATA(AngularTargetVelocityMultiplier);
-	SET_SPARSE_DATA(SkeletalAnimationVelocityMultiplier);
 	SET_SPARSE_DATA(CustomControlPoint);
 	SET_SPARSE_DATA(bEnabled);
 	SET_SPARSE_DATA(bUseCustomControlPoint);
@@ -193,6 +190,34 @@ void FPhysicsControlModifierData::UpdateFromSparseData(const FPhysicsControlModi
 	SET_SPARSE_DATA(bUpdateKinematicFromSimulation);
 }
 
+//======================================================================================================================
+FArchive& operator<<(FArchive& Ar, FPhysicsControlData& Data)
+{
+	Ar << Data.LinearStrength;
+	Ar << Data.LinearDampingRatio;
+	Ar << Data.LinearExtraDamping;
+	Ar << Data.MaxForce;
+	Ar << Data.AngularStrength;
+	Ar << Data.AngularDampingRatio;
+	Ar << Data.AngularExtraDamping;
+	Ar << Data.MaxTorque;
+	Ar << Data.LinearTargetVelocityMultiplier;
+	Ar << Data.AngularTargetVelocityMultiplier;
+	if (Ar.CustomVer(FPhysicsControlObjectVersion::GUID) < FPhysicsControlObjectVersion::RemoveSkeletalAnimationVelocityMultiplier)
+	{
+		float SkeletalAnimationVelocityMultiplier = 0.0f;
+		Ar << SkeletalAnimationVelocityMultiplier;
+	}
+	Ar << Data.CustomControlPoint;
+	Ar << Data.bEnabled;
+	Ar << Data.bUseCustomControlPoint;
+	Ar << Data.bUseSkeletalAnimation;
+	Ar << Data.bDisableCollision;
+	Ar << Data.bOnlyControlChildObject;
+	return Ar;
+}
+
+
 #undef SET_SPARSE_DATA
 #undef INTERPOLATE_PARAM
 #undef SELECT_PARAM
diff --git a/Engine/Plugins/Experimental/PhysicsControl/Source/PhysicsControl/Public/PhysicsControlComponent.h b/Engine/Plugins/Experimental/PhysicsControl/Source/PhysicsControl/Public/PhysicsControlComponent.h
--- a/Engine/Plugins/Experimental/PhysicsControl/Source/PhysicsControl/Public/PhysicsControlComponent.h
+++ b/Engine/Plugins/Experimental/PhysicsControl/Source/PhysicsControl/Public/PhysicsControlComponent.h
@@ -807,22 +807,24 @@ public:
 		const float    VelocityDeltaTime, 
 		const bool     bEnableControl = true);
 
-
 	/**
 	 * Sets whether or not the control should use skeletal animation for the targets
 	 *
 	 * @param Name The name of the control to modify. 
 	 * @param bUseSkeletalAnimation If true then the targets will be a combination of the skeletal animation (if
 	 *        there is any) and the control target that has been set
-	 * @param SkeletalAnimationVelocityMultiplier If skeletal animation is being used, then this determines 
-	 *        the amount of velocity extracted from the animation that is used as targets for the controls
+	 * @param AngularTargetVelocityMultiplier If skeletal animation is being used, then this determines the amount of
+	 *        angular velocity extracted from the animation that is used as targets for the controls
+	 * @param AngularTargetVelocityMultiplier If skeletal animation is being used, then this determines the amount of
+	 *        linear velocity extracted from the animation that is used as targets for the controls
 	 * @return true if the control was found and modified, false if not
 	 */
 	UFUNCTION(BlueprintCallable, Category = PhysicsControl, Meta = (ReturnDisplayName = "Success"))
 	UE_API bool SetControlUseSkeletalAnimation(
 		const FName Name,
 		const bool  bUseSkeletalAnimation = true,
-		const float SkeletalAnimationVelocityMultiplier = 1.0f);
+		const float AngularTargetVelocityMultiplier = 1.0f,
+		const float LinearTargetVelocityMultiplier = 1.0f);
 
 	/**
 	 * Sets whether or not the controls should use skeletal animation for the targets
@@ -831,14 +833,17 @@ public:
 	 *              then it can be split.
 	 * @param bUseSkeletalAnimation If true then the targets will be a combination of the skeletal animation (if
 	 *              there is any) and the control target that has been set
-	 * @param SkeletalAnimationVelocityMultiplier If skeletal animation is being used, then this determines the amount of 
-	 *              velocity extracted from the animation that is used as targets for the controls
+	 * @param AngularTargetVelocityMultiplier If skeletal animation is being used, then this determines the amount of 
+	 *              angular velocity extracted from the animation that is used as targets for the controls
+	 * @param LinearTargetVelocityMultiplier If skeletal animation is being used, then this determines the amount of
+	 *              linear velocity extracted from the animation that is used as targets for the controls
 	 */
 	UFUNCTION(BlueprintCallable, Category = PhysicsControl)
 	UE_API void SetControlsUseSkeletalAnimation(
 		const TArray<FName>& Names,
 		const bool           bUseSkeletalAnimation = true,
-		const float          SkeletalAnimationVelocityMultiplier = 1.0f);
+		const float          AngularTargetVelocityMultiplier = 1.0f,
+		const float          LinearTargetVelocityMultiplier = 1.0f);
 
 	/**
 	 * Sets whether or not the controls should use skeletal animation for the targets
@@ -847,14 +852,17 @@ public:
 	 *        "ParentSpace" and things like "WorldSpace-ArmLeft", depending on how controls have been created.
 	 * @param bUseSkeletalAnimation If true then the targets will be a combination of the skeletal animation (if
 	 *        there is any) and the control target that has been set
-	 * @param SkeletalAnimationVelocityMultiplier If skeletal animation is being used, then this determines the amount of
-	 *        velocity extracted from the animation that is used as targets for the controls
+	 * @param AngularTargetVelocityMultiplier If skeletal animation is being used, then this determines the amount of
+	 *        angular velocity extracted from the animation that is used as targets for the controls
+	 * @param AngularTargetVelocityMultiplier If skeletal animation is being used, then this determines the amount of
+	 *        linear velocity extracted from the animation that is used as targets for the controls
 	 */
 	UFUNCTION(BlueprintCallable, Category = PhysicsControl)
 	UE_API void SetControlsInSetUseSkeletalAnimation(
 		const FName Set,
 		const bool  bUseSkeletalAnimation = true,
-		const float SkeletalAnimationVelocityMultiplier = 1.0f);
+		const float AngularTargetVelocityMultiplier = 1.0f,
+		const float LinearTargetVelocityMultiplier = 1.0f);
 
 	/**
 	 * Activates or deactivates a control
diff --git a/Engine/Plugins/Experimental/PhysicsControl/Source/PhysicsControl/Public/PhysicsControlData.h b/Engine/Plugins/Experimental/PhysicsControl/Source/PhysicsControl/Public/PhysicsControlData.h
--- a/Engine/Plugins/Experimental/PhysicsControl/Source/PhysicsControl/Public/PhysicsControlData.h
+++ b/Engine/Plugins/Experimental/PhysicsControl/Source/PhysicsControl/Public/PhysicsControlData.h
@@ -349,7 +349,6 @@ struct FPhysicsControlSparseData
 		, MaxTorque(0)
 		, LinearTargetVelocityMultiplier(1)
 		, AngularTargetVelocityMultiplier(1)
-		, SkeletalAnimationVelocityMultiplier(1)
 		, CustomControlPoint(FVector::ZeroVector)
 		, bEnabled(true)
 		, bUseCustomControlPoint(false)
@@ -367,7 +366,6 @@ struct FPhysicsControlSparseData
 		, bEnableMaxTorque(true)
 		, bEnableLinearTargetVelocityMultiplier(true)
 		, bEnableAngularTargetVelocityMultiplier(true)
-		, bEnableSkeletalAnimationVelocityMultiplier(true)
 		, bEnableCustomControlPoint(true)
 		, bEnablebEnabled(true)
 		, bEnablebUseCustomControlPoint(true)
@@ -449,13 +447,6 @@ struct FPhysicsControlSparseData
 	UPROPERTY(EditAnywhere, BlueprintReadWrite, Category = PhysicsControl, meta = (ClampMin = "0.0", editcondition = "bEnableAngularTargetVelocityMultiplier"))
 	float AngularTargetVelocityMultiplier;
 
-	/** 
-	 * The amount of skeletal animation velocity to use in the targets 
-	 * NOTE This is currently only used in the PhysicsControlComponent implementation
-	 */
-	UPROPERTY(EditAnywhere, BlueprintReadWrite, Category = PhysicsControl, meta = (ClampMin = "0.0", editcondition = "bEnableSkeletalAnimationVelocityMultiplier"))
-	float SkeletalAnimationVelocityMultiplier;
-
 	/**
 	 * The position of the control point relative to the child mesh, when using a custom control point. 
 	 */
@@ -519,9 +510,6 @@ struct FPhysicsControlSparseData
 	UPROPERTY(EditAnywhere, Category = PhysicsControl, meta = (InlineEditConditionToggle))
 	uint8 bEnableAngularTargetVelocityMultiplier : 1;
 
-	UPROPERTY(EditAnywhere, Category = PhysicsControl, meta = (InlineEditConditionToggle))
-	uint8 bEnableSkeletalAnimationVelocityMultiplier : 1;
-
 	UPROPERTY(EditAnywhere, Category = PhysicsControl, meta = (InlineEditConditionToggle))
 	uint8 bEnableCustomControlPoint : 1;
 
@@ -562,7 +550,6 @@ struct FPhysicsControlData
 		, MaxTorque(0.0f)
 		, LinearTargetVelocityMultiplier(1.0f)
 		, AngularTargetVelocityMultiplier(1.0f)
-		, SkeletalAnimationVelocityMultiplier(1.0f)
 		, CustomControlPoint(FVector::ZeroVector)
 		, bUseCustomControlPoint(false)
 		, bUseSkeletalAnimation(true)
@@ -573,27 +560,7 @@ struct FPhysicsControlData
 
 	bool operator==(const FPhysicsControlData& other) const = default;
 
-	friend FArchive& operator<<(FArchive& Ar, FPhysicsControlData& Data)
-	{
-		Ar << Data.LinearStrength;
-		Ar << Data.LinearDampingRatio;
-		Ar << Data.LinearExtraDamping;
-		Ar << Data.MaxForce;
-		Ar << Data.AngularStrength;
-		Ar << Data.AngularDampingRatio;
-		Ar << Data.AngularExtraDamping;
-		Ar << Data.MaxTorque;
-		Ar << Data.LinearTargetVelocityMultiplier;
-		Ar << Data.AngularTargetVelocityMultiplier;
-		Ar << Data.SkeletalAnimationVelocityMultiplier;
-		Ar << Data.CustomControlPoint;
-		Ar << Data.bEnabled;
-		Ar << Data.bUseCustomControlPoint;
-		Ar << Data.bUseSkeletalAnimation;
-		Ar << Data.bDisableCollision;
-		Ar << Data.bOnlyControlChildObject;
-		return Ar;
-	}
+	friend UE_API FArchive& operator<<(FArchive& Ar, FPhysicsControlData& Data);
 
 	/** Applies the values that have been flagged as enabled from the sparse data */
 	UE_API void UpdateFromSparseData(const FPhysicsControlSparseData& SparseData);
@@ -652,7 +619,6 @@ struct FPhysicsControlData
 	 * Multiplier on the velocity, which gets applied to the damping. A value of 1 means the animation target
 	 * velocity is used, which helps it track the animation. A value of 0 means damping happens in "world space" 
 	 * - so damping acts like drag on the movement.
-	 * NOTE This is currently only used in the RigidBodyWithControl implementation
 	 */
 	UPROPERTY(EditAnywhere, BlueprintReadWrite, Category = PhysicsControl, meta = (ClampMin = "0.0"))
 	float LinearTargetVelocityMultiplier;
@@ -661,18 +627,10 @@ struct FPhysicsControlData
 	 * Multiplier on the angular velocity, which gets applied to the damping. A value of 1 means the animation target
 	 * velocity is used, which helps it track the animation. A value of 0 means damping happens in "world space" 
 	 * - so damping acts like drag on the movement.
-	 * NOTE This is currently only used in the RigidBodyWithControl implementation
 	 */
 	UPROPERTY(EditAnywhere, BlueprintReadWrite, Category = PhysicsControl, meta = (ClampMin = "0.0"))
 	float AngularTargetVelocityMultiplier;
 
-	/** 
-	 * The amount of skeletal animation velocity to use in the targets 
-	 * NOTE This is currently only used in the PhysicsControlComponent implementation
-	 */
-	UPROPERTY(EditAnywhere, BlueprintReadWrite, Category = PhysicsControl)
-	float SkeletalAnimationVelocityMultiplier;
-
 	/** The position of the control point relative to the child mesh, when using a custom control point. */
 	UPROPERTY(EditAnywhere, BlueprintReadWrite, Category = PhysicsControl)
 	FVector CustomControlPoint;
diff --git a/Engine/Plugins/Experimental/PhysicsControl/Source/PhysicsControl/Public/PhysicsControlObjectVersion.h b/Engine/Plugins/Experimental/PhysicsControl/Source/PhysicsControl/Public/PhysicsControlObjectVersion.h
--- a/Engine/Plugins/Experimental/PhysicsControl/Source/PhysicsControl/Public/PhysicsControlObjectVersion.h
+++ b/Engine/Plugins/Experimental/PhysicsControl/Source/PhysicsControl/Public/PhysicsControlObjectVersion.h
@@ -37,6 +37,7 @@ struct FPhysicsControlObjectVersion
 		ControlRigWorldCollision,
 		ControlRigCentreOfMassNudge,
 		ControlRigJointEnabled,
+		RemoveSkeletalAnimationVelocityMultiplier,
 
 		// -----<new versions can be added above this line>-------------------------------------------------
 		VersionPlusOne,
