PhysicsControl - fix docs (comments). No code change.

diff --git a/Engine/Plugins/Experimental/PhysicsControl/Source/PhysicsControl/Public/PhysicsControlComponent.h b/Engine/Plugins/Experimental/PhysicsControl/Source/PhysicsControl/Public/PhysicsControlComponent.h
--- a/Engine/Plugins/Experimental/PhysicsControl/Source/PhysicsControl/Public/PhysicsControlComponent.h
+++ b/Engine/Plugins/Experimental/PhysicsControl/Source/PhysicsControl/Public/PhysicsControlComponent.h
@@ -815,7 +815,7 @@ public:
 	 *        there is any) and the control target that has been set
 	 * @param AngularTargetVelocityMultiplier If skeletal animation is being used, then this determines the amount of
 	 *        angular velocity extracted from the animation that is used as targets for the controls
-	 * @param AngularTargetVelocityMultiplier If skeletal animation is being used, then this determines the amount of
+	 * @param LinearTargetVelocityMultiplier If skeletal animation is being used, then this determines the amount of
 	 *        linear velocity extracted from the animation that is used as targets for the controls
 	 * @return true if the control was found and modified, false if not
 	 */
@@ -854,7 +854,7 @@ public:
 	 *        there is any) and the control target that has been set
 	 * @param AngularTargetVelocityMultiplier If skeletal animation is being used, then this determines the amount of
 	 *        angular velocity extracted from the animation that is used as targets for the controls
-	 * @param AngularTargetVelocityMultiplier If skeletal animation is being used, then this determines the amount of
+	 * @param LinearTargetVelocityMultiplier If skeletal animation is being used, then this determines the amount of
 	 *        linear velocity extracted from the animation that is used as targets for the controls
 	 */
 	UFUNCTION(BlueprintCallable, Category = PhysicsControl)
diff --git a/Engine/Plugins/Experimental/PhysicsControl/Source/PhysicsControl/Public/PhysicsControlData.h b/Engine/Plugins/Experimental/PhysicsControl/Source/PhysicsControl/Public/PhysicsControlData.h
--- a/Engine/Plugins/Experimental/PhysicsControl/Source/PhysicsControl/Public/PhysicsControlData.h
+++ b/Engine/Plugins/Experimental/PhysicsControl/Source/PhysicsControl/Public/PhysicsControlData.h
@@ -430,7 +430,7 @@ struct FPhysicsControlSparseData
 	float MaxTorque;
 
 	/**
-	 * Multiplier on the velocity, which gets applied to the damping. A value of 1 means the animation target
+	 * Multiplier on the linear velocity, which gets applied to the damping. A value of 1 means the animation target
 	 * velocity is used, which helps it track the animation. A value of 0 means damping happens in "world space" 
 	 * - so damping acts like drag on the movement.
 	 * NOTE This is currently only used in the RigidBodyWithControl implementation
@@ -616,7 +616,7 @@ struct FPhysicsControlData
 	float MaxTorque;
 
 	/**
-	 * Multiplier on the velocity, which gets applied to the damping. A value of 1 means the animation target
+	 * Multiplier on the linear velocity, which gets applied to the damping. A value of 1 means the animation target
 	 * velocity is used, which helps it track the animation. A value of 0 means damping happens in "world space" 
 	 * - so damping acts like drag on the movement.
 	 */
