PhysicsControl - disable (just by changing the default) the attempt to recreate controls since it's causing problems.

diff --git a/Engine/Plugins/Experimental/PhysicsControl/Source/PhysicsControl/Public/PhysicsControlComponent.h b/Engine/Plugins/Experimental/PhysicsControl/Source/PhysicsControl/Public/PhysicsControlComponent.h
--- a/Engine/Plugins/Experimental/PhysicsControl/Source/PhysicsControl/Public/PhysicsControlComponent.h
+++ b/Engine/Plugins/Experimental/PhysicsControl/Source/PhysicsControl/Public/PhysicsControlComponent.h
@@ -1689,7 +1689,7 @@ public:
 	// If set then controls that have been disabled due to the bodies being set to NoCollision (for
 	// example) will be reinitialized, if possible, when those bodies are returned to the simulation.
 	UPROPERTY(EditAnywhere, BlueprintReadWrite, Category = Settings)
-	bool bAttemptToRecreateDisabledControls = true;
+	bool bAttemptToRecreateDisabledControls = false;
 
 public:
 	//Begin ActorComponent interface
