Physics Control: Align the KinematicTargetSpace usage

diff --git a/Engine/Plugins/Experimental/ControlRigPhysics/Source/ControlRigPhysics/Private/RigPhysicsBodyComponent.cpp b/Engine/Plugins/Experimental/ControlRigPhysics/Source/ControlRigPhysics/Private/RigPhysicsBodyComponent.cpp
--- a/Engine/Plugins/Experimental/ControlRigPhysics/Source/ControlRigPhysics/Private/RigPhysicsBodyComponent.cpp
+++ b/Engine/Plugins/Experimental/ControlRigPhysics/Source/ControlRigPhysics/Private/RigPhysicsBodyComponent.cpp
@@ -25,7 +25,6 @@ void FRigPhysicsBodyComponent::Save(FArchive& Ar)
 	Ar << Collision;
 	Ar << BodyData;
 	Ar << KinematicTarget;
-	Ar << KinematicTargetSpace;
 	Ar << NoCollisionBodies;
 }
 
@@ -43,7 +42,11 @@ void FRigPhysicsBodyComponent::Load(FArchive& Ar)
 	}
 	Ar << BodyData;
 	Ar << KinematicTarget;
-	Ar << KinematicTargetSpace;
+	if (Ar.CustomVer(FPhysicsControlObjectVersion::GUID) < FPhysicsControlObjectVersion::RemoveUseSkeletalAnimation)
+	{
+		EPhysicsControlKinematicTargetSpace KinematicTargetSpace;
+		Ar << KinematicTargetSpace;
+	}
 
 	if (Ar.CustomVer(FPhysicsControlObjectVersion::GUID) <
 		FPhysicsControlObjectVersion::ControlRigRemoveCurrentDataFromPhysicsComponent)
diff --git a/Engine/Plugins/Experimental/ControlRigPhysics/Source/ControlRigPhysics/Private/RigPhysicsBodyExecution.cpp b/Engine/Plugins/Experimental/ControlRigPhysics/Source/ControlRigPhysics/Private/RigPhysicsBodyExecution.cpp
--- a/Engine/Plugins/Experimental/ControlRigPhysics/Source/ControlRigPhysics/Private/RigPhysicsBodyExecution.cpp
+++ b/Engine/Plugins/Experimental/ControlRigPhysics/Source/ControlRigPhysics/Private/RigPhysicsBodyExecution.cpp
@@ -159,7 +159,21 @@ FRigUnit_HierarchySetPhysicsBodyKinematicTarget_Execute()
 		ExecuteContext.Hierarchy->FindComponent(PhysicsBodyComponentKey)))
 	{
 		Component->KinematicTarget = KinematicTarget;
-		Component->KinematicTargetSpace = KinematicTargetSpace;
+	}
+}
+
+//======================================================================================================================
+FRigUnit_HierarchySetPhysicsBodyKinematicTargetSpace_Execute()
+{
+	DECLARE_SCOPE_HIERARCHICAL_COUNTER_RIGUNIT();
+	if (!ExecuteContext.Hierarchy)
+	{
+		return;
+	}
+	if (FRigPhysicsBodyComponent* Component = Cast<FRigPhysicsBodyComponent>(
+		ExecuteContext.Hierarchy->FindComponent(PhysicsBodyComponentKey)))
+	{
+		Component->BodyData.KinematicTargetSpace = KinematicTargetSpace;
 	}
 }
 
@@ -274,21 +288,6 @@ FRigUnit_HierarchySetPhysicsBodyPhysicsBlendWeight_Execute()
 	}
 }
 
-//======================================================================================================================
-FRigUnit_HierarchySetPhysicsBodyUseSkeletalAnimation_Execute()
-{
-	DECLARE_SCOPE_HIERARCHICAL_COUNTER_RIGUNIT();
-	if (!ExecuteContext.Hierarchy)
-	{
-		return;
-	}
-	if (FRigPhysicsBodyComponent* Component = Cast<FRigPhysicsBodyComponent>(
-		ExecuteContext.Hierarchy->FindComponent(PhysicsBodyComponentKey)))
-	{
-		Component->BodyData.bUseSkeletalAnimation = bUseSkeletalAnimation;
-	}
-}
-
 //======================================================================================================================
 FRigUnit_HierarchySetPhysicsBodyUpdateKinematicFromSimulation_Execute()
 {
diff --git a/Engine/Plugins/Experimental/ControlRigPhysics/Source/ControlRigPhysics/Private/RigPhysicsSimulation.cpp b/Engine/Plugins/Experimental/ControlRigPhysics/Source/ControlRigPhysics/Private/RigPhysicsSimulation.cpp
--- a/Engine/Plugins/Experimental/ControlRigPhysics/Source/ControlRigPhysics/Private/RigPhysicsSimulation.cpp
+++ b/Engine/Plugins/Experimental/ControlRigPhysics/Source/ControlRigPhysics/Private/RigPhysicsSimulation.cpp
@@ -962,7 +962,7 @@ void FRigPhysicsSimulation::UpdateBodyPrePhysics(
 	const FRigComponentKey ComponentKey = PhysicsComponent->GetKey();
 
 	FPhysicsControlModifierData BodyData = PhysicsComponent->BodyData;
-	ERigPhysicsKinematicTargetSpace KinematicTargetSpace = PhysicsComponent->KinematicTargetSpace;
+	EPhysicsControlKinematicTargetSpace KinematicTargetSpace = BodyData.KinematicTargetSpace;
 
 	if (SolverComponent->TrackInputCounter > 0)
 	{
@@ -971,7 +971,7 @@ void FRigPhysicsSimulation::UpdateBodyPrePhysics(
 		if (BodyData.MovementType != EPhysicsMovementType::Kinematic)
 		{
 			BodyData.MovementType = EPhysicsMovementType::Kinematic;
-			KinematicTargetSpace = ERigPhysicsKinematicTargetSpace::IgnoreTarget;
+			KinematicTargetSpace = EPhysicsControlKinematicTargetSpace::IgnoreTarget;
 		}
 	}
 
@@ -986,12 +986,12 @@ void FRigPhysicsSimulation::UpdateBodyPrePhysics(
 		FTransform KinematicTargetCS;
 		switch (KinematicTargetSpace)
 		{
-		case ERigPhysicsKinematicTargetSpace::Component:
+		case EPhysicsControlKinematicTargetSpace::Component:
 		{
 			KinematicTargetCS = PhysicsComponent->KinematicTarget;
 			break;
 		}
-		case ERigPhysicsKinematicTargetSpace::World:
+		case EPhysicsControlKinematicTargetSpace::World:
 		{
 			// Record.KinematicTarget * SimulationSpaceState.ComponentTM.Inverse();
 			KinematicTargetCS = PhysicsComponent->KinematicTarget.GetRelativeTransform(
@@ -1011,26 +1011,34 @@ void FRigPhysicsSimulation::UpdateBodyPrePhysics(
 			{
 				switch (KinematicTargetSpace)
 				{
-				case ERigPhysicsKinematicTargetSpace::OffsetInBoneSpace:
+				case EPhysicsControlKinematicTargetSpace::OffsetInBoneSpace:
 				{
 					KinematicTargetCS = PhysicsComponent->KinematicTarget * Record.SourceComponentSpaceTM.ToTransform();
 					break;
 				}
-				case ERigPhysicsKinematicTargetSpace::OffsetInWorldSpace:
+				case EPhysicsControlKinematicTargetSpace::OffsetInWorldSpace:
 				{
-					// Convert the bone to WS, apply the target, and convert back
+					// Convert the bone to WS (it's really component space!), apply the target, and convert back. 
 					FTransform BoneWS = Record.SourceComponentSpaceTM.ToTransform() * SimulationSpaceState.ComponentTM;
-					FTransform KinematicTargetWS = BoneWS * PhysicsComponent->KinematicTarget;
+					FTransform KinematicTargetWS = BoneWS;
+					// Note that we don't want to rotate the translation by the target - we apply it
+					// individually to orientation and position.
+					KinematicTargetWS.AddToTranslation(PhysicsComponent->KinematicTarget.GetTranslation());
+					KinematicTargetWS.SetRotation(
+						PhysicsComponent->KinematicTarget.GetRotation() * KinematicTargetWS.GetRotation());
 					// Danny TODO figure out which of the GetRelativeTransform versions this is
 					KinematicTargetCS = SimulationSpaceState.ComponentTM.Inverse() * KinematicTargetWS;
 					break;
 				}
-				case ERigPhysicsKinematicTargetSpace::OffsetInComponentSpace:
+				case EPhysicsControlKinematicTargetSpace::OffsetInComponentSpace:
 				{
+					// This applies the offset as a transform in component space. Note that this is
+					// different to what PCC does, because in ControlRig "world" is really
+					// "component" space.
 					KinematicTargetCS = Record.SourceComponentSpaceTM.ToTransform() * PhysicsComponent->KinematicTarget;
 					break;
 				}
-				case ERigPhysicsKinematicTargetSpace::IgnoreTarget:
+				case EPhysicsControlKinematicTargetSpace::IgnoreTarget:
 				{
 					KinematicTargetCS = Record.SourceComponentSpaceTM.ToTransform();
 					break;
diff --git a/Engine/Plugins/Experimental/ControlRigPhysics/Source/ControlRigPhysics/Private/RigPhysicsSimulation.h b/Engine/Plugins/Experimental/ControlRigPhysics/Source/ControlRigPhysics/Private/RigPhysicsSimulation.h
--- a/Engine/Plugins/Experimental/ControlRigPhysics/Source/ControlRigPhysics/Private/RigPhysicsSimulation.h
+++ b/Engine/Plugins/Experimental/ControlRigPhysics/Source/ControlRigPhysics/Private/RigPhysicsSimulation.h
@@ -74,8 +74,9 @@ struct FRigBodyRecord
 	// corrupting the output if anything is bad and we need to reset.
 	FTransform FinalComponentSpaceTM;
 
-	// These are updated for all records in UpdatePrePhysics. The times/validity are determined by
-	// the CurrentDeltaTime, PrevDeltaTime and update counters in the simulation itself.
+	// These source (i.e. bone/element) transforms are updated for all records in UpdatePrePhysics.
+	// The times/validity are determined by the CurrentDeltaTime, PrevDeltaTime and update counters
+	// in the simulation itself.
 	UE::PhysicsControl::FPosQuat SourceComponentSpaceTM;
 	FVector                      SourceComponentSpaceVelocity = FVector::ZeroVector;
 	FVector                      SourceComponentSpaceAngularVelocity = FVector::ZeroVector;
diff --git a/Engine/Plugins/Experimental/ControlRigPhysics/Source/ControlRigPhysics/Public/RigPhysicsBodyComponent.h b/Engine/Plugins/Experimental/ControlRigPhysics/Source/ControlRigPhysics/Public/RigPhysicsBodyComponent.h
--- a/Engine/Plugins/Experimental/ControlRigPhysics/Source/ControlRigPhysics/Public/RigPhysicsBodyComponent.h
+++ b/Engine/Plugins/Experimental/ControlRigPhysics/Source/ControlRigPhysics/Public/RigPhysicsBodyComponent.h
@@ -41,9 +41,6 @@ public:
 	UPROPERTY(BlueprintReadOnly, EditAnywhere, Category = Physics)
 	FTransform KinematicTarget;
 
-	UPROPERTY(BlueprintReadOnly, EditAnywhere, Category = Physics)
-	ERigPhysicsKinematicTargetSpace KinematicTargetSpace = ERigPhysicsKinematicTargetSpace::OffsetInBoneSpace;
-
 	// A list of body components with which we should not collide. The solver component can also be included.
 	UPROPERTY(BlueprintReadOnly, EditAnywhere, Category = Physics)
 	TArray<FRigComponentKey> NoCollisionBodies;
diff --git a/Engine/Plugins/Experimental/ControlRigPhysics/Source/ControlRigPhysics/Public/RigPhysicsBodyExecution.h b/Engine/Plugins/Experimental/ControlRigPhysics/Source/ControlRigPhysics/Public/RigPhysicsBodyExecution.h
--- a/Engine/Plugins/Experimental/ControlRigPhysics/Source/ControlRigPhysics/Public/RigPhysicsBodyExecution.h
+++ b/Engine/Plugins/Experimental/ControlRigPhysics/Source/ControlRigPhysics/Public/RigPhysicsBodyExecution.h
@@ -259,10 +259,31 @@ struct FRigUnit_HierarchySetPhysicsBodyKinematicTarget : public FRigUnit_Physics
 	FRigComponentKey PhysicsBodyComponentKey;
 
 	UPROPERTY(meta = (Input))
-	ERigPhysicsKinematicTargetSpace KinematicTargetSpace = ERigPhysicsKinematicTargetSpace::OffsetInBoneSpace;
+	FTransform KinematicTarget;
+};
+
+/**
+ * Sets the kinematic target space for a body - note that this won't actually make the body kinematic
+ */
+USTRUCT(meta = (DisplayName = "Set Physics Body Kinematic Target Space", Varying))
+struct FRigUnit_HierarchySetPhysicsBodyKinematicTargetSpace : public FRigUnit_PhysicsBaseMutable
+{
+	GENERATED_BODY()
+
+	FRigUnit_HierarchySetPhysicsBodyKinematicTargetSpace()
+	{
+		PhysicsBodyComponentKey.ElementKey.Type = ERigElementType::Bone;
+		PhysicsBodyComponentKey.Name = FRigPhysicsBodyComponent::GetDefaultName();
+	}
+
+	RIGVM_METHOD()
+	UE_API virtual void Execute() override;
 
 	UPROPERTY(meta = (Input))
-	FTransform KinematicTarget;
+	FRigComponentKey PhysicsBodyComponentKey;
+
+	UPROPERTY(meta = (Input))
+	EPhysicsControlKinematicTargetSpace KinematicTargetSpace = EPhysicsControlKinematicTargetSpace::OffsetInBoneSpace;
 };
 
 /**
@@ -408,31 +429,6 @@ struct FRigUnit_HierarchySetPhysicsBodyPhysicsBlendWeight : public FRigUnit_Phys
 	float PhysicsBlendWeight = 1.0f;
 };
 
-/**
- * Danny TODO I think this is redundant given the ability to specify the space when kinematic 
- * targets are set. Investigate propagating that into the rest of Physics Control.
- */
-USTRUCT(meta = (DisplayName = "Set Physics Body Use Skeletal Animation", Varying))
-struct FRigUnit_HierarchySetPhysicsBodyUseSkeletalAnimation : public FRigUnit_PhysicsBaseMutable
-{
-	GENERATED_BODY()
-
-	FRigUnit_HierarchySetPhysicsBodyUseSkeletalAnimation()
-	{
-		PhysicsBodyComponentKey.ElementKey.Type = ERigElementType::Bone;
-		PhysicsBodyComponentKey.Name = FRigPhysicsBodyComponent::GetDefaultName();
-	}
-
-	RIGVM_METHOD()
-	UE_API virtual void Execute() override;
-
-	UPROPERTY(meta = (Input))
-	FRigComponentKey PhysicsBodyComponentKey;
-
-	UPROPERTY(meta = (Input))
-	bool bUseSkeletalAnimation = true;
-};
-
 /**
  * If true, then kinematic objects will be written back from simulation to the bones. This only 
  * necessary when either kinematic targets are being used, or when the target bone differs from the source bone.
diff --git a/Engine/Plugins/Experimental/ControlRigPhysics/Source/ControlRigPhysics/Public/RigPhysicsData.h b/Engine/Plugins/Experimental/ControlRigPhysics/Source/ControlRigPhysics/Public/RigPhysicsData.h
--- a/Engine/Plugins/Experimental/ControlRigPhysics/Source/ControlRigPhysics/Public/RigPhysicsData.h
+++ b/Engine/Plugins/Experimental/ControlRigPhysics/Source/ControlRigPhysics/Public/RigPhysicsData.h
@@ -88,31 +88,6 @@ enum class ERigPhysicsWorldCollisionType : uint8
 	All,
 };
 
-/** 
- * Specifies what space a kinematic target is defined in 
- * Danny TODO this is much better than the similar mechanism in PCC and RBWC (which has a 
- * bUseSkeletalAnimation) - switch them to use this.
- */
-UENUM()
-enum class ERigPhysicsKinematicTargetSpace : uint8
-{
-	// Target is in world space
-	World,
-	// Target is in the space of the component
-	Component,
-	// Target is specified relative to the bone that controls the (kinematic) body, with the offset
-	// applied in the space of that body.
-	OffsetInBoneSpace,
-	// Target is specified relative to the bone that controls the (kinematic) body, with the offset
-	// applied in world space.
-	OffsetInWorldSpace,
-	// Target is specified relative to the bone that controls the (kinematic) body, with the offset
-	// applied in component space.
-	OffsetInComponentSpace,
-	// Ignore the specified target, and just track the bone
-	IgnoreTarget
-};
-
 /** Specifies what space the simulation should run in */
 UENUM()
 enum class ERigPhysicsSimulationSpace : uint8
diff --git a/Engine/Plugins/Experimental/PhysicsControl/Source/PhysicsControl/Private/PhysicsControlComponent.cpp b/Engine/Plugins/Experimental/PhysicsControl/Source/PhysicsControl/Private/PhysicsControlComponent.cpp
--- a/Engine/Plugins/Experimental/PhysicsControl/Source/PhysicsControl/Private/PhysicsControlComponent.cpp
+++ b/Engine/Plugins/Experimental/PhysicsControl/Source/PhysicsControl/Private/PhysicsControlComponent.cpp
@@ -1788,8 +1788,8 @@ bool UPhysicsControlComponent::SetBodyModifierKinematicTarget(
 	FPhysicsBodyModifierRecord* Record = FindBodyModifierRecord(Name);
 	if (Record)
 	{
-		Record->KinematicTargetPosition = KinematicTargetPosition;
-		Record->KinematicTargetOrientation = KinematicTargetOrienation.Quaternion();
+		Record->KinematicTarget.Translation = KinematicTargetPosition;
+		Record->KinematicTarget.Rotation = KinematicTargetOrienation.Quaternion();
 		if (bMakeKinematic)
 		{
 			Record->BodyModifier.ModifierData.MovementType = EPhysicsMovementType::Kinematic;
@@ -1957,41 +1957,41 @@ void UPhysicsControlComponent::SetBodyModifiersInSetPhysicsBlendWeight(
 }
 
 //======================================================================================================================
-bool UPhysicsControlComponent::SetBodyModifierUseSkeletalAnimation(
-	const FName Name,
-	const bool  bUseSkeletalAnimation)
+bool UPhysicsControlComponent::SetBodyModifierKinematicTargetSpace(
+	const FName                               Name,
+	const EPhysicsControlKinematicTargetSpace KinematicTargetSpace)
 {
 	FPhysicsBodyModifierRecord* Record = FindBodyModifierRecord(Name);
 	if (Record)
 	{
-		Record->BodyModifier.ModifierData.bUseSkeletalAnimation = bUseSkeletalAnimation;
+		Record->BodyModifier.ModifierData.KinematicTargetSpace = KinematicTargetSpace;
 		return true;
 	}
 	if (bWarnAboutInvalidNames)
 	{
 		UE_LOG(LogPhysicsControl, Warning,
-			TEXT("SetBodyModifierUseSkeletalAnimation - invalid name %s"), *Name.ToString());
+			TEXT("SetBodyModifierKinematicTargetSpace - invalid name %s"), *Name.ToString());
 	}
 	return false;
 }
 
 //======================================================================================================================
-void UPhysicsControlComponent::SetBodyModifiersUseSkeletalAnimation(
-	const TArray<FName>& Names,
-	const bool           bUseSkeletalAnimation)
+void UPhysicsControlComponent::SetBodyModifiersKinematicTargetSpace(
+	const TArray<FName>&                      Names,
+	const EPhysicsControlKinematicTargetSpace KinematicTargetSpace)
 {
 	for (FName Name : Names)
 	{
-		SetBodyModifierUseSkeletalAnimation(Name, bUseSkeletalAnimation);
+		SetBodyModifierKinematicTargetSpace(Name, KinematicTargetSpace);
 	}
 }
 
 //======================================================================================================================
-void UPhysicsControlComponent::SetBodyModifiersInSetUseSkeletalAnimation(
-	const FName SetName,
-	const bool  bUseSkeletalAnimation)
+void UPhysicsControlComponent::SetBodyModifiersInSetKinematicTargetSpace(
+	const FName                               SetName,
+	const EPhysicsControlKinematicTargetSpace KinematicTargetSpace)
 {
-	SetBodyModifiersUseSkeletalAnimation(GetBodyModifierNamesInSet(SetName), bUseSkeletalAnimation);
+	SetBodyModifiersKinematicTargetSpace(GetBodyModifierNamesInSet(SetName), KinematicTargetSpace);
 }
 
 //======================================================================================================================
diff --git a/Engine/Plugins/Experimental/PhysicsControl/Source/PhysicsControl/Private/PhysicsControlComponentImpl.cpp b/Engine/Plugins/Experimental/PhysicsControl/Source/PhysicsControl/Private/PhysicsControlComponentImpl.cpp
--- a/Engine/Plugins/Experimental/PhysicsControl/Source/PhysicsControl/Private/PhysicsControlComponentImpl.cpp
+++ b/Engine/Plugins/Experimental/PhysicsControl/Source/PhysicsControl/Private/PhysicsControlComponentImpl.cpp
@@ -209,36 +209,97 @@ void UPhysicsControlComponent::ApplyKinematicTarget(const FPhysicsBodyModifierRe
 			return;
 		}
 
-		const FTransform TM = BodyInstance->GetUnrealWorldTransform(); // Preserve scale
-		FTransform KinematicTarget = TM;
-		KinematicTarget.SetRotation(Record.KinematicTargetOrientation);
-		KinematicTarget.SetTranslation(Record.KinematicTargetPosition);
-		if (Record.BodyModifier.ModifierData.bUseSkeletalAnimation)
+		const FTransform CurrentBodyTM = BodyInstance->GetUnrealWorldTransform(); // Preserve scale
+		const FTransform MeshTM = SkeletalMeshComponent->GetComponentTransform();
+		FTransform KinematicTargetWS = CurrentBodyTM;
+		FTransform KinematicTargetOffset = Record.KinematicTarget.ToTransform();
+
+		switch (Record.BodyModifier.ModifierData.KinematicTargetSpace)
 		{
-			UE::PhysicsControl::FBoneData BoneData;
-			const UE::PhysicsControl::FPhysicsControlPoseData* PoseData;
-			if (GetBoneData(BoneData, PoseData, SkeletalMeshComponent, Record.BodyModifier.BoneName))
+			case EPhysicsControlKinematicTargetSpace::World:
+				KinematicTargetWS = KinematicTargetOffset;
+				break;
+			case EPhysicsControlKinematicTargetSpace::Component:
+				KinematicTargetWS = KinematicTargetOffset * MeshTM;
+				break;
+			case EPhysicsControlKinematicTargetSpace::IgnoreTarget:
+				break;
+			default:
 			{
-				FTransform BoneTM = BoneData.CurrentTM.ToTransform();
-				KinematicTarget = KinematicTarget * BoneTM;
+				// All the other options are relative to a bone, so the first task is to get that
+				UE::PhysicsControl::FBoneData BoneData;
+				const UE::PhysicsControl::FPhysicsControlPoseData* PoseData;
+				if (GetBoneData(BoneData, PoseData, SkeletalMeshComponent, Record.BodyModifier.BoneName))
+				{
+					FTransform BoneTM = BoneData.CurrentTM.ToTransform(); // Bone in WS
+					switch (Record.BodyModifier.ModifierData.KinematicTargetSpace)
+					{
+						case EPhysicsControlKinematicTargetSpace::OffsetInBoneSpace:
+							KinematicTargetWS = KinematicTargetOffset * BoneTM;
+							break;
+						case EPhysicsControlKinematicTargetSpace::OffsetInWorldSpace:
+						{
+							KinematicTargetWS = BoneTM;
+							// Note that we don't want to rotate the translation by the target - we apply it
+							// individually to orientation and position.
+							KinematicTargetWS.AddToTranslation(KinematicTargetOffset.GetTranslation());
+							KinematicTargetWS.SetRotation(
+								KinematicTargetOffset.GetRotation() * KinematicTargetWS.GetRotation());
+							break;
+						}
+						case EPhysicsControlKinematicTargetSpace::OffsetInComponentSpace:
+						{
+							FTransform KinematicTargetCS = BoneTM.GetRelativeTransform(MeshTM);
+							KinematicTargetCS.AddToTranslation(KinematicTargetOffset.GetTranslation());
+							KinematicTargetCS.SetRotation(
+								KinematicTargetOffset.GetRotation() * KinematicTargetCS.GetRotation());
+							KinematicTargetWS = KinematicTargetCS * MeshTM;
+							break;
+						}
+						default:
+							break;
+					}
+				}
+				break;
 			}
 		}
+
 		ETeleportType TT = UE::PhysicsControl::DetectTeleport(
-			TM, KinematicTarget, TeleportRotationThreshold, TeleportRotationThreshold) ? 
+			CurrentBodyTM, KinematicTargetWS, TeleportDistanceThreshold, TeleportRotationThreshold) ? 
 			ETeleportType::ResetPhysics : ETeleportType::None;
-		BodyInstance->SetBodyTransform(KinematicTarget, TT);
+		BodyInstance->SetBodyTransform(KinematicTargetWS, TT);
 	}
 	else
 	{
-		const FTransform TM = Record.Component->GetComponentToWorld();
+		const FTransform CurrentBodyTM = Record.Component->GetComponentToWorld();
+		FTransform KinematicTargetWS = CurrentBodyTM;
+
+		FTransform KinematicTargetOffset = Record.KinematicTarget.ToTransform();
+
+		switch (Record.BodyModifier.ModifierData.KinematicTargetSpace)
+		{
+		case EPhysicsControlKinematicTargetSpace::World:
+		case EPhysicsControlKinematicTargetSpace::OffsetInWorldSpace:
+		case EPhysicsControlKinematicTargetSpace::OffsetInBoneSpace:
+		default:
+			KinematicTargetWS = KinematicTargetOffset;
+			break;
+		case EPhysicsControlKinematicTargetSpace::Component:
+		case EPhysicsControlKinematicTargetSpace::OffsetInComponentSpace:
+			KinematicTargetWS = KinematicTargetOffset * this->GetComponentTransform();
+			break;
+		case EPhysicsControlKinematicTargetSpace::IgnoreTarget:
+			KinematicTargetWS = this->GetComponentTransform();
+			break;
+		}
+
 		const ETeleportType TT = UE::PhysicsControl::DetectTeleport(
-			TM.GetTranslation(), TM.GetRotation(), 
-			Record.KinematicTargetPosition, Record.KinematicTargetOrientation,
-			TeleportDistanceThreshold, TeleportRotationThreshold)
+			CurrentBodyTM, KinematicTargetWS, TeleportDistanceThreshold, TeleportRotationThreshold)
 			? ETeleportType::ResetPhysics : ETeleportType::None;
+
 		// Note that calling BodyInstance->SetBodyTransform moves the physics, but not the mesh
 		Record.Component->SetWorldLocationAndRotation(
-			Record.KinematicTargetPosition, Record.KinematicTargetOrientation, false, nullptr, TT);
+			KinematicTargetWS.GetTranslation(), KinematicTargetWS.GetRotation(), false, nullptr, TT);
 	}
 }
 
diff --git a/Engine/Plugins/Experimental/PhysicsControl/Source/PhysicsControl/Private/PhysicsControlData.cpp b/Engine/Plugins/Experimental/PhysicsControl/Source/PhysicsControl/Private/PhysicsControlData.cpp
--- a/Engine/Plugins/Experimental/PhysicsControl/Source/PhysicsControl/Private/PhysicsControlData.cpp
+++ b/Engine/Plugins/Experimental/PhysicsControl/Source/PhysicsControl/Private/PhysicsControlData.cpp
@@ -119,7 +119,7 @@ FPhysicsControlModifierData Interpolate(
 	SELECT_PARAM(CollisionType);
 	INTERPOLATE_PARAM(GravityMultiplier);
 	INTERPOLATE_PARAM(PhysicsBlendWeight);
-	SELECT_PARAM(bUseSkeletalAnimation);
+	SELECT_PARAM(KinematicTargetSpace);
 	SELECT_PARAM(bUpdateKinematicFromSimulation);
 	return Output;
 }
@@ -133,14 +133,14 @@ FPhysicsControlModifierSparseData Interpolate(
 	SELECT_PARAM(CollisionType);
 	INTERPOLATE_PARAM(GravityMultiplier);
 	INTERPOLATE_PARAM(PhysicsBlendWeight);
-	SELECT_PARAM(bUseSkeletalAnimation);
+	SELECT_PARAM(KinematicTargetSpace);
 	SELECT_PARAM(bUpdateKinematicFromSimulation);
 
 	SET_ENABLED_PARAM(MovementType);
 	SET_ENABLED_PARAM(CollisionType);
 	SET_ENABLED_PARAM(GravityMultiplier);
 	SET_ENABLED_PARAM(PhysicsBlendWeight);
-	SET_ENABLED_PARAM(bUseSkeletalAnimation);
+	SET_ENABLED_PARAM(KinematicTargetSpace);
 	SET_ENABLED_PARAM(bUpdateKinematicFromSimulation);
 	return Output;
 }
@@ -186,7 +186,7 @@ void FPhysicsControlModifierData::UpdateFromSparseData(const FPhysicsControlModi
 	SET_SPARSE_DATA(CollisionType);
 	SET_SPARSE_DATA(GravityMultiplier);
 	SET_SPARSE_DATA(PhysicsBlendWeight);
-	SET_SPARSE_DATA(bUseSkeletalAnimation);
+	SET_SPARSE_DATA(KinematicTargetSpace);
 	SET_SPARSE_DATA(bUpdateKinematicFromSimulation);
 }
 
@@ -217,6 +217,26 @@ FArchive& operator<<(FArchive& Ar, FPhysicsControlData& Data)
 	return Ar;
 }
 
+FArchive& operator <<(FArchive& Ar, FPhysicsControlModifierData& Data)
+{
+	Ar << Data.MovementType;
+	Ar << Data.CollisionType;
+	Ar << Data.GravityMultiplier;
+	Ar << Data.PhysicsBlendWeight;
+	if (Ar.CustomVer(FPhysicsControlObjectVersion::GUID) < FPhysicsControlObjectVersion::RemoveUseSkeletalAnimation)
+	{
+		bool bUseSkeletalAnimation = false;
+		Ar << bUseSkeletalAnimation;
+	}
+	else
+	{
+		Ar << Data.KinematicTargetSpace;
+	}
+	Ar << Data.bUpdateKinematicFromSimulation;
+	return Ar;
+}
+
+
 
 #undef SET_SPARSE_DATA
 #undef INTERPOLATE_PARAM
diff --git a/Engine/Plugins/Experimental/PhysicsControl/Source/PhysicsControl/Public/PhysicsControlComponent.h b/Engine/Plugins/Experimental/PhysicsControl/Source/PhysicsControl/Public/PhysicsControlComponent.h
--- a/Engine/Plugins/Experimental/PhysicsControl/Source/PhysicsControl/Public/PhysicsControlComponent.h
+++ b/Engine/Plugins/Experimental/PhysicsControl/Source/PhysicsControl/Public/PhysicsControlComponent.h
@@ -1270,45 +1270,40 @@ public:
 		const float PhysicsBlendWeight = 1.0f);
 
 	/**
-	 * Sets whether a body modifier should use skeletal animation for its kinematic targets
+	 * Sets what space kinematic targets should be applied in
 	 *
 	 * @param Name The name of the body modifier to access. 
-	 * @param bUseSkeletalAnimation Whether the kinematic target is specified in the frame of the skeletal
-	 *        animation, rather than world space. Only relevant if the body is part of a skeletal mesh.
+	 * @param KinematicTargetSpace The space in which the kinematic target should be applied. 
 	 * @return true if the body modifier was found, false if not
 	 */
 	UFUNCTION(BlueprintCallable, Category = PhysicsControl, Meta = (ReturnDisplayName = "Success"))
-	UE_API bool SetBodyModifierUseSkeletalAnimation(
-		const FName Name,
-		const bool  bUseSkeletalAnimation);
+	UE_API bool SetBodyModifierKinematicTargetSpace(
+		const FName                               Name,
+		const EPhysicsControlKinematicTargetSpace KinematicTargetSpace);
 
 	/**
-	 * Sets whether body modifiers should use skeletal animation for their kinematic targets
+	 * Sets what space kinematic targets should be applied in
 	 *
 	 * @param Names The names of the body modifiers to access. Note that if you have these in a FPhysicsControlNameArray
 	 *        then it can be split.
-	 * @param bUseSkeletalAnimation Whether the kinematic target is specified in the frame of the 
-	 *        skeletal animation, rather than world space. Only relevant if the
-	 *        body is part of a skeletal mesh.
+	 * @param KinematicTargetSpace The space in which the kinematic target should be applied.
 	 */
 	UFUNCTION(BlueprintCallable, Category = PhysicsControl)
-	UE_API void SetBodyModifiersUseSkeletalAnimation(
-		const TArray<FName>& Names,
-		const bool           bUseSkeletalAnimation);
+	UE_API void SetBodyModifiersKinematicTargetSpace(
+		const TArray<FName>&                      Names,
+		const EPhysicsControlKinematicTargetSpace KinematicTargetSpace);
 
 	/**
-	 * Sets whether body modifiers should use skeletal animation for their kinematic targets
+	 * Sets what space kinematic targets should be applied in
 	 *
 	 * @param Set The set of body modifiers to modify. Standard sets will include "All" and things like
 	 *        "ArmLeft", depending on how body modifiers have been created.
-	 * @param bUseSkeletalAnimation Whether the kinematic target is specified in the frame of the
-	 *        skeletal animation, rather than world space. Only relevant if the
-	 *        body is part of a skeletal mesh.
+	 * @param KinematicTargetSpace The space in which the kinematic target should be applied.
 	 */
 	UFUNCTION(BlueprintCallable, Category = PhysicsControl)
-	UE_API void SetBodyModifiersInSetUseSkeletalAnimation(
-		const FName Set,
-		const bool  bUseSkeletalAnimation);
+	UE_API void SetBodyModifiersInSetKinematicTargetSpace(
+		const FName                               Set,
+		const EPhysicsControlKinematicTargetSpace KinematicTargetSpace);
 
 	/**
 	 * Sets whether a body modifier should update kinematics from the simulation results
diff --git a/Engine/Plugins/Experimental/PhysicsControl/Source/PhysicsControl/Public/PhysicsControlData.h b/Engine/Plugins/Experimental/PhysicsControl/Source/PhysicsControl/Public/PhysicsControlData.h
--- a/Engine/Plugins/Experimental/PhysicsControl/Source/PhysicsControl/Public/PhysicsControlData.h
+++ b/Engine/Plugins/Experimental/PhysicsControl/Source/PhysicsControl/Public/PhysicsControlData.h
@@ -12,8 +12,8 @@
 #define UE_API PHYSICSCONTROL_API
 
 /**
- * Note that this file defines structures that mostly are, or could be, shared between the PhysicsControlComponent
- * and the RigidBodyWithControl node.
+ * Note that this file defines structures that mostly are, or could be, shared between the PhysicsControlComponent,
+ * the RigidBodyWithControl node, and Control Rig physics.
  */
 
 /**
@@ -114,6 +114,32 @@ struct FPhysicsControlSetUpdates
 	TArray<FPhysicsControlSetUpdate> ModifierSetUpdates;
 };
 
+// Specifies what space a kinematic target is defined in.
+UENUM()
+enum class EPhysicsControlKinematicTargetSpace : uint8
+{
+	// Target is in world space. For a control rig, this is the space of the control rig, rather
+	// than the "real" world.
+	World,
+	// If a skeleton/bone is being used, the target is in the space of that component. If not then
+	// it is in the space of the owning component (e.g. PhysicsControl, or ControlRig etc).
+	Component,
+	// If a skeleton/bone is being used, the target is specified relative to the bone that controls
+	// the (kinematic) body, with the offset applied in the space of that body. If not, then this is
+	// the same as World.
+	OffsetInBoneSpace,
+	// If a skeleton/bone is being used, the target is specified relative to the bone that controls
+	// the (kinematic) body, with the offset applied in world space. If not, then this is
+	// the same as World.
+	OffsetInWorldSpace,
+	// If a skeleton/bone is being used, the target is specified relative to the bone that controls
+	// the (kinematic) body, with the offset applied in component space. If not, then this is
+	// the same as Component.
+	OffsetInComponentSpace,
+	// Ignore the specified target, and just track the bone/object
+	IgnoreTarget
+};
+
 /**
  * Analogous to the ControlData, this indicates how an individual controlled body should move, with flags indicating
  * whether each element should get used.
@@ -124,24 +150,24 @@ struct FPhysicsControlModifierSparseData
 	GENERATED_BODY();
 
 	FPhysicsControlModifierSparseData(
-		const EPhysicsMovementType    InMovementType = EPhysicsMovementType::Simulated,
-		const ECollisionEnabled::Type InCollisionType = ECollisionEnabled::QueryAndPhysics,
-		const float                   InGravityMultiplier = 1.0f,
-		const float                   InPhysicsBlendWeight = 1.0f,
-		const bool                    InUseSkeletalAnimation = true,
-		const bool                    InUpdateKinematicFromSimulation = true)
+		const EPhysicsMovementType                InMovementType = EPhysicsMovementType::Simulated,
+		const ECollisionEnabled::Type             InCollisionType = ECollisionEnabled::QueryAndPhysics,
+		const float                               InGravityMultiplier = 1.0f,
+		const float                               InPhysicsBlendWeight = 1.0f,
+		const EPhysicsControlKinematicTargetSpace InKinematicTargetSpace = EPhysicsControlKinematicTargetSpace::OffsetInBoneSpace,
+		const bool                                InUpdateKinematicFromSimulation = true)
 		: MovementType(InMovementType)
 		, CollisionType(InCollisionType)
 		, GravityMultiplier(InGravityMultiplier)
 		, PhysicsBlendWeight(InPhysicsBlendWeight)
-		, bUseSkeletalAnimation(InUseSkeletalAnimation)
+		, KinematicTargetSpace(InKinematicTargetSpace)
 		, bUpdateKinematicFromSimulation(InUpdateKinematicFromSimulation)
 		// The flags
 		, bEnableMovementType(true)
 		, bEnableCollisionType(true)
 		, bEnableGravityMultiplier(true)
 		, bEnablePhysicsBlendWeight(true)
-		, bEnablebUseSkeletalAnimation(true)
+		, bEnableKinematicTargetSpace(true)
 		, bEnablebUpdateKinematicFromSimulation(true)
 	{
 	}
@@ -170,19 +196,19 @@ struct FPhysicsControlModifierSparseData
 	/**
 	 * When converting back from simulation to animation, how much to use the simulation output
 	 * versus the original animation input.
-	 * NOTE This is currently only used in the PhysicsControlComponent implementation
+	 * Note - this is not used in RBWC
 	 */
 	UPROPERTY(EditAnywhere, BlueprintReadWrite, Category = PhysicsControl, meta = (editcondition = "bEnablePhysicsBlendWeight"))
 	float PhysicsBlendWeight;
 
-	/** If true then the target will be applied on top of the skeletal animation (if there is any) */
-	UPROPERTY(EditAnywhere, BlueprintReadWrite, Category = PhysicsControl, meta = (editcondition = "bEnablebUseSkeletalAnimation"))
-	uint8 bUseSkeletalAnimation : 1;
+	// The space in which kinematic targets should operate
+	UPROPERTY(EditAnywhere, BlueprintReadWrite, Category = PhysicsControl, meta = (editcondition = "bEnableKinematicTargetSpace"))
+	EPhysicsControlKinematicTargetSpace KinematicTargetSpace = EPhysicsControlKinematicTargetSpace::OffsetInBoneSpace;
 
 	/**
 	 * If true then the associated actor's transform will be updated from the simulation when it is
 	 * kinematic. This is most likely useful when using async physics in order to prevent different
-	 * parts of the skeleton from being torn apart.
+	 * parts of the skeleton from being torn apart. Note - this is not used in RBWC
 	 */
 	UPROPERTY(EditAnywhere, BlueprintReadWrite, Category = PhysicsControl, meta = (editcondition = "bEnablebUpdateKinematicFromSimulation"))
 	uint8 bUpdateKinematicFromSimulation : 1;
@@ -200,7 +226,7 @@ struct FPhysicsControlModifierSparseData
 	uint8 bEnablePhysicsBlendWeight : 1;
 
 	UPROPERTY(EditAnywhere, Category = PhysicsControl, meta = (InlineEditConditionToggle))
-	uint8 bEnablebUseSkeletalAnimation : 1;
+	uint8 bEnableKinematicTargetSpace : 1;
 
 	UPROPERTY(EditAnywhere, Category = PhysicsControl, meta = (InlineEditConditionToggle))
 	uint8 bEnablebUpdateKinematicFromSimulation : 1;
@@ -215,39 +241,28 @@ struct FPhysicsControlModifierData
 	GENERATED_BODY();
 
 	FPhysicsControlModifierData(
-		const EPhysicsMovementType    InMovementType = EPhysicsMovementType::Simulated,
-		const ECollisionEnabled::Type InCollisionType = ECollisionEnabled::QueryAndPhysics,
-		const float                   InGravityMultiplier = 1.0f,
-		const float                   InPhysicsBlendWeight = 1.0f,
-		const bool                    InUseSkeletalAnimation = true,
-		const bool                    InUpdateKinematicFromSimulation = true)
+		const EPhysicsMovementType                InMovementType = EPhysicsMovementType::Simulated,
+		const ECollisionEnabled::Type             InCollisionType = ECollisionEnabled::QueryAndPhysics,
+		const float                               InGravityMultiplier = 1.0f,
+		const float                               InPhysicsBlendWeight = 1.0f,
+		const EPhysicsControlKinematicTargetSpace InKinematicTargetSpace = EPhysicsControlKinematicTargetSpace::OffsetInBoneSpace,
+		const bool                                InUpdateKinematicFromSimulation = true)
 		: MovementType(InMovementType)
 		, CollisionType(InCollisionType)
+		, KinematicTargetSpace(InKinematicTargetSpace)
 		, GravityMultiplier(InGravityMultiplier)
 		, PhysicsBlendWeight(InPhysicsBlendWeight)
-		, bUseSkeletalAnimation(InUseSkeletalAnimation)
 		, bUpdateKinematicFromSimulation(InUpdateKinematicFromSimulation)
 	{
 	}
 
 	bool operator==(const FPhysicsControlModifierData& other) const = default;
 
-	friend FArchive& operator <<(FArchive& Ar, FPhysicsControlModifierData& Data)
-	{
-		Ar << Data.MovementType;
-		Ar << Data.CollisionType;
-		Ar << Data.GravityMultiplier;
-		Ar << Data.PhysicsBlendWeight;
-		Ar << Data.bUseSkeletalAnimation;
-		Ar << Data.bUpdateKinematicFromSimulation;
-		return Ar;
-	}
+	friend UE_API FArchive& operator <<(FArchive& Ar, FPhysicsControlModifierData& Data);
 
 	UE_API void UpdateFromSparseData(const FPhysicsControlModifierSparseData& SparseData);
 
-	/**
-	 * How the body should move.
-	 */
+	// How the body should move.
 	UPROPERTY(EditAnywhere, BlueprintReadWrite, Category = PhysicsControl)
 	EPhysicsMovementType MovementType = EPhysicsMovementType::Simulated;
 
@@ -258,6 +273,10 @@ struct FPhysicsControlModifierData
 	UPROPERTY(EditAnywhere, BlueprintReadWrite, Category = PhysicsControl)
 	TEnumAsByte<ECollisionEnabled::Type> CollisionType = ECollisionEnabled::QueryAndPhysics;
 
+	// The space in which kinematic targets should operate
+	UPROPERTY(BlueprintReadOnly, EditAnywhere, Category = Physics)
+	EPhysicsControlKinematicTargetSpace KinematicTargetSpace = EPhysicsControlKinematicTargetSpace::OffsetInBoneSpace;
+
 	/**
 	 * How much gravity to use when simulating - typically between 0 and 1
 	 */
@@ -267,29 +286,21 @@ struct FPhysicsControlModifierData
 	/**
 	 * When converting back from simulation to animation, how much to use the simulation output
 	 * versus the original animation input.
-	 * NOTE This is currently only used in the PhysicsControlComponent implementation
+	 * Note - this is not used in RBWC
 	 */
 	UPROPERTY(EditAnywhere, BlueprintReadWrite, Category = PhysicsControl)
 	float PhysicsBlendWeight = 1.0f;
 
-	/** 
-	 * If true then the target will be applied on top of the skeletal animation (if there is any) 
-	 * NOTE This is currently only used in the PhysicsControlComponent implementation
-	 */
-	UPROPERTY(EditAnywhere, BlueprintReadWrite, Category = PhysicsControl)
-	bool bUseSkeletalAnimation;
-
 	/**
 	 * If true then the associated actor's transform will be updated from the simulation when it is
 	 * kinematic. This is most likely useful when using async physics in order to prevent different
 	 * parts of the skeleton from being torn apart.
-	 * NOTE This is currently only used in the PhysicsControlComponent implementation (is it needed in RBWC?)
+	 * Note - this is not used in RBWC
 	 */
 	UPROPERTY(EditAnywhere, BlueprintReadWrite, Category = PhysicsControl)
 	bool bUpdateKinematicFromSimulation;
 };
 
-
 /**
  * Represents a single/individual body modifier, with the modifier data
  */
@@ -398,7 +409,6 @@ struct FPhysicsControlSparseData
 
 	/** 
 	 * The maximum force used to drive the linear motion. Zero indicates no limit. 
-	 * Note - not yet implemented for RigidBodyWithControl 
 	 */
 	UPROPERTY(EditAnywhere, BlueprintReadWrite, Category = PhysicsControl, meta = (ClampMin = "0.0", editcondition = "bEnableMaxForce"))
 	float MaxForce;
@@ -424,7 +434,6 @@ struct FPhysicsControlSparseData
 
 	/** 
 	 * The maximum torque used to drive the angular motion. Zero indicates no limit. 
- 	 * Note - not yet implemented for RigidBodyWithControl 
 	 */
 	UPROPERTY(EditAnywhere, BlueprintReadWrite, Category = PhysicsControl, meta = (ClampMin = "0.0", editcondition = "bEnableMaxTorque"))
 	float MaxTorque;
@@ -433,7 +442,6 @@ struct FPhysicsControlSparseData
 	 * Multiplier on the linear velocity, which gets applied to the damping. A value of 1 means the animation target
 	 * velocity is used, which helps it track the animation. A value of 0 means damping happens in "world space" 
 	 * - so damping acts like drag on the movement.
-	 * NOTE This is currently only used in the RigidBodyWithControl implementation
 	 */
 	UPROPERTY(EditAnywhere, BlueprintReadWrite, Category = PhysicsControl, meta = (ClampMin = "0.0", editcondition = "bEnableLinearTargetVelocityMultiplier"))
 	float LinearTargetVelocityMultiplier;
@@ -442,7 +450,6 @@ struct FPhysicsControlSparseData
 	 * Multiplier on the angular velocity, which gets applied to the damping. A value of 1 means the animation target
 	 * velocity is used, which helps it track the animation. A value of 0 means damping happens in "world space" 
 	 * - so damping acts like drag on the movement.
-	 * NOTE This is currently only used in the RigidBodyWithControl implementation
 	 */
 	UPROPERTY(EditAnywhere, BlueprintReadWrite, Category = PhysicsControl, meta = (ClampMin = "0.0", editcondition = "bEnableAngularTargetVelocityMultiplier"))
 	float AngularTargetVelocityMultiplier;
diff --git a/Engine/Plugins/Experimental/PhysicsControl/Source/PhysicsControl/Public/PhysicsControlObjectVersion.h b/Engine/Plugins/Experimental/PhysicsControl/Source/PhysicsControl/Public/PhysicsControlObjectVersion.h
--- a/Engine/Plugins/Experimental/PhysicsControl/Source/PhysicsControl/Public/PhysicsControlObjectVersion.h
+++ b/Engine/Plugins/Experimental/PhysicsControl/Source/PhysicsControl/Public/PhysicsControlObjectVersion.h
@@ -38,6 +38,7 @@ struct FPhysicsControlObjectVersion
 		ControlRigCentreOfMassNudge,
 		ControlRigJointEnabled,
 		RemoveSkeletalAnimationVelocityMultiplier,
+		RemoveUseSkeletalAnimation,
 
 		// -----<new versions can be added above this line>-------------------------------------------------
 		VersionPlusOne,
diff --git a/Engine/Plugins/Experimental/PhysicsControl/Source/PhysicsControl/Public/PhysicsControlRecord.h b/Engine/Plugins/Experimental/PhysicsControl/Source/PhysicsControl/Public/PhysicsControlRecord.h
--- a/Engine/Plugins/Experimental/PhysicsControl/Source/PhysicsControl/Public/PhysicsControlRecord.h
+++ b/Engine/Plugins/Experimental/PhysicsControl/Source/PhysicsControl/Public/PhysicsControlRecord.h
@@ -78,8 +78,6 @@ struct FPhysicsBodyModifierRecord
 		FPhysicsControlModifierData     InBodyModifierData)
 		: Component(InComponent)
 		, BodyModifier(InBoneName, InBodyModifierData)
-		, KinematicTargetPosition(FVector::ZeroVector)
-		, KinematicTargetOrientation(FQuat::Identity)
 		, bResetToCachedTarget(false)
 	{}
 
@@ -89,13 +87,8 @@ struct FPhysicsBodyModifierRecord
 	// The core data
 	FPhysicsBodyModifier BodyModifier;
 
-	// The target position when kinematic. Note that this is applied on top of any animation 
-	// target if bUseSkeletalAnimation is set. 
-	FVector KinematicTargetPosition = FVector::ZeroVector;
-
-	// The target orientation when kinematic. Note that this is applied on top of any animation
-	// target if bUseSkeletalAnimation is set.
-	FQuat KinematicTargetOrientation = FQuat::Identity;
+	// The target transform when kinematic. 
+	UE::PhysicsControl::FPosQuat KinematicTarget;
 
 	// If true then the body will be set to the transform/velocity stored in any cached target (if that
 	// exists), and then this flag will be cleared.
