Align RigidBodyWithControl KinematicTargetSpace with the other parts of PhysicsControl

diff --git a/Engine/Plugins/Experimental/PhysicsControl/Source/PhysicsControl/Private/RigidBodyWithControl.cpp b/Engine/Plugins/Experimental/PhysicsControl/Source/PhysicsControl/Private/RigidBodyWithControl.cpp
--- a/Engine/Plugins/Experimental/PhysicsControl/Source/PhysicsControl/Private/RigidBodyWithControl.cpp
+++ b/Engine/Plugins/Experimental/PhysicsControl/Source/PhysicsControl/Private/RigidBodyWithControl.cpp
@@ -936,9 +936,24 @@ void FAnimNode_RigidBodyWithControl::ApplyKinematicTargets()
 				if (ActorHandle->GetIsKinematic() && BodyIndex != INDEX_NONE)
 				{
 					UE::PhysicsControl::FPosQuat TM(Target.TargetOrientation, Target.TargetPosition);
-					if (Target.bUseSkeletalAnimation && ensure(PoseData.IsValidIndex(BodyIndex)))
+					if (ensure(PoseData.IsValidIndex(BodyIndex)))
 					{
-						TM = PoseData.GetTM(BodyIndex) * TM;
+						switch (ModifierRecord->ModifierData.KinematicTargetSpace)
+						{
+						case EPhysicsControlKinematicTargetSpace::OffsetInBoneSpace:
+						{
+							TM = PoseData.GetTM(BodyIndex) * TM;
+							break;
+						}
+						default:
+							// Danny TODO Note that this implementation only supports the target
+							// being an offset in bone space. We need to add support for the other
+							// spaces too. If we reach here, the offset will be applied in
+							// simulation space, which is not even one of the options.
+							UE_LOG(LogPhysicsControl, Warning,
+								TEXT("RigidBodyWithControl only supports kinematic target offsets in bone space"));
+							break;
+						}
 					}
 					ActorHandle->SetKinematicTarget(TM.ToTransform());
 				}
diff --git a/Engine/Plugins/Experimental/PhysicsControl/Source/PhysicsControl/Public/RigidBodyControlData.h b/Engine/Plugins/Experimental/PhysicsControl/Source/PhysicsControl/Public/RigidBodyControlData.h
--- a/Engine/Plugins/Experimental/PhysicsControl/Source/PhysicsControl/Public/RigidBodyControlData.h
+++ b/Engine/Plugins/Experimental/PhysicsControl/Source/PhysicsControl/Public/RigidBodyControlData.h
@@ -55,7 +55,7 @@ struct FRigidBodyControlTargets
 };
 
 /**
- * A single kinematic target, which may be defined as an offset from the (implicit) animation target.
+ * A single kinematic target, which will be applied in the KinematicTargetSpace set in the relevant body modifier.
  */
 USTRUCT(BlueprintType)
 struct FRigidBodyKinematicTarget
@@ -65,7 +65,6 @@ struct FRigidBodyKinematicTarget
 	FRigidBodyKinematicTarget()
 		: TargetPosition(ForceInitToZero)
 		, TargetOrientation(ForceInitToZero)
-		, bUseSkeletalAnimation(true)
 	{
 	}
 
@@ -76,10 +75,6 @@ struct FRigidBodyKinematicTarget
 	/** The target orientation of the body */
 	UPROPERTY(EditAnywhere, BlueprintReadWrite, Category = PhysicsControl)
 	FRotator TargetOrientation;
-
-	/** If true then the target will be applied on top of the skeletal animation */
-	UPROPERTY(EditAnywhere, BlueprintReadWrite, Category = PhysicsControl)
-	uint8 bUseSkeletalAnimation : 1;
 };
 
 /**
