Support control/modifier and set names in all the functions. Also tidies up the docs etc. No behavior changes expected.

diff --git a/Engine/Plugins/Experimental/PhysicsControl/Source/PhysicsControl/Private/PhysicsControlComponent.cpp b/Engine/Plugins/Experimental/PhysicsControl/Source/PhysicsControl/Private/PhysicsControlComponent.cpp
--- a/Engine/Plugins/Experimental/PhysicsControl/Source/PhysicsControl/Private/PhysicsControlComponent.cpp
+++ b/Engine/Plugins/Experimental/PhysicsControl/Source/PhysicsControl/Private/PhysicsControlComponent.cpp
@@ -665,106 +665,125 @@ TMap<FName, FPhysicsControlNames> UPhysicsControlComponent::CreateControlsFromLi
 //======================================================================================================================
 void UPhysicsControlComponent::DestroyAllControlsAndBodyModifiers()
 {
-	DestroyControlsInSet(TEXT("All"));
-	DestroyBodyModifiersInSet(TEXT("All"));
+	DestroyControl(TEXT("All"), true, true);
+	DestroyBodyModifier(TEXT("All"), true, true);
 }
 
 //======================================================================================================================
-bool UPhysicsControlComponent::DestroyControl(const FName Name)
+void UPhysicsControlComponent::DestroyControl(const FName Name, const bool bApplyToControlsWithName, const bool bApplyToSetsWithName)
 {
-	return DestroyControl(Name, EDestroyBehavior::RemoveRecord);
+	if (bApplyToControlsWithName)
+	{
+		DestroyControl(Name, EDestroyBehavior::RemoveRecord);
+	}
+	if (bApplyToSetsWithName)
+	{
+		DestroyControlsInSet(Name);
+	}
 }
 
 //======================================================================================================================
-void UPhysicsControlComponent::DestroyControls(const TArray<FName>& Names)
+void UPhysicsControlComponent::DestroyControls(const TArray<FName>& Names, const bool bApplyToControlsWithName, const bool bApplyToSetsWithName)
 {
 	for (FName Name : Names)
 	{
-		DestroyControl(Name, EDestroyBehavior::RemoveRecord);
+		DestroyControl(Name, bApplyToControlsWithName, bApplyToSetsWithName);
 	}
 }
 
 //======================================================================================================================
 void UPhysicsControlComponent::DestroyControlsInSet(const FName SetName)
 {
-	// Make a copy as the set will be being modified during 
 	TArray<FName> Names = GetControlNamesInSet(SetName);
-	DestroyControls(Names);
+	DestroyControls(Names, true, false);
 }
 
 //======================================================================================================================
-bool UPhysicsControlComponent::SetControlEnabled(const FName Name, const bool bEnable)
+void UPhysicsControlComponent::SetControlEnabled(
+	const FName Name, const bool bEnable, const bool bApplyToControlsWithName, const bool bApplyToSetsWithName)
 {
-	FPhysicsControlRecord* Record = FindControlRecord(Name);
-	if (Record)
+	if (bApplyToControlsWithName)
 	{
-		Record->PhysicsControl.ControlData.bEnabled = bEnable;
-		return true;
+		FPhysicsControlRecord* Record = FindControlRecord(Name);
+		if (Record)
+		{
+			Record->PhysicsControl.ControlData.bEnabled = bEnable;
+		}
+		else if (bWarnAboutInvalidNames)
+		{
+			UE_LOG(LogPhysicsControl, Warning,
+				TEXT("SetControlEnabled - invalid name %s"), *Name.ToString());
+		}
 	}
-	if (bWarnAboutInvalidNames)
+	if (bApplyToSetsWithName)
 	{
-		UE_LOG(LogPhysicsControl, Warning,
-			TEXT("SetControlEnabled - invalid name %s"), *Name.ToString());
+		SetControlsInSetEnabled(Name, bEnable);
 	}
-	return false;
 }
 
 //======================================================================================================================
-void UPhysicsControlComponent::SetControlsEnabled(const TArray<FName>& Names, const bool bEnable)
+void UPhysicsControlComponent::SetControlsEnabled(
+	const TArray<FName>& Names, const bool bEnable, const bool bApplyToControlsWithName, const bool bApplyToSetsWithName)
 {
 	for (FName Name : Names)
 	{
-		SetControlEnabled(Name, bEnable);
+		SetControlEnabled(Name, bEnable, bApplyToControlsWithName, bApplyToSetsWithName);
 	}
 }
 
 //======================================================================================================================
 void UPhysicsControlComponent::SetControlsInSetEnabled(FName SetName, bool bEnable)
 {
-	SetControlsEnabled(GetControlNamesInSet(SetName), bEnable);
+	SetControlsEnabled(GetControlNamesInSet(SetName), bEnable, true, false);
 }
 
 //======================================================================================================================
-bool UPhysicsControlComponent::SetControlParent(
+void UPhysicsControlComponent::SetControlParent(
 	const FName          Name,
 	UPrimitiveComponent* ParentComponent,
-	const FName          ParentBoneName)
+	const FName          ParentBoneName,
+	const bool           bApplyToControlsWithName, 
+	const bool           bApplyToSetsWithName)
 {
-	FPhysicsControlRecord* Record = FindControlRecord(Name);
-	if (Record)
+
+	if (bApplyToControlsWithName)
 	{
-		if (USkeletalMeshComponent* SkeletalMeshComponent = 
-			Cast<USkeletalMeshComponent>(Record->ParentComponent.Get()))
+		FPhysicsControlRecord* Record = FindControlRecord(Name);
+		if (Record)
 		{
-			RemoveSkeletalMeshReferenceForCaching(SkeletalMeshComponent);
-		}
+			if (USkeletalMeshComponent* SkeletalMeshComponent = 
+				Cast<USkeletalMeshComponent>(Record->ParentComponent.Get()))
+			{
+				RemoveSkeletalMeshReferenceForCaching(SkeletalMeshComponent);
+			}
 
-		if (USkeletalMeshComponent* SkeletalMeshComponent = Cast<USkeletalMeshComponent>(ParentComponent))
-		{
-			AddSkeletalMeshReferenceForCaching(SkeletalMeshComponent);
-		}
+			if (USkeletalMeshComponent* SkeletalMeshComponent = Cast<USkeletalMeshComponent>(ParentComponent))
+			{
+				AddSkeletalMeshReferenceForCaching(SkeletalMeshComponent);
+			}
 
-		Record->ParentComponent = ParentComponent;
-		Record->PhysicsControl.ParentBoneName = ParentBoneName;
-		return Record->InitConstraint(this, Name, bWarnAboutInvalidNames);
+			Record->ParentComponent = ParentComponent;
+			Record->PhysicsControl.ParentBoneName = ParentBoneName;
+			Record->InitConstraint(this, Name, bWarnAboutInvalidNames);
+		}
 	}
-	if (bWarnAboutInvalidNames)
+	if (bApplyToSetsWithName)
 	{
-		UE_LOG(LogPhysicsControl, Warning,
-			TEXT("SetControlParent - invalid name %s"), *Name.ToString());
+		SetControlParentsInSet(Name, ParentComponent, ParentBoneName);
 	}
-	return false;
 }
 
 //======================================================================================================================
 void UPhysicsControlComponent::SetControlParents(
 	const TArray<FName>& Names,
 	UPrimitiveComponent* ParentComponent,
-	const FName          ParentBoneName)
+	const FName          ParentBoneName,
+	const bool           bApplyToControlsWithName,
+	const bool           bApplyToSetsWithName)
 {
 	for (FName Name : Names)
 	{
-		SetControlParent(Name, ParentComponent, ParentBoneName);
+		SetControlParent(Name, ParentComponent, ParentBoneName, bApplyToControlsWithName, bApplyToSetsWithName);
 	}
 }
 
@@ -774,37 +793,41 @@ void UPhysicsControlComponent::SetControlParentsInSet(
 	UPrimitiveComponent* ParentComponent,
 	const FName          ParentBoneName)
 {
-	SetControlParents(GetControlNamesInSet(SetName), ParentComponent, ParentBoneName);
+	SetControlParents(GetControlNamesInSet(SetName), ParentComponent, ParentBoneName, true, false);
 }
 
 //======================================================================================================================
 // Note - some params passed by value to allow them to be unconnected and enable inline BP editing
-bool UPhysicsControlComponent::SetControlData(
+void UPhysicsControlComponent::SetControlData(
 	const FName               Name, 
-	const FPhysicsControlData ControlData)
+	const FPhysicsControlData ControlData,
+	const bool                bApplyToControlsWithName,
+	const bool                bApplyToSetsWithName)
 {
-	FPhysicsControlRecord* Record = FindControlRecord(Name);
-	if (Record)
+	if (bApplyToControlsWithName)
 	{
-		Record->PhysicsControl.ControlData = ControlData;
-		return true;
+		FPhysicsControlRecord* Record = FindControlRecord(Name);
+		if (Record)
+		{
+			Record->PhysicsControl.ControlData = ControlData;
+		}
 	}
-	if (bWarnAboutInvalidNames)
+	if (bApplyToSetsWithName)
 	{
-		UE_LOG(LogPhysicsControl, Warning,
-			TEXT("SetControlData - invalid name %s"), *Name.ToString());
+		SetControlDatasInSet(Name, ControlData);
 	}
-	return false;
 }
 
 //======================================================================================================================
 void UPhysicsControlComponent::SetControlDatas(
 	const TArray<FName>&      Names, 
-	const FPhysicsControlData ControlData)
+	const FPhysicsControlData ControlData,
+	const bool                bApplyToControlsWithName,
+	const bool                bApplyToSetsWithName)
 {
 	for (FName Name : Names)
 	{
-		SetControlData(Name, ControlData);
+		SetControlData(Name, ControlData, bApplyToControlsWithName, bApplyToSetsWithName);
 	}
 }
 
@@ -813,37 +836,41 @@ void UPhysicsControlComponent::SetControlDatasInSet(
 	const FName               SetName,
 	const FPhysicsControlData ControlData)
 {
-	SetControlDatas(GetControlNamesInSet(SetName), ControlData);
+	SetControlDatas(GetControlNamesInSet(SetName), ControlData, true, false);
 }
 
 //======================================================================================================================
 // Note - some params passed by value to allow them to be unconnected and enable inline BP editing
-bool UPhysicsControlComponent::SetControlSparseData(
+void UPhysicsControlComponent::SetControlSparseData(
 	const FName                     Name, 
-	const FPhysicsControlSparseData ControlData)
+	const FPhysicsControlSparseData ControlData,
+	const bool                      bApplyToControlsWithName,
+	const bool                      bApplyToSetsWithName)
 {
-	FPhysicsControlRecord* Record = FindControlRecord(Name);
-	if (Record)
+	if (bApplyToControlsWithName)
 	{
-		Record->PhysicsControl.ControlData.UpdateFromSparseData(ControlData);
-		return true;
+		FPhysicsControlRecord* Record = FindControlRecord(Name);
+		if (Record)
+		{
+			Record->PhysicsControl.ControlData.UpdateFromSparseData(ControlData);
+		}
 	}
-	if (bWarnAboutInvalidNames)
+	if (bApplyToSetsWithName)
 	{
-		UE_LOG(LogPhysicsControl, Warning,
-			TEXT("SetControlData - invalid name %s"), *Name.ToString());
+		SetControlSparseDatasInSet(Name, ControlData);
 	}
-	return false;
 }
 
 //======================================================================================================================
 void UPhysicsControlComponent::SetControlSparseDatas(
 	const TArray<FName>&            Names, 
-	const FPhysicsControlSparseData ControlData)
+	const FPhysicsControlSparseData ControlData,
+	const bool                      bApplyToControlsWithName,
+	const bool                      bApplyToSetsWithName)
 {
 	for (FName Name : Names)
 	{
-		SetControlSparseData(Name, ControlData);
+		SetControlSparseData(Name, ControlData, bApplyToControlsWithName, bApplyToSetsWithName);
 	}
 }
 
@@ -852,43 +879,47 @@ void UPhysicsControlComponent::SetControlSparseDatasInSet(
 	const FName                     SetName,
 	const FPhysicsControlSparseData ControlData)
 {
-	SetControlSparseDatas(GetControlNamesInSet(SetName), ControlData);
+	SetControlSparseDatas(GetControlNamesInSet(SetName), ControlData, true, false);
 }
 
 //======================================================================================================================
 // Note - some params passed by value to allow them to be unconnected and enable inline BP editing
-bool UPhysicsControlComponent::SetControlMultiplier(
+void UPhysicsControlComponent::SetControlMultiplier(
 	const FName                      Name, 
 	const FPhysicsControlMultiplier  ControlMultiplier, 
-	const bool                       bEnableControl)
+	const bool                       bEnableControl,
+	const bool                       bApplyToControlsWithName,
+	const bool                       bApplyToSetsWithName)
 {
-	FPhysicsControlRecord* Record = FindControlRecord(Name);
-	if (Record)
+	if (bApplyToControlsWithName)
 	{
-		Record->PhysicsControl.ControlMultiplier = ControlMultiplier;
-		if (bEnableControl)
+		FPhysicsControlRecord* Record = FindControlRecord(Name);
+		if (Record)
 		{
-			Record->PhysicsControl.ControlData.bEnabled = true;
+			Record->PhysicsControl.ControlMultiplier = ControlMultiplier;
+			if (bEnableControl)
+			{
+				Record->PhysicsControl.ControlData.bEnabled = true;
+			}
 		}
-		return true;
 	}
-	if (bWarnAboutInvalidNames)
+	if (bApplyToSetsWithName)
 	{
-		UE_LOG(LogPhysicsControl, Warning,
-			TEXT("SetControlMultiplier - invalid name %s"), *Name.ToString());
+		SetControlMultipliersInSet(Name, ControlMultiplier, bEnableControl);
 	}
-	return false;
 }
 
 //======================================================================================================================
 void UPhysicsControlComponent::SetControlMultipliers(
 	const TArray<FName>&             Names,
 	const FPhysicsControlMultiplier  ControlMultiplier, 
-	const bool                       bEnableControl)
+	const bool                       bEnableControl,
+	const bool                       bApplyToControlsWithName,
+	const bool                       bApplyToSetsWithName)
 {
 	for (FName Name : Names)
 	{
-		SetControlMultiplier(Name, ControlMultiplier, bEnableControl);
+		SetControlMultiplier(Name, ControlMultiplier, bEnableControl, bApplyToControlsWithName, bApplyToSetsWithName);
 	}
 }
 
@@ -898,43 +929,47 @@ void UPhysicsControlComponent::SetControlMultipliersInSet(
 	const FPhysicsControlMultiplier  ControlMultiplier, 
 	const bool                       bEnableControl)
 {
-	SetControlMultipliers(GetControlNamesInSet(SetName), ControlMultiplier, bEnableControl);
+	SetControlMultipliers(GetControlNamesInSet(SetName), ControlMultiplier, bEnableControl, true, false);
 }
 
 //======================================================================================================================
 // Note - some params passed by value to allow them to be unconnected and enable inline BP editing
-bool UPhysicsControlComponent::SetControlSparseMultiplier(
+void UPhysicsControlComponent::SetControlSparseMultiplier(
 	const FName                            Name,
 	const FPhysicsControlSparseMultiplier  ControlMultiplier,
-	const bool                             bEnableControl)
+	const bool                             bEnableControl,
+	const bool                             bApplyToControlsWithName,
+	const bool                             bApplyToSetsWithName)
 {
-	FPhysicsControlRecord* Record = FindControlRecord(Name);
-	if (Record)
+	if (bApplyToControlsWithName)
 	{
-		Record->PhysicsControl.ControlMultiplier.UpdateFromSparseData(ControlMultiplier);
-		if (bEnableControl)
+		FPhysicsControlRecord* Record = FindControlRecord(Name);
+		if (Record)
 		{
-			Record->PhysicsControl.ControlData.bEnabled = true;
+			Record->PhysicsControl.ControlMultiplier.UpdateFromSparseData(ControlMultiplier);
+			if (bEnableControl)
+			{
+				Record->PhysicsControl.ControlData.bEnabled = true;
+			}
 		}
-		return true;
 	}
-	if (bWarnAboutInvalidNames)
+	if (bApplyToSetsWithName)
 	{
-		UE_LOG(LogPhysicsControl, Warning,
-			TEXT("SetControlSparseMultiplier - invalid name %s"), *Name.ToString());
+		SetControlSparseMultipliersInSet(Name, ControlMultiplier, bEnableControl);
 	}
-	return false;
 }
 
 //======================================================================================================================
 void UPhysicsControlComponent::SetControlSparseMultipliers(
 	const TArray<FName>&                  Names,
 	const FPhysicsControlSparseMultiplier ControlMultiplier,
-	const bool                            bEnableControl)
+	const bool                            bEnableControl,
+	const bool                            bApplyToControlsWithName,
+	const bool                            bApplyToSetsWithName)
 {
 	for (FName Name : Names)
 	{
-		SetControlSparseMultiplier(Name, ControlMultiplier, bEnableControl);
+		SetControlSparseMultiplier(Name, ControlMultiplier, bEnableControl, bApplyToControlsWithName, bApplyToSetsWithName);
 	}
 }
 
@@ -944,67 +979,79 @@ void UPhysicsControlComponent::SetControlSparseMultipliersInSet(
 	const FPhysicsControlSparseMultiplier ControlMultiplier,
 	const bool                            bEnableControl)
 {
-	SetControlSparseMultipliers(GetControlNamesInSet(SetName), ControlMultiplier, bEnableControl);
+	SetControlSparseMultipliers(GetControlNamesInSet(SetName), ControlMultiplier, bEnableControl, true, false);
 }
 
 //======================================================================================================================
-bool UPhysicsControlComponent::SetControlLinearData(
+void UPhysicsControlComponent::SetControlLinearData(
 	const FName Name, 
 	const float Strength, 
 	const float DampingRatio, 
 	const float ExtraDamping, 
 	const float MaxForce, 
-	const bool  bEnableControl)
+	const bool  bEnableControl,
+	const bool  bApplyToControlsWithName,
+	const bool  bApplyToSetsWithName)
 {
-	FPhysicsControlRecord* Record = FindControlRecord(Name);
-	if (Record)
+	if (bApplyToControlsWithName)
 	{
-		Record->PhysicsControl.ControlData.LinearStrength = Strength;
-		Record->PhysicsControl.ControlData.LinearDampingRatio = DampingRatio;
-		Record->PhysicsControl.ControlData.LinearExtraDamping = ExtraDamping;
-		Record->PhysicsControl.ControlData.MaxForce = MaxForce;
-		if (bEnableControl)
+		FPhysicsControlRecord* Record = FindControlRecord(Name);
+		if (Record)
 		{
-			Record->PhysicsControl.ControlData.bEnabled = true;
+			Record->PhysicsControl.ControlData.LinearStrength = Strength;
+			Record->PhysicsControl.ControlData.LinearDampingRatio = DampingRatio;
+			Record->PhysicsControl.ControlData.LinearExtraDamping = ExtraDamping;
+			Record->PhysicsControl.ControlData.MaxForce = MaxForce;
+			if (bEnableControl)
+			{
+				Record->PhysicsControl.ControlData.bEnabled = true;
+			}
 		}
-		return true;
 	}
-	if (bWarnAboutInvalidNames)
+	if (bApplyToSetsWithName)
 	{
-		UE_LOG(LogPhysicsControl, Warning,
-			TEXT("SetControlLinearData - invalid name %s"), *Name.ToString());
+		TArray<FName> Names = GetControlNamesInSet(Name);
+		for (FName ControlName : Names)
+		{
+			SetControlLinearData(ControlName, Strength, DampingRatio, ExtraDamping, MaxForce, bEnableControl, true, false);
+		}
 	}
-	return false;
 }
 
 //======================================================================================================================
-bool UPhysicsControlComponent::SetControlAngularData(
+void UPhysicsControlComponent::SetControlAngularData(
 	const FName Name, 
 	const float Strength, 
 	const float DampingRatio, 
 	const float ExtraDamping, 
 	const float MaxTorque, 
-	const bool  bEnableControl)
+	const bool  bEnableControl,
+	const bool  bApplyToControlsWithName,
+	const bool  bApplyToSetsWithName)
 {
-	FPhysicsControlRecord* Record = FindControlRecord(Name);
-	if (Record)
+	if (bApplyToControlsWithName)
 	{
-		Record->PhysicsControl.ControlData.AngularStrength = Strength;
-		Record->PhysicsControl.ControlData.AngularDampingRatio = DampingRatio;
-		Record->PhysicsControl.ControlData.AngularExtraDamping = ExtraDamping;
-		Record->PhysicsControl.ControlData.MaxTorque = MaxTorque;
-		if (bEnableControl)
+		FPhysicsControlRecord* Record = FindControlRecord(Name);
+		if (Record)
 		{
-			Record->PhysicsControl.ControlData.bEnabled = true;
+			Record->PhysicsControl.ControlData.AngularStrength = Strength;
+			Record->PhysicsControl.ControlData.AngularDampingRatio = DampingRatio;
+			Record->PhysicsControl.ControlData.AngularExtraDamping = ExtraDamping;
+			Record->PhysicsControl.ControlData.MaxTorque = MaxTorque;
+			if (bEnableControl)
+			{
+				Record->PhysicsControl.ControlData.bEnabled = true;
+			}
 		}
-		return true;
 	}
-	if (bWarnAboutInvalidNames)
+	if (bApplyToSetsWithName)
 	{
-		UE_LOG(LogPhysicsControl, Warning,
-			TEXT("SetControlAngularData - invalid name %s"), *Name.ToString());
+		TArray<FName> Names = GetControlNamesInSet(Name);
+		for (FName ControlName : Names)
+		{
+			SetControlAngularData(ControlName, Strength, DampingRatio, ExtraDamping, MaxTorque, bEnableControl, true, false);
+		}
 	}
-	return false;
 }
 
 //======================================================================================================================
@@ -1044,38 +1091,42 @@ bool UPhysicsControlComponent::ResetControlPoint(const FName Name)
 }
 
 //======================================================================================================================
-bool UPhysicsControlComponent::SetControlTarget(
+void UPhysicsControlComponent::SetControlTarget(
 	const FName                 Name, 
 	const FPhysicsControlTarget ControlTarget, 
-	const bool                  bEnableControl)
+	const bool                  bEnableControl,
+	const bool                  bApplyToControlsWithName,
+	const bool                  bApplyToSetsWithName)
 {
-	FPhysicsControlRecord* Record = FindControlRecord(Name);
-	if (Record)
+	if (bApplyToControlsWithName)
 	{
-		Record->ControlTarget = ControlTarget;
-		if (bEnableControl)
+		FPhysicsControlRecord* Record = FindControlRecord(Name);
+		if (Record)
 		{
-			Record->PhysicsControl.ControlData.bEnabled = true;
+			Record->ControlTarget = ControlTarget;
+			if (bEnableControl)
+			{
+				Record->PhysicsControl.ControlData.bEnabled = true;
+			}
 		}
-		return true;
 	}
-	if (bWarnAboutInvalidNames)
+	if (bApplyToSetsWithName)
 	{
-		UE_LOG(LogPhysicsControl, Warning,
-			TEXT("SetControlTarget - invalid name %s"), *Name.ToString());
+		SetControlTargetsInSet(Name, ControlTarget, bEnableControl);
 	}
-	return false;
 }
 
 //======================================================================================================================
 void UPhysicsControlComponent::SetControlTargets(
 	const TArray<FName>&        Names, 
 	const FPhysicsControlTarget ControlTarget, 
-	const bool                  bEnableControl)
+	const bool                  bEnableControl,
+	const bool                  bApplyToControlsWithName,
+	const bool                  bApplyToSetsWithName)
 {
 	for (FName Name : Names)
 	{
-		SetControlTarget(Name, ControlTarget, bEnableControl);
+		SetControlTarget(Name, ControlTarget, bEnableControl, bApplyToControlsWithName, bApplyToSetsWithName);
 	}
 }
 
@@ -1085,33 +1136,36 @@ void UPhysicsControlComponent::SetControlTargetsInSet(
 	const FPhysicsControlTarget ControlTarget, 
 	const bool                  bEnableControl)
 {
-	SetControlTargets(GetControlNamesInSet(SetName), ControlTarget, bEnableControl);
+	SetControlTargets(GetControlNamesInSet(SetName), ControlTarget, bEnableControl, true, false);
 }
 
 //======================================================================================================================
-bool UPhysicsControlComponent::SetControlTargetPositionAndOrientation(
+void UPhysicsControlComponent::SetControlTargetPositionAndOrientation(
 	const FName    Name, 
 	const FVector  Position, 
 	const FRotator Orientation, 
 	const float    VelocityDeltaTime, 
 	const bool     bEnableControl, 
-	const bool     bApplyControlPointToTarget)
+	const bool     bApplyControlPointToTarget,
+	const bool     bApplyToControlsWithName,
+	const bool     bApplyToSetsWithName)
 {
-	FPhysicsControlRecord* Record = FindControlRecord(Name);
-	if (Record)
+	if (bApplyToControlsWithName)
 	{
-		SetControlTargetPosition(
-			Name, Position, VelocityDeltaTime, bEnableControl, bApplyControlPointToTarget);
-		SetControlTargetOrientation(
-			Name, Orientation, VelocityDeltaTime, bEnableControl, bApplyControlPointToTarget);
-		return true;
+		FPhysicsControlRecord* Record = FindControlRecord(Name);
+		if (Record)
+		{
+			SetControlTargetPosition(
+				Name, Position, VelocityDeltaTime, bEnableControl, bApplyControlPointToTarget);
+			SetControlTargetOrientation(
+				Name, Orientation, VelocityDeltaTime, bEnableControl, bApplyControlPointToTarget);
+		}
 	}
-	if (bWarnAboutInvalidNames)
+	if (bApplyToSetsWithName)
 	{
-		UE_LOG(LogPhysicsControl, Warning,
-			TEXT("SetControlTargetPositionAndOrientation - invalid name %s"), *Name.ToString());
+		SetControlTargetPositionsAndOrientationsInSet(
+			Name, Position, Orientation, VelocityDeltaTime, bEnableControl, bApplyControlPointToTarget);
 	}
-	return false;
 }
 
 //======================================================================================================================
@@ -1121,12 +1175,15 @@ void UPhysicsControlComponent::SetControlTargetPositionsAndOrientations(
 	const FRotator       Orientation,
 	const float          VelocityDeltaTime,
 	const bool           bEnableControl,
-	const bool           bApplyControlPointToTarget)
+	const bool           bApplyControlPointToTarget,
+	const bool           bApplyToControlsWithName,
+	const bool           bApplyToSetsWithName)
 {
 	for (FName Name : Names)
 	{
 		SetControlTargetPositionAndOrientation(
-			Name, Position, Orientation, VelocityDeltaTime, bEnableControl, bApplyControlPointToTarget);
+			Name, Position, Orientation, VelocityDeltaTime, bEnableControl, bApplyControlPointToTarget,
+			bApplyToControlsWithName, bApplyToSetsWithName);
 	}
 }
 
@@ -1140,43 +1197,45 @@ void UPhysicsControlComponent::SetControlTargetPositionsAndOrientationsInSet(
 	const bool           bApplyControlPointToTarget)
 {
 	SetControlTargetPositionsAndOrientations(GetControlNamesInSet(SetName), 
-		Position, Orientation, VelocityDeltaTime, bEnableControl, bApplyControlPointToTarget);
+		Position, Orientation, VelocityDeltaTime, bEnableControl, bApplyControlPointToTarget, true, false);
 }
 
 //======================================================================================================================
-bool UPhysicsControlComponent::SetControlTargetPosition(
+void UPhysicsControlComponent::SetControlTargetPosition(
 	const FName   Name, 
 	const FVector Position, 
 	const float   VelocityDeltaTime, 
 	const bool    bEnableControl, 
-	const bool    bApplyControlPointToTarget)
+	const bool    bApplyControlPointToTarget,
+	const bool    bApplyToControlsWithName, 
+	const bool    bApplyToSetsWithName)
 {
-	FPhysicsControlRecord* Record = FindControlRecord(Name);
-	if (Record)
+	if (bApplyToControlsWithName)
 	{
-		if (VelocityDeltaTime != 0)
+		FPhysicsControlRecord* Record = FindControlRecord(Name);
+		if (Record)
 		{
-			Record->ControlTarget.TargetVelocity =
-				(Position - Record->ControlTarget.TargetPosition) / VelocityDeltaTime;
-		}
-		else
-		{
-			Record->ControlTarget.TargetVelocity = FVector::ZeroVector;
-		}
-		Record->ControlTarget.TargetPosition = Position;
-		Record->ControlTarget.bApplyControlPointToTarget = bApplyControlPointToTarget;
-		if (bEnableControl)
-		{
-			Record->PhysicsControl.ControlData.bEnabled = true;
+			if (VelocityDeltaTime != 0)
+			{
+				Record->ControlTarget.TargetVelocity =
+					(Position - Record->ControlTarget.TargetPosition) / VelocityDeltaTime;
+			}
+			else
+			{
+				Record->ControlTarget.TargetVelocity = FVector::ZeroVector;
+			}
+			Record->ControlTarget.TargetPosition = Position;
+			Record->ControlTarget.bApplyControlPointToTarget = bApplyControlPointToTarget;
+			if (bEnableControl)
+			{
+				Record->PhysicsControl.ControlData.bEnabled = true;
+			}
 		}
-		return true;
 	}
-	if (bWarnAboutInvalidNames)
+	if (bApplyToSetsWithName)
 	{
-		UE_LOG(LogPhysicsControl, Warning,
-			TEXT("SetControlTargetPosition - invalid name %s"), *Name.ToString());
+		SetControlTargetPositionsInSet(Name, Position, VelocityDeltaTime, bEnableControl, bApplyControlPointToTarget);
 	}
-	return false;
 }
 
 //======================================================================================================================
@@ -1185,12 +1244,15 @@ void UPhysicsControlComponent::SetControlTargetPositions(
 	const FVector        Position,
 	const float          VelocityDeltaTime,
 	const bool           bEnableControl,
-	const bool           bApplyControlPointToTarget)
+	const bool           bApplyControlPointToTarget,
+	const bool           bApplyToControlsWithName,
+	const bool           bApplyToSetsWithName)
 {
 	for (FName Name : Names)
 	{
 		SetControlTargetPosition(
-			Name, Position, VelocityDeltaTime, bEnableControl, bApplyControlPointToTarget);
+			Name, Position, VelocityDeltaTime, bEnableControl, bApplyControlPointToTarget,
+			bApplyToControlsWithName, bApplyToSetsWithName);
 	}
 }
 
@@ -1203,47 +1265,50 @@ void UPhysicsControlComponent::SetControlTargetPositionsInSet(
 	const bool    bApplyControlPointToTarget)
 {
 	SetControlTargetPositions(GetControlNamesInSet(SetName),
-		Position, VelocityDeltaTime, bEnableControl, bApplyControlPointToTarget);
+		Position, VelocityDeltaTime, bEnableControl, bApplyControlPointToTarget, true, false);
 }
 
 //======================================================================================================================
-bool UPhysicsControlComponent::SetControlTargetOrientation(
+void UPhysicsControlComponent::SetControlTargetOrientation(
 	const FName    Name, 
 	const FRotator Orientation, 
 	const float    AngularVelocityDeltaTime, 
 	const bool     bEnableControl, 
-	const bool     bApplyControlPointToTarget)
+	const bool     bApplyControlPointToTarget,
+	const bool     bApplyToControlsWithName,
+	const bool     bApplyToSetsWithName)
 {
-	FPhysicsControlRecord* Record = FindControlRecord(Name);
-	if (Record)
+	if (bApplyToControlsWithName)
 	{
-		if (AngularVelocityDeltaTime != 0)
+		FPhysicsControlRecord* Record = FindControlRecord(Name);
+		if (Record)
 		{
-			const FQuat OldQ = Record->ControlTarget.TargetOrientation.Quaternion();
-			const FQuat OrientationQ = Orientation.Quaternion();
-			// Note that quats multiply in the opposite order to TMs
-			const FQuat DeltaQ = (OrientationQ * OldQ.Inverse()).GetShortestArcWith(FQuat::Identity);
-			Record->ControlTarget.TargetAngularVelocity =
-				DeltaQ.ToRotationVector() / (UE_TWO_PI * AngularVelocityDeltaTime);
-		}
-		else
-		{
-			Record->ControlTarget.TargetAngularVelocity = FVector::ZeroVector;
-		}
-		Record->ControlTarget.TargetOrientation = Orientation;
-		Record->ControlTarget.bApplyControlPointToTarget = bApplyControlPointToTarget;
-		if (bEnableControl)
-		{
-			Record->PhysicsControl.ControlData.bEnabled = true;
+			if (AngularVelocityDeltaTime != 0)
+			{
+				const FQuat OldQ = Record->ControlTarget.TargetOrientation.Quaternion();
+				const FQuat OrientationQ = Orientation.Quaternion();
+				// Note that quats multiply in the opposite order to TMs
+				const FQuat DeltaQ = (OrientationQ * OldQ.Inverse()).GetShortestArcWith(FQuat::Identity);
+				Record->ControlTarget.TargetAngularVelocity =
+					DeltaQ.ToRotationVector() / (UE_TWO_PI * AngularVelocityDeltaTime);
+			}
+			else
+			{
+				Record->ControlTarget.TargetAngularVelocity = FVector::ZeroVector;
+			}
+			Record->ControlTarget.TargetOrientation = Orientation;
+			Record->ControlTarget.bApplyControlPointToTarget = bApplyControlPointToTarget;
+			if (bEnableControl)
+			{
+				Record->PhysicsControl.ControlData.bEnabled = true;
+			}
 		}
-		return true;
 	}
-	if (bWarnAboutInvalidNames)
+	if (bApplyToSetsWithName)
 	{
-		UE_LOG(LogPhysicsControl, Warning,
-			TEXT("SetControlTargetOrientation - invalid name %s"), *Name.ToString());
+		SetControlTargetOrientationsInSet(
+			Name, Orientation, AngularVelocityDeltaTime, bEnableControl, bApplyControlPointToTarget);
 	}
-	return false;
 }
 
 //======================================================================================================================
@@ -1252,12 +1317,15 @@ void UPhysicsControlComponent::SetControlTargetOrientations(
 	const FRotator Orientation,
 	const float    AngularVelocityDeltaTime,
 	const bool     bEnableControl,
-	const bool     bApplyControlPointToTarget)
+	const bool     bApplyControlPointToTarget,
+	const bool     bApplyToControlsWithName,
+	const bool     bApplyToSetsWithName)
 {
 	for (FName Name : Names)
 	{
 		SetControlTargetOrientation(
-			Name, Orientation, AngularVelocityDeltaTime, bEnableControl, bApplyControlPointToTarget);
+			Name, Orientation, AngularVelocityDeltaTime, bEnableControl, bApplyControlPointToTarget,
+			bApplyToControlsWithName, bApplyToSetsWithName);
 	}
 }
 
@@ -1270,7 +1338,7 @@ void UPhysicsControlComponent::SetControlTargetOrientationsInSet(
 	const bool     bApplyControlPointToTarget)
 {
 	SetControlTargetOrientations(GetControlNamesInSet(SetName),
-		Orientation, AngularVelocityDeltaTime, bEnableControl, bApplyControlPointToTarget);
+		Orientation, AngularVelocityDeltaTime, bEnableControl, bApplyControlPointToTarget, true, false);
 }
 
 //======================================================================================================================
@@ -1400,26 +1468,29 @@ bool UPhysicsControlComponent::SetControlTargetPoses(
 }
 
 //======================================================================================================================
-bool UPhysicsControlComponent::SetControlUseSkeletalAnimation(
+void UPhysicsControlComponent::SetControlUseSkeletalAnimation(
 	const FName Name,
 	const bool  bUseSkeletalAnimation,
 	const float AngularTargetVelocityMultiplier,
-	const float LinearTargetVelocityMultiplier)
+	const float LinearTargetVelocityMultiplier,
+	const bool  bApplyToControlsWithName,
+	const bool  bApplyToSetsWithName)
 {
-	FPhysicsControl* PhysicsControl = FindControl(Name);
-	if (PhysicsControl)
+	if (bApplyToControlsWithName)
 	{
-		PhysicsControl->ControlData.bUseSkeletalAnimation = bUseSkeletalAnimation;
-		PhysicsControl->ControlData.AngularTargetVelocityMultiplier = AngularTargetVelocityMultiplier;
-		PhysicsControl->ControlData.LinearTargetVelocityMultiplier = LinearTargetVelocityMultiplier;
-		return true;
+		FPhysicsControl* PhysicsControl = FindControl(Name);
+		if (PhysicsControl)
+		{
+			PhysicsControl->ControlData.bUseSkeletalAnimation = bUseSkeletalAnimation;
+			PhysicsControl->ControlData.AngularTargetVelocityMultiplier = AngularTargetVelocityMultiplier;
+			PhysicsControl->ControlData.LinearTargetVelocityMultiplier = LinearTargetVelocityMultiplier;
+		}
 	}
-	if (bWarnAboutInvalidNames)
+	if (bApplyToSetsWithName)
 	{
-		UE_LOG(LogPhysicsControl, Warning,
-			TEXT("SetControlUseSkeletalAnimation - invalid name %s"), *Name.ToString());
+		SetControlsInSetUseSkeletalAnimation(
+			Name, bUseSkeletalAnimation, AngularTargetVelocityMultiplier, LinearTargetVelocityMultiplier);
 	}
-	return false;
 }
 
 //======================================================================================================================
@@ -1427,12 +1498,15 @@ void UPhysicsControlComponent::SetControlsUseSkeletalAnimation(
 	const TArray<FName>& Names,
 	const bool           bUseSkeletalAnimation,
 	const float          AngularTargetVelocityMultiplier,
-	const float          LinearTargetVelocityMultiplier)
+	const float          LinearTargetVelocityMultiplier,
+	const bool           bApplyToControlsWithName,
+	const bool           bApplyToSetsWithName)
 {
 	for (FName Name : Names)
 	{
 		SetControlUseSkeletalAnimation(
-			Name, bUseSkeletalAnimation, AngularTargetVelocityMultiplier, LinearTargetVelocityMultiplier);
+			Name, bUseSkeletalAnimation, AngularTargetVelocityMultiplier, LinearTargetVelocityMultiplier,
+			bApplyToControlsWithName, bApplyToSetsWithName);
 	}
 }
 
@@ -1444,40 +1518,48 @@ void UPhysicsControlComponent::SetControlsInSetUseSkeletalAnimation(
 	const float LinearTargetVelocityMultiplier)
 {
 	SetControlsUseSkeletalAnimation(
-		GetControlNamesInSet(SetName), 
-		bUseSkeletalAnimation, AngularTargetVelocityMultiplier, LinearTargetVelocityMultiplier);
+		GetControlNamesInSet(SetName), bUseSkeletalAnimation, AngularTargetVelocityMultiplier, 
+		LinearTargetVelocityMultiplier, true, false);
 }
 
 //======================================================================================================================
-bool UPhysicsControlComponent::SetControlDisableCollision(const FName Name, const bool bDisableCollision)
+void UPhysicsControlComponent::SetControlDisableCollision(
+	const FName   Name, 
+	const bool    bDisableCollision,
+	const bool    bApplyToControlsWithName,
+	const bool    bApplyToSetsWithName)
 {
-	FPhysicsControl* PhysicsControl = FindControl(Name);
-	if (PhysicsControl)
+	if (bApplyToControlsWithName)
 	{
-		PhysicsControl->ControlData.bDisableCollision = bDisableCollision;
-		return true;
+		FPhysicsControl* PhysicsControl = FindControl(Name);
+		if (PhysicsControl)
+		{
+			PhysicsControl->ControlData.bDisableCollision = bDisableCollision;
+		}
 	}
-	if (bWarnAboutInvalidNames)
+	if (bApplyToSetsWithName)
 	{
-		UE_LOG(LogPhysicsControl, Warning,
-			TEXT("SetControlDisableCollision - invalid name %s"), *Name.ToString());
+		SetControlsInSetDisableCollision(Name, bDisableCollision);
 	}
-	return false;
 }
 
 //======================================================================================================================
-void UPhysicsControlComponent::SetControlsDisableCollision(const TArray<FName>& Names, const bool bDisableCollision)
+void UPhysicsControlComponent::SetControlsDisableCollision(
+	const TArray<FName>& Names, 
+	const bool           bDisableCollision,
+	const bool           bApplyToControlsWithName,
+	const bool           bApplyToSetsWithName)
 {
 	for (FName Name : Names)
 	{
-		SetControlDisableCollision(Name, bDisableCollision);
+		SetControlDisableCollision(Name, bDisableCollision, bApplyToControlsWithName, bApplyToSetsWithName);
 	}
 }
 
 //======================================================================================================================
 void UPhysicsControlComponent::SetControlsInSetDisableCollision(const FName SetName, const bool bDisableCollision)
 {
-	SetControlsDisableCollision(GetControlNamesInSet(SetName), bDisableCollision);
+	SetControlsDisableCollision(GetControlNamesInSet(SetName), bDisableCollision, true, false);
 }
 
 //======================================================================================================================
@@ -1678,17 +1760,26 @@ TMap<FName, FPhysicsControlNames> UPhysicsControlComponent::CreateBodyModifiersF
 }
 
 //======================================================================================================================
-bool UPhysicsControlComponent::DestroyBodyModifier(const FName Name)
+void UPhysicsControlComponent::DestroyBodyModifier(
+	const FName Name, const bool bApplyToModifiersWithName, const bool bApplyToSetsWithName)
 {
-	return DestroyBodyModifier(Name, EDestroyBehavior::RemoveRecord);
+	if (bApplyToModifiersWithName)
+	{
+		DestroyBodyModifier(Name, EDestroyBehavior::RemoveRecord);
+	}
+	if (bApplyToSetsWithName)
+	{
+		DestroyBodyModifiersInSet(Name);
+	}
 }
 
 //======================================================================================================================
-void UPhysicsControlComponent::DestroyBodyModifiers(const TArray<FName>& Names)
+void UPhysicsControlComponent::DestroyBodyModifiers(
+	const TArray<FName>& Names, const bool bApplyToModifiersWithName, const bool bApplyToSetsWithName)
 {
 	for (FName Name : Names)
 	{
-		DestroyBodyModifier(Name, EDestroyBehavior::RemoveRecord);
+		DestroyBodyModifier(Name, bApplyToModifiersWithName, bApplyToSetsWithName);
 	}
 }
 
@@ -1696,37 +1787,41 @@ void UPhysicsControlComponent::DestroyBodyModifiers(const TArray<FName>& Names)
 void UPhysicsControlComponent::DestroyBodyModifiersInSet(const FName SetName)
 {
 	TArray<FName> Names = GetBodyModifierNamesInSet(SetName);
-	DestroyBodyModifiers(Names);
+	DestroyBodyModifiers(Names, true, false);
 }
 
 //======================================================================================================================
 // Note - some params passed by value to allow them to be unconnected and enable inline BP editing
-bool UPhysicsControlComponent::SetBodyModifierData(
+void UPhysicsControlComponent::SetBodyModifierData(
 	const FName                       Name,
-	const FPhysicsControlModifierData ModifierData)
+	const FPhysicsControlModifierData ModifierData, 
+	const bool                        bApplyToModifiersWithName, 
+	const bool                        bApplyToSetsWithName)
 {
-	FPhysicsBodyModifierRecord* Record = FindBodyModifierRecord(Name);
-	if (Record)
+	if (bApplyToModifiersWithName)
 	{
-		Record->BodyModifier.ModifierData = ModifierData;
-		return true;
+		FPhysicsBodyModifierRecord* Record = FindBodyModifierRecord(Name);
+		if (Record)
+		{
+			Record->BodyModifier.ModifierData = ModifierData;
+		}
 	}
-	if (bWarnAboutInvalidNames)
+	if (bApplyToSetsWithName)
 	{
-		UE_LOG(LogPhysicsControl, Warning,
-			TEXT("SetBodyModifierData - invalid name %s"), *Name.ToString());
+		SetBodyModifierDatasInSet(Name, ModifierData);
 	}
-	return false;
 }
 
 //======================================================================================================================
 void UPhysicsControlComponent::SetBodyModifierDatas(
 	const TArray<FName>&              Names,
-	const FPhysicsControlModifierData ModifierData)
+	const FPhysicsControlModifierData ModifierData,
+	const bool                        bApplyToModifiersWithName,
+	const bool                        bApplyToSetsWithName)
 {
 	for (FName Name : Names)
 	{
-		SetBodyModifierData(Name, ModifierData);
+		SetBodyModifierData(Name, ModifierData, bApplyToModifiersWithName, bApplyToSetsWithName);
 	}
 }
 
@@ -1735,37 +1830,41 @@ void UPhysicsControlComponent::SetBodyModifierDatasInSet(
 	const FName                       SetName,
 	const FPhysicsControlModifierData ModifierData)
 {
-	SetBodyModifierDatas(GetControlNamesInSet(SetName), ModifierData);
+	SetBodyModifierDatas(GetBodyModifierNamesInSet(SetName), ModifierData, true, false);
 }
 
 //======================================================================================================================
 // Note - some params passed by value to allow them to be unconnected and enable inline BP editing
-bool UPhysicsControlComponent::SetBodyModifierSparseData(
+void UPhysicsControlComponent::SetBodyModifierSparseData(
 	const FName                             Name,
-	const FPhysicsControlModifierSparseData ModifierData)
+	const FPhysicsControlModifierSparseData ModifierData,
+	const bool                              bApplyToModifiersWithName,
+	const bool                              bApplyToSetsWithName)
 {
-	FPhysicsBodyModifierRecord* Record = FindBodyModifierRecord(Name);
-	if (Record)
+	if (bApplyToModifiersWithName)
 	{
-		Record->BodyModifier.ModifierData.UpdateFromSparseData(ModifierData);
-		return true;
+		FPhysicsBodyModifierRecord* Record = FindBodyModifierRecord(Name);
+		if (Record)
+		{
+			Record->BodyModifier.ModifierData.UpdateFromSparseData(ModifierData);
+		}
 	}
-	if (bWarnAboutInvalidNames)
+	if (bApplyToSetsWithName)
 	{
-		UE_LOG(LogPhysicsControl, Warning,
-			TEXT("SetBodyModifierData - invalid name %s"), *Name.ToString());
+		SetBodyModifierSparseDatasInSet(Name, ModifierData);
 	}
-	return false;
 }
 
 //======================================================================================================================
 void UPhysicsControlComponent::SetBodyModifierSparseDatas(
 	const TArray<FName>&                    Names,
-	const FPhysicsControlModifierSparseData ModifierData)
+	const FPhysicsControlModifierSparseData ModifierData,
+	const bool                              bApplyToModifiersWithName,
+	const bool                              bApplyToSetsWithName)
 {
 	for (FName Name : Names)
 	{
-		SetBodyModifierSparseData(Name, ModifierData);
+		SetBodyModifierSparseData(Name, ModifierData, bApplyToModifiersWithName, bApplyToSetsWithName);
 	}
 }
 
@@ -1774,10 +1873,9 @@ void UPhysicsControlComponent::SetBodyModifierSparseDatasInSet(
 	const FName                             SetName,
 	const FPhysicsControlModifierSparseData ModifierData)
 {
-	SetBodyModifierSparseDatas(GetControlNamesInSet(SetName), ModifierData);
+	SetBodyModifierSparseDatas(GetBodyModifierNamesInSet(SetName), ModifierData, true, false);
 }
 
-
 //======================================================================================================================
 bool UPhysicsControlComponent::SetBodyModifierKinematicTarget(
 	const FName    Name,
@@ -1805,32 +1903,36 @@ bool UPhysicsControlComponent::SetBodyModifierKinematicTarget(
 }
 
 //======================================================================================================================
-bool UPhysicsControlComponent::SetBodyModifierMovementType(
+void UPhysicsControlComponent::SetBodyModifierMovementType(
 	const FName                Name,
-	const EPhysicsMovementType MovementType)
+	const EPhysicsMovementType MovementType,
+	const bool                 bApplyToModifiersWithName,
+	const bool                 bApplyToSetsWithName)
 {
-	FPhysicsBodyModifierRecord* Record = FindBodyModifierRecord(Name);
-	if (Record)
+	if (bApplyToModifiersWithName)
 	{
-		Record->BodyModifier.ModifierData.MovementType = MovementType;
-		return true;
+		FPhysicsBodyModifierRecord* Record = FindBodyModifierRecord(Name);
+		if (Record)
+		{
+			Record->BodyModifier.ModifierData.MovementType = MovementType;
+		}
 	}
-	if (bWarnAboutInvalidNames)
+	if (bApplyToSetsWithName)
 	{
-		UE_LOG(LogPhysicsControl, Warning,
-			TEXT("SetBodyModifierMovementType - invalid name %s"), *Name.ToString());
+		SetBodyModifiersInSetMovementType(Name, MovementType);
 	}
-	return false;
 }
 
 //======================================================================================================================
 void UPhysicsControlComponent::SetBodyModifiersMovementType(
 	const TArray<FName>&       Names,
-	const EPhysicsMovementType MovementType)
+	const EPhysicsMovementType MovementType,
+	const bool                 bApplyToModifiersWithName,
+	const bool                 bApplyToSetsWithName)
 {
 	for (FName Name : Names)
 	{
-		SetBodyModifierMovementType(Name, MovementType);
+		SetBodyModifierMovementType(Name, MovementType, bApplyToModifiersWithName, bApplyToSetsWithName);
 	}
 }
 
@@ -1839,36 +1941,40 @@ void UPhysicsControlComponent::SetBodyModifiersInSetMovementType(
 	const FName                SetName,
 	const EPhysicsMovementType MovementType)
 {
-	SetBodyModifiersMovementType(GetBodyModifierNamesInSet(SetName), MovementType);
+	SetBodyModifiersMovementType(GetBodyModifierNamesInSet(SetName), MovementType, true, false);
 }
 
 //======================================================================================================================
-bool UPhysicsControlComponent::SetBodyModifierCollisionType(
+void UPhysicsControlComponent::SetBodyModifierCollisionType(
 	const FName                   Name,
-	const ECollisionEnabled::Type CollisionType)
+	const ECollisionEnabled::Type CollisionType,
+	const bool                    bApplyToModifiersWithName,
+	const bool                    bApplyToSetsWithName)
 {
-	FPhysicsBodyModifierRecord* Record = FindBodyModifierRecord(Name);
-	if (Record)
+	if (bApplyToModifiersWithName)
 	{
-		Record->BodyModifier.ModifierData.CollisionType = CollisionType;
-		return true;
+		FPhysicsBodyModifierRecord* Record = FindBodyModifierRecord(Name);
+		if (Record)
+		{
+			Record->BodyModifier.ModifierData.CollisionType = CollisionType;
+		}
 	}
-	if (bWarnAboutInvalidNames)
+	if (bApplyToSetsWithName)
 	{
-		UE_LOG(LogPhysicsControl, Warning,
-			TEXT("SetBodyModifierCollisionType - invalid name %s"), *Name.ToString());
+		SetBodyModifiersInSetCollisionType(Name, CollisionType);
 	}
-	return false;
 }
 
 //======================================================================================================================
 void UPhysicsControlComponent::SetBodyModifiersCollisionType(
 	const TArray<FName>&          Names,
-	const ECollisionEnabled::Type CollisionType)
+	const ECollisionEnabled::Type CollisionType,
+	const bool                    bApplyToModifiersWithName,
+	const bool                    bApplyToSetsWithName)
 {
 	for (FName Name : Names)
 	{
-		SetBodyModifierCollisionType(Name, CollisionType);
+		SetBodyModifierCollisionType(Name, CollisionType, bApplyToModifiersWithName, bApplyToSetsWithName);
 	}
 }
 
@@ -1877,36 +1983,40 @@ void UPhysicsControlComponent::SetBodyModifiersInSetCollisionType(
 	const FName                   SetName,
 	const ECollisionEnabled::Type CollisionType)
 {
-	SetBodyModifiersCollisionType(GetBodyModifierNamesInSet(SetName), CollisionType);
+	SetBodyModifiersCollisionType(GetBodyModifierNamesInSet(SetName), CollisionType, true, false);
 }
 
 //======================================================================================================================
-bool UPhysicsControlComponent::SetBodyModifierGravityMultiplier(
+void UPhysicsControlComponent::SetBodyModifierGravityMultiplier(
 	const FName Name,
-	const float GravityMultiplier)
+	const float GravityMultiplier,
+	const bool  bApplyToModifiersWithName,
+	const bool  bApplyToSetsWithName)
 {
-	FPhysicsBodyModifierRecord* Record = FindBodyModifierRecord(Name);
-	if (Record)
+	if (bApplyToModifiersWithName)
 	{
-		Record->BodyModifier.ModifierData.GravityMultiplier = GravityMultiplier;
-		return true;
+		FPhysicsBodyModifierRecord* Record = FindBodyModifierRecord(Name);
+		if (Record)
+		{
+			Record->BodyModifier.ModifierData.GravityMultiplier = GravityMultiplier;
+		}
 	}
-	if (bWarnAboutInvalidNames)
+	if (bApplyToSetsWithName)
 	{
-		UE_LOG(LogPhysicsControl, Warning,
-			TEXT("SetBodyModifierGravityMultiplier - invalid name %s"), *Name.ToString());
+		SetBodyModifiersInSetGravityMultiplier(Name, GravityMultiplier);
 	}
-	return false;
 }
 
 //======================================================================================================================
 void UPhysicsControlComponent::SetBodyModifiersGravityMultiplier(
 	const TArray<FName>& Names,
-	const float          GravityMultiplier)
+	const float          GravityMultiplier,
+	const bool           bApplyToModifiersWithName,
+	const bool           bApplyToSetsWithName)
 {
 	for (FName Name : Names)
 	{
-		SetBodyModifierGravityMultiplier(Name, GravityMultiplier);
+		SetBodyModifierGravityMultiplier(Name, GravityMultiplier, bApplyToModifiersWithName, bApplyToSetsWithName);
 	}
 }
 
@@ -1915,36 +2025,40 @@ void UPhysicsControlComponent::SetBodyModifiersInSetGravityMultiplier(
 	const FName SetName,
 	const float GravityMultiplier)
 {
-	SetBodyModifiersGravityMultiplier(GetBodyModifierNamesInSet(SetName), GravityMultiplier);
+	SetBodyModifiersGravityMultiplier(GetBodyModifierNamesInSet(SetName), GravityMultiplier, true, false);
 }
 
 //======================================================================================================================
-bool UPhysicsControlComponent::SetBodyModifierPhysicsBlendWeight(
+void UPhysicsControlComponent::SetBodyModifierPhysicsBlendWeight(
 	const FName Name,
-	const float PhysicsBlendWeight)
+	const float PhysicsBlendWeight,
+	const bool  bApplyToModifiersWithName,
+	const bool  bApplyToSetsWithName)
 {
-	FPhysicsBodyModifierRecord* Record = FindBodyModifierRecord(Name);
-	if (Record)
+	if (bApplyToModifiersWithName)
 	{
-		Record->BodyModifier.ModifierData.PhysicsBlendWeight = PhysicsBlendWeight;
-		return true;
+		FPhysicsBodyModifierRecord* Record = FindBodyModifierRecord(Name);
+		if (Record)
+		{
+			Record->BodyModifier.ModifierData.PhysicsBlendWeight = PhysicsBlendWeight;
+		}
 	}
-	if (bWarnAboutInvalidNames)
+	if (bApplyToSetsWithName)
 	{
-		UE_LOG(LogPhysicsControl, Warning,
-			TEXT("SetBodyModifierPhysicsBlendWeight - invalid name %s"), *Name.ToString());
+		SetBodyModifiersInSetPhysicsBlendWeight(Name, PhysicsBlendWeight);
 	}
-	return false;
 }
 
 //======================================================================================================================
 void UPhysicsControlComponent::SetBodyModifiersPhysicsBlendWeight(
 	const TArray<FName>& Names,
-	const float          PhysicsBlendWeight)
+	const float          PhysicsBlendWeight,
+	const bool           bApplyToModifiersWithName,
+	const bool           bApplyToSetsWithName)
 {
 	for (FName Name : Names)
 	{
-		SetBodyModifierPhysicsBlendWeight(Name, PhysicsBlendWeight);
+		SetBodyModifierPhysicsBlendWeight(Name, PhysicsBlendWeight, bApplyToModifiersWithName, bApplyToSetsWithName);
 	}
 }
 
@@ -1953,36 +2067,40 @@ void UPhysicsControlComponent::SetBodyModifiersInSetPhysicsBlendWeight(
 	const FName SetName,
 	const float PhysicsBlendWeight)
 {
-	SetBodyModifiersPhysicsBlendWeight(GetBodyModifierNamesInSet(SetName), PhysicsBlendWeight);
+	SetBodyModifiersPhysicsBlendWeight(GetBodyModifierNamesInSet(SetName), PhysicsBlendWeight, true, false);
 }
 
 //======================================================================================================================
-bool UPhysicsControlComponent::SetBodyModifierKinematicTargetSpace(
+void UPhysicsControlComponent::SetBodyModifierKinematicTargetSpace(
 	const FName                               Name,
-	const EPhysicsControlKinematicTargetSpace KinematicTargetSpace)
+	const EPhysicsControlKinematicTargetSpace KinematicTargetSpace,
+	const bool                                bApplyToModifiersWithName,
+	const bool                                bApplyToSetsWithName)
 {
-	FPhysicsBodyModifierRecord* Record = FindBodyModifierRecord(Name);
-	if (Record)
+	if (bApplyToModifiersWithName)
 	{
-		Record->BodyModifier.ModifierData.KinematicTargetSpace = KinematicTargetSpace;
-		return true;
+		FPhysicsBodyModifierRecord* Record = FindBodyModifierRecord(Name);
+		if (Record)
+		{
+			Record->BodyModifier.ModifierData.KinematicTargetSpace = KinematicTargetSpace;
+		}
 	}
-	if (bWarnAboutInvalidNames)
+	if (bApplyToSetsWithName)
 	{
-		UE_LOG(LogPhysicsControl, Warning,
-			TEXT("SetBodyModifierKinematicTargetSpace - invalid name %s"), *Name.ToString());
+		SetBodyModifiersInSetKinematicTargetSpace(Name, KinematicTargetSpace);
 	}
-	return false;
 }
 
 //======================================================================================================================
 void UPhysicsControlComponent::SetBodyModifiersKinematicTargetSpace(
 	const TArray<FName>&                      Names,
-	const EPhysicsControlKinematicTargetSpace KinematicTargetSpace)
+	const EPhysicsControlKinematicTargetSpace KinematicTargetSpace,
+	const bool                                bApplyToModifiersWithName,
+	const bool                                bApplyToSetsWithName)
 {
 	for (FName Name : Names)
 	{
-		SetBodyModifierKinematicTargetSpace(Name, KinematicTargetSpace);
+		SetBodyModifierKinematicTargetSpace(Name, KinematicTargetSpace, bApplyToModifiersWithName, bApplyToSetsWithName);
 	}
 }
 
@@ -1991,36 +2109,41 @@ void UPhysicsControlComponent::SetBodyModifiersInSetKinematicTargetSpace(
 	const FName                               SetName,
 	const EPhysicsControlKinematicTargetSpace KinematicTargetSpace)
 {
-	SetBodyModifiersKinematicTargetSpace(GetBodyModifierNamesInSet(SetName), KinematicTargetSpace);
+	SetBodyModifiersKinematicTargetSpace(GetBodyModifierNamesInSet(SetName), KinematicTargetSpace, true, false);
 }
 
 //======================================================================================================================
-bool UPhysicsControlComponent::SetBodyModifierUpdateKinematicFromSimulation(
+void UPhysicsControlComponent::SetBodyModifierUpdateKinematicFromSimulation(
 	const FName Name,
-	const bool  bUpdateKinematicFromSimulation)
+	const bool  bUpdateKinematicFromSimulation,
+	const bool  bApplyToModifiersWithName,
+	const bool  bApplyToSetsWithName)
 {
-	FPhysicsBodyModifierRecord* Record = FindBodyModifierRecord(Name);
-	if (Record)
+	if (bApplyToModifiersWithName)
 	{
-		Record->BodyModifier.ModifierData.bUpdateKinematicFromSimulation = bUpdateKinematicFromSimulation;
-		return true;
+		FPhysicsBodyModifierRecord* Record = FindBodyModifierRecord(Name);
+		if (Record)
+		{
+			Record->BodyModifier.ModifierData.bUpdateKinematicFromSimulation = bUpdateKinematicFromSimulation;
+		}
 	}
-	if (bWarnAboutInvalidNames)
+	if (bApplyToSetsWithName)
 	{
-		UE_LOG(LogPhysicsControl, Warning,
-			TEXT("SetBodyModifierUpdateKinematicFromSimulation - invalid name %s"), *Name.ToString());
+		SetBodyModifiersInSetUpdateKinematicFromSimulation(Name, bUpdateKinematicFromSimulation);
 	}
-	return false;
 }
 
 //======================================================================================================================
 void UPhysicsControlComponent::SetBodyModifiersUpdateKinematicFromSimulation(
 	const TArray<FName>& Names,
-	const bool           bUpdateKinematicFromSimulation)
+	const bool           bUpdateKinematicFromSimulation,
+	const bool           bApplyToModifiersWithName,
+	const bool           bApplyToSetsWithName)
 {
 	for (FName Name : Names)
 	{
-		SetBodyModifierUpdateKinematicFromSimulation(Name, bUpdateKinematicFromSimulation);
+		SetBodyModifierUpdateKinematicFromSimulation(
+			Name, bUpdateKinematicFromSimulation, bApplyToModifiersWithName, bApplyToSetsWithName);
 	}
 }
 
@@ -2029,7 +2152,8 @@ void UPhysicsControlComponent::SetBodyModifiersInSetUpdateKinematicFromSimulatio
 	const FName SetName,
 	const bool  bUpdateKinematicFromSimulation)
 {
-	SetBodyModifiersUpdateKinematicFromSimulation(GetBodyModifierNamesInSet(SetName), bUpdateKinematicFromSimulation);
+	SetBodyModifiersUpdateKinematicFromSimulation(
+		GetBodyModifierNamesInSet(SetName), bUpdateKinematicFromSimulation, true, false);
 }
 
 //======================================================================================================================
@@ -2533,39 +2657,43 @@ void UPhysicsControlComponent::SetCachedBoneVelocitiesToZero()
 }
 
 //======================================================================================================================
-bool UPhysicsControlComponent::ResetBodyModifierToCachedBoneTransform(
+void UPhysicsControlComponent::ResetBodyModifierToCachedBoneTransform(
 	const FName                        Name,
-	const EResetToCachedTargetBehavior Behavior)
+	const EResetToCachedTargetBehavior Behavior,
+	const bool                         bApplyToModifiersWithName,
+	const bool                         bApplyToSetsWithName)
 {
-	FPhysicsBodyModifierRecord* Record = FindBodyModifierRecord(Name);
-	if (Record)
+	if (bApplyToModifiersWithName)
 	{
-		if (Behavior == EResetToCachedTargetBehavior::ResetImmediately)
+		FPhysicsBodyModifierRecord* Record = FindBodyModifierRecord(Name);
+		if (Record)
 		{
-			ResetToCachedTarget(*Record);
-		}
-		else
-		{
-			Record->bResetToCachedTarget = true;
+			if (Behavior == EResetToCachedTargetBehavior::ResetImmediately)
+			{
+				ResetToCachedTarget(*Record);
+			}
+			else
+			{
+				Record->bResetToCachedTarget = true;
+			}
 		}
-		return true;
 	}
-	if (bWarnAboutInvalidNames)
+	if (bApplyToSetsWithName)
 	{
-		UE_LOG(LogPhysicsControl, Warning,
-			TEXT("ResetBodyModifierToCachedBoneTransform - invalid name %s"), *Name.ToString());
+		ResetBodyModifiersInSetToCachedBoneTransforms(Name, Behavior);
 	}
-	return false;
 }
 
 //======================================================================================================================
 void UPhysicsControlComponent::ResetBodyModifiersToCachedBoneTransforms(
 	const TArray<FName>&               Names,
-	const EResetToCachedTargetBehavior Behavior)
+	const EResetToCachedTargetBehavior Behavior,
+	const bool                         bApplyToModifiersWithName,
+	const bool                         bApplyToSetsWithName)
 {
 	for (FName Name : Names)
 	{
-		ResetBodyModifierToCachedBoneTransform(Name, Behavior);
+		ResetBodyModifierToCachedBoneTransform(Name, Behavior, bApplyToModifiersWithName, bApplyToSetsWithName);
 	}
 }
 
@@ -2574,7 +2702,7 @@ void UPhysicsControlComponent::ResetBodyModifiersInSetToCachedBoneTransforms(
 	const FName                        SetName,
 	const EResetToCachedTargetBehavior Behavior)
 {
-	ResetBodyModifiersToCachedBoneTransforms(GetBodyModifierNamesInSet(SetName), Behavior);
+	ResetBodyModifiersToCachedBoneTransforms(GetBodyModifierNamesInSet(SetName), Behavior, true, false);
 }
 
 //======================================================================================================================
diff --git a/Engine/Plugins/Experimental/PhysicsControl/Source/PhysicsControl/Public/PhysicsControlComponent.h b/Engine/Plugins/Experimental/PhysicsControl/Source/PhysicsControl/Public/PhysicsControlComponent.h
--- a/Engine/Plugins/Experimental/PhysicsControl/Source/PhysicsControl/Public/PhysicsControlComponent.h
+++ b/Engine/Plugins/Experimental/PhysicsControl/Source/PhysicsControl/Public/PhysicsControlComponent.h
@@ -310,22 +310,31 @@ public:
 	UE_API void DestroyAllControlsAndBodyModifiers();
 
 	/**
-	 * Destroys a control
+	 * Destroys a control or controls in a set
+	 * Standard set names will include "All", "WorldSpace", "ParentSpace" and things like "WorldSpace-ArmLeft",
+	 * depending on how controls have been created.
 	 *
-	 * @param Name The name of the control to destroy. 
-	 * @return     Returns true if the control was found and destroyed, false if not
+	 * @param Name The name of the control or set of controls to destroy. 
+	 * @param bApplyToControlsWithName If true, then controls with the name will be modified
+	 * @param bApplyToSetsWithName If true, then controls in sets with the name will be modified
 	 */
 	UFUNCTION(BlueprintCallable, Category = PhysicsControl, Meta = (ReturnDisplayName = "Success"))
-	UE_API bool DestroyControl(const FName Name);
+	UE_API void DestroyControl(
+		const FName Name, const bool bApplyToControlsWithName = true, const bool bApplyToSetsWithName = true);
 
 	/**
-	 * Destroys all controls
+	 * Destroys all controls or sets of controls
+	 * Standard set names will include "All", "WorldSpace", "ParentSpace" and things like "WorldSpace-ArmLeft",
+	 * depending on how controls have been created.
 	 *
-	 * @param Names The names of the controls to destroy. Note that if you have these in a FPhysicsControlNameArray
-	 *              then it can be split.
+	 * @param Names The names of the controls or sets of controls to destroy. Note that if you have these 
+	 *        in a FPhysicsControlNameArray then it can be split.
+	 * @param bApplyToControlsWithName If true, then controls with the name will be modified
+	 * @param bApplyToSetsWithName If true, then controls in sets with the name will be modified
 	 */
 	UFUNCTION(BlueprintCallable, Category = PhysicsControl)
-	UE_API void DestroyControls(const TArray<FName>& Names);
+	UE_API void DestroyControls(
+		const TArray<FName>& Names, const bool bApplyToControlsWithName = true, const bool bApplyToSetsWithName = true);
 
 	/**
 	 * Destroys all controls in a set
@@ -337,24 +346,32 @@ public:
 	UE_API void DestroyControlsInSet(const FName Set);
 
 	/**
-	 * Updates the parent object part of a control. Note that this won't change the name of the control (which may
-	 * subsequently be misleading), or any set it is included in, etc.
+	 * Updates the parent object part of a control or controls in a set. Note that this won't change 
+	 * the name of the control (which may subsequently be misleading), or any set it is included in, etc.
+	 * Standard set names will include "All", "WorldSpace", "ParentSpace" and things like "WorldSpace-ArmLeft",
+	 * depending on how controls have been created.
+	 * 
+	 * @param bApplyToControlsWithName If true, then controls with the name will be modified
+	 * @param bApplyToSetsWithName If true, then controls in sets with the name will be modified
 	 */
 	UFUNCTION(BlueprintCallable, Category = PhysicsControl, Meta = (ReturnDisplayName = "Success"))
-	UE_API bool SetControlParent(
-		const FName          Name, 
-		UPrimitiveComponent* ParentComponent,
-		const FName          ParentBoneName);
+	UE_API void SetControlParent(
+		const FName Name, UPrimitiveComponent* ParentComponent, const FName ParentBoneName,
+		const bool bApplyToControlsWithName = true, const bool bApplyToSetsWithName = true);
 
 	/**
-	 * Updates the parent object part of controls. Note that this won't change the name of the controls (which may
-	 * subsequently be misleading), or any set they are included in, etc.
+	 * Updates the parent object part of controls or controls in a set. Note that this won't change 
+	 * the name of the controls (which may subsequently be misleading), or any set they are included in, etc.
+	 * Standard set names will include "All", "WorldSpace", "ParentSpace" and things like "WorldSpace-ArmLeft",
+	 * depending on how controls have been created.
+	 * 
+	 * @param bApplyToControlsWithName If true, then controls with the name will be modified
+	 * @param bApplyToSetsWithName If true, then controls in sets with the name will be modified
 	 */
 	UFUNCTION(BlueprintCallable, Category = PhysicsControl)
 	UE_API void SetControlParents(
-		const TArray<FName>& Names,
-		UPrimitiveComponent* ParentComponent,
-		const FName          ParentBoneName);
+		const TArray<FName>& Names, UPrimitiveComponent* ParentComponent, const FName ParentBoneName,
+		const bool bApplyToControlsWithName = true, const bool bApplyToSetsWithName = true);
 
 	UFUNCTION(BlueprintCallable, Category = PhysicsControl)
 	UE_API void SetControlParentsInSet(
@@ -363,24 +380,37 @@ public:
 		const FName          ParentBoneName);
 
 	/**
-	 * Modifies an existing control data - i.e. the strengths etc of the control driving towards the target
+	 * Modifies an existing control data in a control or controls in a set - i.e. the strengths 
+	 * etc of the control driving towards the target
+	 * Standard set names will include "All", "WorldSpace", "ParentSpace" and things like "WorldSpace-ArmLeft",
+	 * depending on how controls have been created.
 	 *
-	 * @param Name The name of the control to update. 
+	 * @param Name The name of the control or set of controls to update. 
 	 * @param ControlData The new control data
-	 * @return true if the control was found and modified, false if not
+	 * @param bApplyToControlsWithName If true, then controls with the name will be modified
+	 * @param bApplyToSetsWithName If true, then controls in sets with the name will be modified
 	 */
 	UFUNCTION(BlueprintCallable, Category = PhysicsControl, Meta = (ReturnDisplayName = "Success"))
-	UE_API bool SetControlData(const FName Name, const FPhysicsControlData ControlData);
+	UE_API void SetControlData(
+		const FName Name, const FPhysicsControlData ControlData, 
+		const bool bApplyToControlsWithName = true, const bool bApplyToSetsWithName = true);
 
 	/**
-	 * Modifies existing control data - i.e. the strengths etc of the controls driving towards the targets
+	 * Modifies existing control data in controls or controls in sets - i.e. the strengths 
+	 * etc of the controls driving towards the targets
+	 * Standard set names will include "All", "WorldSpace", "ParentSpace" and things like "WorldSpace-ArmLeft",
+	 * depending on how controls have been created.
 	 *
-	 * @param Names The names of the controls to update. Note that if you have these in a FPhysicsControlNameArray
+	 * @param Names The names of the controls or sets of controls to update. Note that if you have these in a FPhysicsControlNameArray
 	 *              then it can be split.
 	 * @param ControlData The new control data
+	 * @param bApplyToControlsWithName If true, then controls with the name will be modified
+	 * @param bApplyToSetsWithName If true, then controls in sets with the name will be modified
 	 */
 	UFUNCTION(BlueprintCallable, Category = PhysicsControl)
-	UE_API void SetControlDatas(const TArray<FName>& Names, const FPhysicsControlData ControlData);
+	UE_API void SetControlDatas(
+		const TArray<FName>& Names, const FPhysicsControlData ControlData,
+		const bool bApplyToControlsWithName = true, const bool bApplyToSetsWithName = true);
 
 	/**
 	 * Modifies existing control data - i.e. the strengths etc of the controls driving towards the targets
@@ -393,24 +423,37 @@ public:
 	UE_API void SetControlDatasInSet(const FName Set, const FPhysicsControlData ControlData);
 
 	/**
-	 * Modifies an existing control data - i.e. the strengths etc of the control driving towards the target
+	 * Modifies an existing control data in a control or controls in a set - i.e. the strengths 
+	 * etc of the control driving towards the target
+	 * Standard set names will include "All", "WorldSpace", "ParentSpace" and things like "WorldSpace-ArmLeft",
+	 * depending on how controls have been created.
 	 *
-	 * @param Name The name of the control to update. 
+	 * @param Name The name of the control or set of controls to update. 
 	 * @param ControlData The new control data
-	 * @return true if the control was found and modified, false if not
+	 * @param bApplyToControlsWithName If true, then controls with the name will be modified
+	 * @param bApplyToSetsWithName If true, then controls in sets with the name will be modified
 	 */
 	UFUNCTION(BlueprintCallable, Category = PhysicsControl, Meta = (ReturnDisplayName = "Success"))
-	UE_API bool SetControlSparseData(const FName Name, const FPhysicsControlSparseData ControlData);
+	UE_API void SetControlSparseData(
+		const FName Name, const FPhysicsControlSparseData ControlData,
+		const bool bApplyToControlsWithName = true, const bool bApplyToSetsWithName = true);
 
 	/**
-	 * Modifies existing control data - i.e. the strengths etc of the controls driving towards the targets
+	 * Modifies existing control data in controls or controls in sets - i.e. the strengths 
+	 * etc of the controls driving towards the targets
+	 * Standard set names will include "All", "WorldSpace", "ParentSpace" and things like "WorldSpace-ArmLeft",
+	 * depending on how controls have been created.
 	 *
-	 * @param Names The names of the controls to update. Note that if you have these in a FPhysicsControlNameArray
+	 * @param Names The names of the controls or sets of controls to update. Note that if you have these in a FPhysicsControlNameArray
 	 *              then it can be split.
 	 * @param ControlData The new control data
+	 * @param bApplyToControlsWithName If true, then controls with the name will be modified
+	 * @param bApplyToSetsWithName If true, then controls in sets with the name will be modified
 	 */
 	UFUNCTION(BlueprintCallable, Category = PhysicsControl)
-	UE_API void SetControlSparseDatas(const TArray<FName>& Names, const FPhysicsControlSparseData ControlData);
+	UE_API void SetControlSparseDatas(
+		const TArray<FName>& Names, const FPhysicsControlSparseData ControlData,
+		const bool bApplyToControlsWithName = true, const bool bApplyToSetsWithName = true);
 
 	/**
 	 * Modifies existing control data - i.e. the strengths etc of the controls driving towards the targets
@@ -423,75 +466,87 @@ public:
 	UE_API void SetControlSparseDatasInSet(const FName Set, const FPhysicsControlSparseData ControlData);
 
 	/**
-	 * Modifies an existing control data using the multipliers
+	 * Modifies existing control data in a control or controls in a set using the multipliers
+	 * Standard set names will include "All", "WorldSpace", "ParentSpace" and things like "WorldSpace-ArmLeft",
+	 * depending on how controls have been created.
 	 *
-	 * @param Name The name of the control to modify. 
+	 * @param Name The name of the control or set of controls to update. 
 	 * @param ControlMultipliers The new control multipliers
 	 * @param bEnableControl Enables the control if it is currently disabled
-	 * @return true if the control was found and modified, false if not
+	 * @param bApplyToControlsWithName If true, then controls with the name will be modified
+	 * @param bApplyToSetsWithName If true, then controls in sets with the name will be modified
 	 */
 	UFUNCTION(BlueprintCallable, Category = PhysicsControl, Meta = (ReturnDisplayName = "Success"))
-	UE_API bool SetControlMultiplier
-	(const FName Name, const FPhysicsControlMultiplier ControlMultiplier, const bool bEnableControl = true);
+	UE_API void SetControlMultiplier(
+		const FName Name, const FPhysicsControlMultiplier ControlMultiplier, const bool bEnableControl = true,
+		const bool bApplyToControlsWithName = true, const bool bApplyToSetsWithName = true);
 
 	/**
-	 * Modifies existing control data using the multipliers
+	 * Modifies existing control data in controls or controls in sets using the multipliers
+	 * Standard set names will include "All", "WorldSpace", "ParentSpace" and things like "WorldSpace-ArmLeft",
+	 * depending on how controls have been created.
 	 *
-	 * @param Names The names of the controls to modify. Note that if you have these in a FPhysicsControlNameArray
+	 * @param Names The names of the controls or sets of controls to update. Note that if you have these in a FPhysicsControlNameArray
 	 *              then it can be split.
 	 * @param ControlMultiplier The new control multiplier
 	 * @param bEnableControl Enables the controls if currently disabled
+	 * @param bApplyToControlsWithName If true, then controls with the name will be modified
+	 * @param bApplyToSetsWithName If true, then controls in sets with the name will be modified
 	 */
 	UFUNCTION(BlueprintCallable, Category = PhysicsControl)
 	UE_API void SetControlMultipliers(
-		const TArray<FName>&            Names, 
-		const FPhysicsControlMultiplier ControlMultiplier, 
-		const bool                      bEnableControl = true);
+		const TArray<FName>& Names, const FPhysicsControlMultiplier ControlMultiplier, const bool bEnableControl = true,
+		const bool bApplyToControlsWithName = true, const bool bApplyToSetsWithName = true);
 
 	/**
 	 * Modifies existing control data using the multipliers
 	 *
-	 * @param Set The set of controls to modify. Standard sets will include "All", "WorldSpace",
+	 * @param Set The set of controls to update. Standard sets will include "All", "WorldSpace",
 	 *        "ParentSpace" and things like "WorldSpace-ArmLeft", depending on how controls have been created.
 	 * @param ControlMultiplier The new control multiplier
 	 * @param bEnableControl Enables the controls if currently disabled
 	 */
 	UFUNCTION(BlueprintCallable, Category = PhysicsControl)
 	UE_API void SetControlMultipliersInSet(
-		const FName                      Set, 
-		const FPhysicsControlMultiplier  ControlMultiplier, 
-		const bool                       bEnableControl = true);
+		const FName Set, const FPhysicsControlMultiplier  ControlMultiplier, const bool bEnableControl = true);
 
 	/**
-	 * Modifies an existing control data using the multipliers
+	 * Modifies existing control data in a control or controls in a set using the multipliers
+	 * Standard set names will include "All", "WorldSpace", "ParentSpace" and things like "WorldSpace-ArmLeft",
+	 * depending on how controls have been created.
 	 *
-	 * @param Name The name of the control to modify. 
+	 * @param Name The name of the control or set of controls to update. 
 	 * @param ControlMultipliers The new control multipliers
 	 * @param bEnableControl Enables the control if it is currently disabled
-	 * @return true if the control was found and modified, false if not
+	 * @param bApplyToControlsWithName If true, then controls with the name will be modified
+	 * @param bApplyToSetsWithName If true, then controls in sets with the name will be modified
 	 */
 	UFUNCTION(BlueprintCallable, Category = PhysicsControl, Meta = (ReturnDisplayName = "Success"))
-	UE_API bool SetControlSparseMultiplier(
-		const FName Name, const FPhysicsControlSparseMultiplier ControlMultiplier, const bool bEnableControl = true);
+	UE_API void SetControlSparseMultiplier(
+		const FName Name, const FPhysicsControlSparseMultiplier ControlMultiplier, const bool bEnableControl = true,
+		const bool bApplyToControlsWithName = true, const bool bApplyToSetsWithName = true);
 
 	/**
-	 * Modifies existing control data using the multipliers
+	 * Modifies existing control data in controls or controls in sets using the multipliers
+	 * Standard set names will include "All", "WorldSpace", "ParentSpace" and things like "WorldSpace-ArmLeft",
+	 * depending on how controls have been created.
 	 *
-	 * @param Names The names of the controls to modify. Note that if you have these in a FPhysicsControlNameArray
-	 *              then it can be split.
+	 * @param Names The names of the controls or sets of controls to update. Note that if you have these in a 
+	 *        FPhysicsControlNameArray then it can be split.
 	 * @param ControlMultiplier The new control multiplier
 	 * @param bEnableControl Enables the controls if currently disabled
+	 * @param bApplyToControlsWithName If true, then controls with the name will be modified
+	 * @param bApplyToSetsWithName If true, then controls in sets with the name will be modified
 	 */
 	UFUNCTION(BlueprintCallable, Category = PhysicsControl)
 	UE_API void SetControlSparseMultipliers(
-		const TArray<FName>&                  Names, 
-		const FPhysicsControlSparseMultiplier ControlMultiplier, 
-		const bool                            bEnableControl = true);
+		const TArray<FName>& Names, const FPhysicsControlSparseMultiplier ControlMultiplier, const bool bEnableControl = true,
+		const bool bApplyToControlsWithName = true, const bool bApplyToSetsWithName = true);
 
 	/**
 	 * Modifies existing control data using the multipliers
 	 *
-	 * @param Set The set of controls to modify. Standard sets will include "All", "WorldSpace",
+	 * @param Set The set of controls to update. Standard sets will include "All", "WorldSpace",
 	 *        "ParentSpace" and things like "WorldSpace-ArmLeft", depending on how controls have been created.
 	 * @param ControlMultiplier The new control multiplier
 	 * @param bEnableControl Enables the controls if currently disabled
@@ -503,49 +558,61 @@ public:
 		const bool                            bEnableControl = true);
 
 	/**
-	 * Modifies an existing control's linear data - i.e. the strengths etc of the control driving towards the target
+	 * Modifies existing control's linear data in a control or controls in a set - i.e. the strengths 
+	 * etc of the control driving towards the target
+	 * Standard set names will include "All", "WorldSpace", "ParentSpace" and things like "WorldSpace-ArmLeft",
+	 * depending on how controls have been created.
 	 *
-	 * @param Name The name of the control to modify. 
+	 * @param Name The name of the control or set of controls to update. 
 	 * @param Strength The strength used to drive linear motion
 	 * @param DampingRatio The amount of damping associated with the linear strength. 1 Results in critically damped motion
 	 * @param ExtraDamping The amount of additional linear damping
 	 * @param MaxForce The maximum force used to drive the linear motion. Zero indicates no limit.
 	 * @param bEnableControl Enables the control if it is currently disabled
-	 * @return true if the control was found and modified, false if not
+	 * @param bApplyToControlsWithName If true, then controls with the name will be modified
+	 * @param bApplyToSetsWithName If true, then controls in sets with the name will be modified
 	 */
 	UFUNCTION(BlueprintCallable, Category = PhysicsControl, Meta = (ReturnDisplayName = "Success"))
-	UE_API bool SetControlLinearData(
+	UE_API void SetControlLinearData(
 		const FName Name,
 		const float Strength = 1.0f, 
 		const float DampingRatio = 1.0f, 
 		const float ExtraDamping = 0.0f, 
 		const float MaxForce = 0.0f, 
-		const bool  bEnableControl = true);
+		const bool  bEnableControl = true,
+		const bool  bApplyToControlsWithName = true, 
+		const bool  bApplyToSetsWithName = true);
 
 	/**
-	 * Modifies an existing control's angular data - i.e. the strengths etc of the control driving towards the target
+	 * Modifies existing control's angular data in a control or controls in a set - i.e. the strengths 
+	 * etc of the control driving towards the target
+	 * Standard set names will include "All", "WorldSpace", "ParentSpace" and things like "WorldSpace-ArmLeft",
+	 * depending on how controls have been created.
 	 *
-	 * @param Name The name of the control to modify. 
+	 * @param Name The name of the control or set of controls to update. 
 	 * @param Strength The strength used to drive angular motion
 	 * @param DampingRatio The amount of damping associated with the angular strength. 1 Results in critically damped motion
 	 * @param ExtraDamping The amount of additional angular damping
 	 * @param MaxTorque The maximum torque used to drive the angular motion. Zero indicates no limit.
 	 * @param bEnableControl Enables the control if it is currently disabled
-	 * @return true if the control was found and modified, false if not
+	 * @param bApplyToControlsWithName If true, then controls with the name will be modified
+	 * @param bApplyToSetsWithName If true, then controls in sets with the name will be modified
 	 */
 	UFUNCTION(BlueprintCallable, Category = PhysicsControl, Meta = (ReturnDisplayName = "Success"))
-	UE_API bool SetControlAngularData(
+	UE_API void SetControlAngularData(
 		const FName Name,
 		const float Strength = 1.0f, 
 		const float DampingRatio = 1.0f, 
 		const float ExtraDamping = 0.0f, 
 		const float MaxTorque = 0.0f, 
-		const bool  bEnableControl = true);
+		const bool  bEnableControl = true,
+		const bool  bApplyToControlsWithName = true,
+		const bool  bApplyToSetsWithName = true);
 
 	/**
 	 * Sets the point at which controls will "push" the child object.
 	 * 
-	 * @param Name The name of the control to modify. 
+	 * @param Name The name of the control to update. 
 	 * @param Position The position of the control point on the child mesh object (only relevant if that 
 	 *        object is in use and is being simulated)
 	 */
@@ -555,40 +622,51 @@ public:
 	/**
 	 * Resets the control point to the center of mass of the mesh
 	 *
-	 * @param Name The name of the control to modify. 
+	 * @param Name The name of the control to update. 
 	 */
 	UFUNCTION(BlueprintCallable, Category = PhysicsControl, Meta = (ReturnDisplayName = "Success"))
 	UE_API bool ResetControlPoint(const FName Name);
 
 	/**
-	 * Modifies an existing control target - i.e. what it is driving towards, relative to the parent object
+	 * Modifies an existing control target in a control or controls in a set - i.e. what it is driving 
+	 * towards, relative to the parent object
+	 * Standard set names will include "All", "WorldSpace", "ParentSpace" and things like "WorldSpace-ArmLeft",
+	 * depending on how controls have been created.
 	 *
-	 * @param Name The name of the control to modify. 
+	 * @param Name The name of the control or set of controls to update. 
 	 * @param ControlTarget The new target for the control
 	 * @param bEnableControl Enables the control if it is currently disabled
-	 * @return true if the control was found and modified, false if not
+	 * @param bApplyToControlsWithName If true, then controls with the name will be modified
+	 * @param bApplyToSetsWithName If true, then controls in sets with the name will be modified
 	 */
 	UFUNCTION(BlueprintCallable, Category = PhysicsControl, Meta = (ReturnDisplayName = "Success"))
-	UE_API bool SetControlTarget(const FName Name, const FPhysicsControlTarget ControlTarget, const bool bEnableControl = true);
+	UE_API void SetControlTarget(
+		const FName Name, const FPhysicsControlTarget ControlTarget, const bool bEnableControl = true,
+		const bool bApplyToControlsWithName = true, const bool bApplyToSetsWithName = true);
 
 	/**
-	 * Modifies existing control targets - i.e. what they are driving towards, relative to the parent objects
+	 * Modifies existing control targets in controls or controls in sets - i.e. what they are driving 
+	 * towards, relative to the parent objects
+	 * Standard set names will include "All", "WorldSpace", "ParentSpace" and things like "WorldSpace-ArmLeft",
+	 * depending on how controls have been created.
 	 *
-	 * @param Names The names of the controls to modify. Note that if you have these in a FPhysicsControlNameArray
-	 *              then it can be split.
+	 * @param Names The names of the controls or sets of controls to update. Note that if you have these in a 
+	 *        FPhysicsControlNameArray then it can be split.
 	 * @param ControlTarget The new target for the controls
+	 * 
 	 * @param bEnableControl Enables the controls if currently disabled
+	 * @param bApplyToControlsWithName If true, then controls with the name will be modified
+	 * @param bApplyToSetsWithName If true, then controls in sets with the name will be modified
 	 */
 	UFUNCTION(BlueprintCallable, Category = PhysicsControl)
 	UE_API void SetControlTargets(
-		const TArray<FName>&        Names, 
-		const FPhysicsControlTarget ControlTarget, 
-		const bool                  bEnableControl = true);
+		const TArray<FName>& Names, const FPhysicsControlTarget ControlTarget, const bool bEnableControl = true,
+		const bool bApplyToControlsWithName = true, const bool bApplyToSetsWithName = true);
 
 	/**
 	 * Modifies existing control targets - i.e. what they are driving towards, relative to the parent objects
 	 *
-	 * @param Set The set of controls to modify. Standard sets will include "All", "WorldSpace",
+	 * @param Set The set of controls to update. Standard sets will include "All", "WorldSpace",
 	 *        "ParentSpace" and things like "WorldSpace-ArmLeft", depending on how controls have been created.
 	 * @param ControlTarget The new target for the controls
 	 * @param bEnableControl Enables the controls if currently disabled
@@ -600,9 +678,12 @@ public:
 		const bool                  bEnableControl = true);
 
 	/**
-	 * Modifies an existing control target - i.e. what it is driving towards, relative to the parent object
+	 * Modifies an existing control target in a control or controls in a set - i.e. what it is driving 
+	 * towards, relative to the parent object
+	 * Standard set names will include "All", "WorldSpace", "ParentSpace" and things like "WorldSpace-ArmLeft",
+	 * depending on how controls have been created.
 	 *
-	 * @param Name The name of the control to modify. 
+	 * @param Name The name of the control or set of controls to update. 
 	 * @param Position The new position target for the control
 	 * @param Orientation The new orientation target for the control
 	 * @param VelocityDeltaTime If non-zero, the target velocity will be calculated using the current target 
@@ -614,19 +695,40 @@ public:
 	 *        want to specify the target transform for the object as a whole. If false, then the 
 	 *        target transform is used as is, and the system drives the control point towards this
 	 *        transform.
-	 * @return Returns true if the control was found and modified, false if not
+	 * @param bApplyToControlsWithName If true, then controls with the name will be modified
+	 * @param bApplyToSetsWithName If true, then controls in sets with the name will be modified
 	 */
 	UFUNCTION(BlueprintCallable, Category = PhysicsControl, Meta = (ReturnDisplayName = "Success"))
-	UE_API bool SetControlTargetPositionAndOrientation(
+	UE_API void SetControlTargetPositionAndOrientation(
 		const FName    Name,
 		const FVector  Position,
 		const FRotator Orientation,
 		const float    VelocityDeltaTime,
 		const bool     bEnableControl = true,
-		const bool     bApplyControlPointToTarget = false);
+		const bool     bApplyControlPointToTarget = false,
+		const bool     bApplyToControlsWithName = true, 
+		const bool     bApplyToSetsWithName = true);
 
 	/**
-	 * Calls SetControlTargetPositionAndOrientation for each of the control names
+	 * Modifies existing control targets in controls or controls in sets - i.e. what it is driving
+	 * towards, relative to the parent object
+	 * Standard set names will include "All", "WorldSpace", "ParentSpace" and things like "WorldSpace-ArmLeft",
+	 * depending on how controls have been created.
+	 *
+	 * @param Names The name of the controls or sets of controls to update.
+	 * @param Position The new position target for the control
+	 * @param Orientation The new orientation target for the control
+	 * @param VelocityDeltaTime If non-zero, the target velocity will be calculated using the current target
+	 *        position. If zero, the target velocity will be set to zero.
+	 * @param bEnableControl Enables the control if it is currently disabled
+	 * @param bApplyControlPointToTarget If true, then the target position/orientation is treated as
+	 *        a "virtual" object, where the system attempts to move the object to match the pose of
+	 *        this "virtual" object that has been placed at the target transform. Use this when you
+	 *        want to specify the target transform for the object as a whole. If false, then the
+	 *        target transform is used as is, and the system drives the control point towards this
+	 *        transform.
+	 * @param bApplyToControlsWithName If true, then controls with the name will be modified
+	 * @param bApplyToSetsWithName If true, then controls in sets with the name will be modified
 	 */
 	UFUNCTION(BlueprintCallable, Category = PhysicsControl)
 	UE_API void SetControlTargetPositionsAndOrientations(
@@ -635,7 +737,9 @@ public:
 		const FRotator       Orientation,
 		const float          VelocityDeltaTime,
 		const bool           bEnableControl = true,
-		const bool           bApplyControlPointToTarget = false);
+		const bool           bApplyControlPointToTarget = false,
+		const bool           bApplyToControlsWithName = true, 
+		const bool           bApplyToSetsWithName = true);
 
 	/**
 	 * Calls SetControlTargetPositionAndOrientation for each control in the set
@@ -650,9 +754,12 @@ public:
 		const bool           bApplyControlPointToTarget = false);
 
 	/**
-	 * Modifies an existing control target - i.e. what it is driving towards, relative to the parent object
+	 * Modifies an existing control target of a control or controls in a set - i.e. what it is driving 
+	 * towards, relative to the parent object
+	 * Standard set names will include "All", "WorldSpace", "ParentSpace" and things like "WorldSpace-ArmLeft",
+	 * depending on how controls have been created.
 	 *
-	 * @param Name The name of the control to modify. 
+	 * @param Name The name of the control or set of controls to update. 
 	 * @param Position          The new position target for the control
 	 * @param VelocityDeltaTime If non-zero, the target velocity will be calculated using the current target 
 	 *        position. If zero, the target velocity will be set to zero.    
@@ -664,18 +771,38 @@ public:
 	 *        transform for the object as a whole. If false, then the target transform
 	 *        is used as is, and the system drives the control point towards this
 	 *        transform.
-	 * @return true if the control was found and modified, false if not
+	 * @param bApplyToControlsWithName If true, then controls with the name will be modified
+	 * @param bApplyToSetsWithName If true, then controls in sets with the name will be modified
 	 */
 	UFUNCTION(BlueprintCallable, Category = PhysicsControl, Meta = (ReturnDisplayName = "Success"))
-	UE_API bool SetControlTargetPosition(
+	UE_API void SetControlTargetPosition(
 		const FName   Name,
 		const FVector Position, 
 		const float   VelocityDeltaTime, 
 		const bool    bEnableControl = true,
-		const bool    bApplyControlPointToTarget = false);
+		const bool    bApplyControlPointToTarget = false,
+		const bool    bApplyToControlsWithName = true, 
+		const bool    bApplyToSetsWithName = true);
 
 	/**
-	 * Calls SetControlTargetPosition for each of the control names
+	 * Modifies existing control targets in controls or controls in sets - i.e. what it is driving
+	 * towards, relative to the parent object
+	 * Standard set names will include "All", "WorldSpace", "ParentSpace" and things like "WorldSpace-ArmLeft",
+	 * depending on how controls have been created.
+	 *
+	 * @param Names The name of the controls or sets of controls to update.
+	 * @param Position The new position target for the control
+	 * @param VelocityDeltaTime If non-zero, the target velocity will be calculated using the current target
+	 *        position. If zero, the target velocity will be set to zero.
+	 * @param bEnableControl Enables the control if it is currently disabled
+	 * @param bApplyControlPointToTarget If true, then the target position/orientation is treated as
+	 *        a "virtual" object, where the system attempts to move the object to match the pose of
+	 *        this "virtual" object that has been placed at the target transform. Use this when you
+	 *        want to specify the target transform for the object as a whole. If false, then the
+	 *        target transform is used as is, and the system drives the control point towards this
+	 *        transform.
+	 * @param bApplyToControlsWithName If true, then controls with the name will be modified
+	 * @param bApplyToSetsWithName If true, then controls in sets with the name will be modified
 	 */
 	UFUNCTION(BlueprintCallable, Category = PhysicsControl)
 	UE_API void SetControlTargetPositions(
@@ -683,7 +810,9 @@ public:
 		const FVector        Position, 
 		const float          VelocityDeltaTime, 
 		const bool           bEnableControl = true,
-		const bool           bApplyControlPointToTarget = false);
+		const bool           bApplyControlPointToTarget = false,
+		const bool           bApplyToControlsWithName = true, 
+		const bool           bApplyToSetsWithName = true);
 
 	/**
 	 * Calls SetControlTargetPosition for each of the controls in the set
@@ -697,9 +826,12 @@ public:
 		const bool    bApplyControlPointToTarget = false);
 
 	/**
-	 * Modifies an existing control target - i.e. what it is driving towards, relative to the parent object
+	 * Modifies an existing control target in a control or controls in a set - i.e. what it is driving 
+	 * towards, relative to the parent object
+	 * Standard set names will include "All", "WorldSpace", "ParentSpace" and things like "WorldSpace-ArmLeft",
+	 * depending on how controls have been created.
 	 *
-	 * @param Name The name of the control to modify. 
+	 * @param Name The name of the control or set of controls to update. 
 	 * @param Orientation The new orientation target for the control
 	 * @param AngularVelocityDeltaTime If non-zero, the target angular velocity will be calculated using the current 
 	 *        target position. If zero, the target velocity will be set to zero.
@@ -711,18 +843,39 @@ public:
 	 *        transform for the object as a whole. If false, then the target transform
 	 *        is used as is, and the system drives the control point towards this
 	 *        transform.
-	 * @return true if the control was found and modified, false if not
+	 * @param bApplyToControlsWithName If true, then controls with the name will be modified
+	 * @param bApplyToSetsWithName If true, then controls in sets with the name will be modified
 	 */
 	UFUNCTION(BlueprintCallable, Category = PhysicsControl, Meta = (ReturnDisplayName = "Success"))
-	UE_API bool SetControlTargetOrientation(
+	UE_API void SetControlTargetOrientation(
 		const FName    Name,
 		const FRotator Orientation, 
 		const float    AngularVelocityDeltaTime, 
 		const bool     bEnableControl = true,
-		const bool     bApplyControlPointToTarget = false);
+		const bool     bApplyControlPointToTarget = false,
+		const bool     bApplyToControlsWithName = true, 
+		const bool     bApplyToSetsWithName = true);
 
 	/**
-	 * Calls SetControlTargetOrientation for each of the control names
+	 * Modifies existing control targets in controls or controls in sets - i.e. what it is driving
+	 * towards, relative to the parent object
+	 * Standard set names will include "All", "WorldSpace", "ParentSpace" and things like "WorldSpace-ArmLeft",
+	 * depending on how controls have been created.
+	 *
+	 * @param Name The name of the control or set of controls to update.
+	 * @param Orientation The new orientation target for the control
+	 * @param AngularVelocityDeltaTime If non-zero, the target angular velocity will be calculated using the current
+	 *        target position. If zero, the target velocity will be set to zero.
+	 * @param bEnableControl Enables the control if it is currently disabled
+	 * @param bApplyControlPointToTarget If true, then the target position/orientation is treated as
+	 *        a "virtual" object, where the system attempts to move the object
+	 *        to match the pose of this "virtual" object that has been placed at
+	 *        the target transform. Use this when you want to specify the target
+	 *        transform for the object as a whole. If false, then the target transform
+	 *        is used as is, and the system drives the control point towards this
+	 *        transform.
+	 * @param bApplyToControlsWithName If true, then controls with the name will be modified
+	 * @param bApplyToSetsWithName If true, then controls in sets with the name will be modified
 	 */
 	UFUNCTION(BlueprintCallable, Category = PhysicsControl)
 	UE_API void SetControlTargetOrientations(
@@ -730,7 +883,9 @@ public:
 		const FRotator Orientation,
 		const float    AngularVelocityDeltaTime, 
 		const bool     bEnableControl = true,
-		const bool     bApplyControlPointToTarget = false);
+		const bool     bApplyControlPointToTarget = false,
+		const bool     bApplyToControlsWithName = true,
+		const bool     bApplyToSetsWithName = true);
 
 	/**
 	 * Calls SetControlTargetOrientation for each of the controls in the set
@@ -787,7 +942,7 @@ public:
 	 * Calculates and sets an existing control target. This takes the "virtual" position/orientation of the parent 
 	 * and child and calculates the relative control. Note that this will set bApplyControlPointToTarget to true.
 	 * 
-	 * @param Name              The name of the control to modify. 
+	 * @param Name              The name of the control to update. 
 	 * @param ParentPosition    The virtual/target parent position
 	 * @param ParentOrientation The virtual/target parent orientation
 	 * @param ChildPosition     The virtual/target child position
@@ -808,28 +963,35 @@ public:
 		const bool     bEnableControl = true);
 
 	/**
-	 * Sets whether or not the control should use skeletal animation for the targets
+	 * Sets whether or not the control or controls in a set should use skeletal animation for the targets
+	 * Standard set names will include "All", "WorldSpace", "ParentSpace" and things like "WorldSpace-ArmLeft",
+	 * depending on how controls have been created.
 	 *
-	 * @param Name The name of the control to modify. 
+	 * @param Name The name of the control or set of controls to update. 
 	 * @param bUseSkeletalAnimation If true then the targets will be a combination of the skeletal animation (if
 	 *        there is any) and the control target that has been set
 	 * @param AngularTargetVelocityMultiplier If skeletal animation is being used, then this determines the amount of
 	 *        angular velocity extracted from the animation that is used as targets for the controls
 	 * @param LinearTargetVelocityMultiplier If skeletal animation is being used, then this determines the amount of
 	 *        linear velocity extracted from the animation that is used as targets for the controls
-	 * @return true if the control was found and modified, false if not
+	 * @param bApplyToControlsWithName If true, then controls with the name will be modified
+	 * @param bApplyToSetsWithName If true, then controls in sets with the name will be modified
 	 */
 	UFUNCTION(BlueprintCallable, Category = PhysicsControl, Meta = (ReturnDisplayName = "Success"))
-	UE_API bool SetControlUseSkeletalAnimation(
+	UE_API void SetControlUseSkeletalAnimation(
 		const FName Name,
 		const bool  bUseSkeletalAnimation = true,
 		const float AngularTargetVelocityMultiplier = 1.0f,
-		const float LinearTargetVelocityMultiplier = 1.0f);
+		const float LinearTargetVelocityMultiplier = 1.0f,
+		const bool  bApplyToControlsWithName = true, 
+		const bool  bApplyToSetsWithName = true);
 
 	/**
-	 * Sets whether or not the controls should use skeletal animation for the targets
+	 * Sets whether or not the controls or controls in sets should use skeletal animation for the targets
+	 * Standard set names will include "All", "WorldSpace", "ParentSpace" and things like "WorldSpace-ArmLeft",
+	 * depending on how controls have been created.
 	 *
-	 * @param Names The names of the controls to modify. Note that if you have these in a FPhysicsControlNameArray 
+	 * @param Names The names of the controls or sets of controls to update. Note that if you have these in a FPhysicsControlNameArray 
 	 *              then it can be split.
 	 * @param bUseSkeletalAnimation If true then the targets will be a combination of the skeletal animation (if
 	 *              there is any) and the control target that has been set
@@ -837,18 +999,22 @@ public:
 	 *              angular velocity extracted from the animation that is used as targets for the controls
 	 * @param LinearTargetVelocityMultiplier If skeletal animation is being used, then this determines the amount of
 	 *              linear velocity extracted from the animation that is used as targets for the controls
+	 * @param bApplyToControlsWithName If true, then controls with the name will be modified
+	 * @param bApplyToSetsWithName If true, then controls in sets with the name will be modified
 	 */
 	UFUNCTION(BlueprintCallable, Category = PhysicsControl)
 	UE_API void SetControlsUseSkeletalAnimation(
 		const TArray<FName>& Names,
 		const bool           bUseSkeletalAnimation = true,
 		const float          AngularTargetVelocityMultiplier = 1.0f,
-		const float          LinearTargetVelocityMultiplier = 1.0f);
+		const float          LinearTargetVelocityMultiplier = 1.0f,
+		const bool           bApplyToControlsWithName = true,
+		const bool           bApplyToSetsWithName = true);
 
 	/**
 	 * Sets whether or not the controls should use skeletal animation for the targets
 	 *
-	 * @param Set The set of controls to modify. Standard sets will include "All", "WorldSpace",
+	 * @param Set The set of controls to update. Standard sets will include "All", "WorldSpace",
 	 *        "ParentSpace" and things like "WorldSpace-ArmLeft", depending on how controls have been created.
 	 * @param bUseSkeletalAnimation If true then the targets will be a combination of the skeletal animation (if
 	 *        there is any) and the control target that has been set
@@ -865,52 +1031,73 @@ public:
 		const float LinearTargetVelocityMultiplier = 1.0f);
 
 	/**
-	 * Activates or deactivates a control
+	 * Activates or deactivates a control or controls in a set.
+	 * Standard set names will include "All", "WorldSpace", "ParentSpace" and things like "WorldSpace-ArmLeft", 
+	 * depending on how controls have been created.
 	 *
-	 * @param Name     The name of the control to modify. 
+	 * @param Name     The name of the control or set of controls to update. 
 	 * @param bEnable  Whether to enable/disable the control
-	 * @return         Returns true if the control was found and modified, false if not
+	 * @param bApplyToControlsWithName If true, then controls with the name will be modified
+	 * @param bApplyToSetsWithName If true, then controls in sets with the name will be modified
 	 */
 	UFUNCTION(BlueprintCallable, Category = PhysicsControl, Meta = (ReturnDisplayName = "Success"))
-	UE_API bool SetControlEnabled(const FName Name, const bool bEnable = true);
+	UE_API void SetControlEnabled(
+		const FName Name, const bool bEnable = true, 
+		const bool bApplyToControlsWithName = true, const bool bApplyToSetsWithName = true);
 
 	/**
-	 * Activates or deactivates controls
+	 * Activates or deactivates controls or sets of controls
+	 * Standard set names will include "All", "WorldSpace", "ParentSpace" and things like "WorldSpace-ArmLeft",
+	 * depending on how controls have been created.
 	 *
-	 * @param Names The names of the controls to modify. Note that if you have these in a FPhysicsControlNameArray
+	 * @param Names The names of the controls or sets of controls to update. Note that if you have these in a FPhysicsControlNameArray
 	 *              then it can be split.
 	 * @param bEnable  Whether to enable/disable the controls
+	 * @param bApplyToControlsWithName If true, then controls with the name will be modified
+	 * @param bApplyToSetsWithName If true, then controls in sets with the name will be modified
 	 */
 	UFUNCTION(BlueprintCallable, Category = PhysicsControl)
-	UE_API void SetControlsEnabled(const TArray<FName>& Names, const bool bEnable = true);
+	UE_API void SetControlsEnabled(
+		const TArray<FName>& Names, const bool bEnable = true, 
+		const bool bApplyToControlsWithName = true, const bool bApplyToSetsWithName = true);
 
 	/**
 	 * Activates or deactivates controls
 	 *
-	 * @param Set The set of controls to modify. Standard sets will include "All", "WorldSpace",
+	 * @param Set The set of controls to update. Standard sets will include "All", "WorldSpace",
 	 *        "ParentSpace" and things like "WorldSpace-ArmLeft", depending on how controls have been created.
 	 * @param bEnable  Whether to enable/disable the controls
 	 */
+
 	UFUNCTION(BlueprintCallable, Category = PhysicsControl)
 	UE_API void SetControlsInSetEnabled(const FName Set, const bool bEnable = true);
 
 	/**
-	 * @param Name The name of the control to modify. 
+	 * @param Name The name of the control or set of controls to update. Standard sets will include "All", "WorldSpace",
+	 *        "ParentSpace" and things like "WorldSpace-ArmLeft", depending on how controls have been created.
 	 * @param bDisableCollision If set then the control will disable collision between the bodies it connects.
-	 * @return true if the control was found and modified, false if not
+	 * @param bApplyToControlsWithName If true, then controls with the name will be modified
+	 * @param bApplyToSetsWithName If true, then controls in sets with the name will be modified
 	 */
 	UFUNCTION(BlueprintCallable, Category = PhysicsControl, Meta = (ReturnDisplayName = "Success"))
-	UE_API bool SetControlDisableCollision(const FName Name, const bool bDisableCollision);
+	UE_API void SetControlDisableCollision(
+		const FName Name, const bool bDisableCollision,
+		const bool bApplyToControlsWithName = true, const bool bApplyToSetsWithName = true);
 
 	/**
-	 * @param Names The names of the controls to modify. 
+	 * @param Names The names of the controls or sets of controls to update. Standard sets will include "All", "WorldSpace",
+	 *        "ParentSpace" and things like "WorldSpace-ArmLeft", depending on how controls have been created.
 	 * @param bDisableCollision If set then the control will disable collision between the bodies it connects.
+	 * @param bApplyToControlsWithName If true, then controls with the name will be modified
+	 * @param bApplyToSetsWithName If true, then controls in sets with the name will be modified
 	 */
 	UFUNCTION(BlueprintCallable, Category = PhysicsControl)
-	UE_API void SetControlsDisableCollision(const TArray<FName>& Names, const bool bDisableCollision);
+	UE_API void SetControlsDisableCollision(
+		const TArray<FName>& Names, const bool bDisableCollision,
+		const bool bApplyToControlsWithName = true, const bool bApplyToSetsWithName = true);
 
 	/**
-	 * @param Set The set of controls to modify. Standard sets will include "All", "WorldSpace",
+	 * @param Set The set of controls to update. Standard sets will include "All", "WorldSpace",
 	 *        "ParentSpace" and things like "WorldSpace-ArmLeft", depending on how controls have been created.
 	 * @param bDisableCollision If set then the control will disable collision between the bodies it connects.
 	 */
@@ -1000,7 +1187,6 @@ public:
 		const FName                       Set,
 		const FPhysicsControlModifierData BodyModifierData);
 
-
 	/**
 	 * Creates a collection of controls controlling a skeletal mesh, grouped together in limbs
 	 *
@@ -1018,25 +1204,32 @@ public:
 		const FPhysicsControlModifierData            BodyModifierData);
 
 	/**
-	 * Destroys a BodyModifier
+	 * Destroys a body modifier or body modifiers in a set
+	 * Standard sets will include "All" and things like "ArmLeft", depending on how body modifiers have been created.
 	 *
-	 * @param Name        The name of the body modifier to destroy. 
-	 * @return            Returns true if the body modifier was found and destroyed, false if not
+	 * @param Name The name of the body modifier or set of body modifiers to destroy. 
+	 * @param bApplyToModifiersWithName If true, then body modifiers with the name will be modified
+	 * @param bApplyToSetsWithName If true, then body modifiers in sets with the name will be modified
 	 */
 	UFUNCTION(BlueprintCallable, Category = PhysicsControl, Meta = (ReturnDisplayName = "Success"))
-	UE_API bool DestroyBodyModifier(const FName Name);
+	UE_API void DestroyBodyModifier(
+		const FName Name, const bool bApplyToModifiersWithName = true, const bool bApplyToSetsWithName = true);
 
 	/**
-	 * Destroys BodyModifiers
+	 * Destroys body modifiers or body modifiers in sets
+	 * Standard sets will include "All" and things like "ArmLeft", depending on how body modifiers have been created.
 	 *
-	 * @param Names The names of the body modifiers to destroy. Note that if you have these in a FPhysicsControlNameArray
+	 * @param Names The names of the body modifiers or sets of body modifiers to destroy. Note that if you have these in a FPhysicsControlNameArray
 	 *              then it can be split.
+	 * @param bApplyToModifiersWithName If true, then body modifiers with the name will be modified
+	 * @param bApplyToSetsWithName If true, then body modifiers in sets with the name will be modified
 	 */
 	UFUNCTION(BlueprintCallable, Category = PhysicsControl)
-	UE_API void DestroyBodyModifiers(const TArray<FName>& Names);
+	UE_API void DestroyBodyModifiers(
+		const TArray<FName>& Names, const bool bApplyToModifiersWithName = true, const bool bApplyToSetsWithName = true);
 
 	/**
-	 * Destroys BodyModifiers
+	 * Destroys body modifiers
 	 *
 	 * @param Set The set of body modifiers to destroy. Standard sets will include "All" and things like
 	 *        "ArmLeft", depending on how body modifiers have been created.
@@ -1045,24 +1238,33 @@ public:
 	UE_API void DestroyBodyModifiersInSet(const FName Set);
 
 	/**
-	 * Modifies an existing Body Modifier Data
+	 * Modifies an existing Body Modifier Data in a body modifier or body modifiers in a set
+	 * Standard sets will include "All" and things like "ArmLeft", depending on how body modifiers have been created.
 	 *
-	 * @param Name The name of the modifier to update. 
+	 * @param Name The name of the modifier or set to update. 
 	 * @param ModifierData The new data
-	 * @return true if the modifier was found and modified, false if not
+	 * @param bApplyToModifiersWithName If true, then body modifiers with the name will be modified
+	 * @param bApplyToSetsWithName If true, then body modifiers in sets with the name will be modified
 	 */
 	UFUNCTION(BlueprintCallable, Category = PhysicsControl, Meta = (ReturnDisplayName = "Success"))
-	UE_API bool SetBodyModifierData(const FName Name, const FPhysicsControlModifierData ModifierData);
+	UE_API void SetBodyModifierData(
+		const FName Name, const FPhysicsControlModifierData ModifierData, 
+		const bool bApplyToModifiersWithName = true, const bool bApplyToSetsWithName = true);
 
 	/**
-	 * Modifies existing Body Modifier Data
+	 * Modifies existing body modifier data in body modifiers or body modifiers in sets
+	 * Standard sets will include "All" and things like "ArmLeft", depending on how body modifiers have been created.
 	 *
-	 * @param Names The names of the modifiers to update. Note that if you have these in a FPhysicsControlNameArray
+	 * @param Names The names of the modifiers or sets to update. Note that if you have these in a FPhysicsControlNameArray
 	 *              then it can be split.
 	 * @param ModifierData The new data
+	 * @param bApplyToModifiersWithName If true, then body modifiers with the name will be modified
+	 * @param bApplyToSetsWithName If true, then body modifiers in sets with the name will be modified
 	 */
 	UFUNCTION(BlueprintCallable, Category = PhysicsControl)
-	UE_API void SetBodyModifierDatas(const TArray<FName>& Names, const FPhysicsControlModifierData ModifierData);
+	UE_API void SetBodyModifierDatas(
+		const TArray<FName>& Names, const FPhysicsControlModifierData ModifierData,
+		const bool bApplyToModifiersWithName = true, const bool bApplyToSetsWithName = true);
 
 	/**
 	 * Modifies existing Body Modifier Data
@@ -1075,24 +1277,32 @@ public:
 	UE_API void SetBodyModifierDatasInSet(const FName Set, const FPhysicsControlModifierData ModifierData);
 
 	/**
-	 * Modifies an existing Body Modifier Data
+	 * Modifies an existing Body Modifier Data in a body modifier or body modifiers in a set
 	 *
-	 * @param Name The name of the modifier to update. 
+	 * @param Name The name of the modifier or set to update. 
 	 * @param ModifierData The new data
-	 * @return true if the modifier was found and modified, false if not
+	 * @param bApplyToModifiersWithName If true, then body modifiers with the name will be modified
+	 * @param bApplyToSetsWithName If true, then body modifiers in sets with the name will be modified
 	 */
 	UFUNCTION(BlueprintCallable, Category = PhysicsControl, Meta = (ReturnDisplayName = "Success"))
-	UE_API bool SetBodyModifierSparseData(const FName Name, const FPhysicsControlModifierSparseData ModifierData);
+	UE_API void SetBodyModifierSparseData(
+		const FName Name, const FPhysicsControlModifierSparseData ModifierData,
+		const bool bApplyToModifiersWithName = true, const bool bApplyToSetsWithName = true);
 
 	/**
-	 * Modifies existing Body Modifier Data
+	 * Modifies existing Body Modifier Data in body modifiers or body modifiers in sets
+	 * Standard sets will include "All" and things like "ArmLeft", depending on how body modifiers have been created.
 	 *
-	 * @param Names The names of the modifiers to update. Note that if you have these in a FPhysicsControlNameArray
+	 * @param Names The names of the modifiers or sets to update. Note that if you have these in a FPhysicsControlNameArray
 	 *              then it can be split.
 	 * @param ModifierData The new data
+	 * @param bApplyToModifiersWithName If true, then body modifiers with the name will be modified
+	 * @param bApplyToSetsWithName If true, then body modifiers in sets with the name will be modified
 	 */
 	UFUNCTION(BlueprintCallable, Category = PhysicsControl)
-	UE_API void SetBodyModifierSparseDatas(const TArray<FName>& Names, const FPhysicsControlModifierSparseData ModifierData);
+	UE_API void SetBodyModifierSparseDatas(
+		const TArray<FName>& Names, const FPhysicsControlModifierSparseData ModifierData,
+		const bool bApplyToModifiersWithName = true, const bool bApplyToSetsWithName = true);
 
 	/**
 	 * Modifies existing Body Modifier Data
@@ -1107,11 +1317,11 @@ public:
 	/**
 	 * Sets the kinematic target transform for a body modifier. 
 	 * 
-	 * @param Name The name of the body modifier to access. 
+	 * @param Name The name of the body modifier or set to update. 
 	 * @param KinematicTargetPosition The position to use as the kinematic target of the associated body,
-	 *                                if it is kinematic
+	 *        if it is kinematic
 	 * @param KinematicTargetOrientation The orientation to use as the kinematic target of the associated body,
-	 *                                   if it is kinematic
+	 *        if it is kinematic
 	 * @param bMakeKinematic If set then the body will be made kinematic. If not set, then it won't be changed.
 	 * @return true if the body modifier was found, false if not
 	 */
@@ -1123,33 +1333,36 @@ public:
 		const bool     bMakeKinematic);
 
 	/**
-	 * Sets the movement type for a body modifier
+	 * Sets the movement type for a body modifier or body modifiers in a set
 	 *
-	 * @param Name The name of the body modifier to access. 
+	 * @param Name The name of the body modifier or set to update. 
 	 * @param MovementType Whether to enable/disable simulation on the body
-	 * @return true if the body modifier was found, false if not
+	 * @param bApplyToModifiersWithName If true, then body modifiers with the name will be modified
+	 * @param bApplyToSetsWithName If true, then body modifiers in sets with the name will be modified
 	 */
 	UFUNCTION(BlueprintCallable, Category = PhysicsControl, Meta = (ReturnDisplayName = "Success"))
-	UE_API bool SetBodyModifierMovementType(
-		const FName                Name,
-		const EPhysicsMovementType MovementType = EPhysicsMovementType::Simulated);
+	UE_API void SetBodyModifierMovementType(
+		const FName Name, const EPhysicsMovementType MovementType = EPhysicsMovementType::Simulated,
+		const bool bApplyToModifiersWithName = true, const bool bApplyToSetsWithName = true);
 
 	/**
-	 * Sets the movement type for body modifiers
+	 * Sets the movement type for body modifiers or body modifiers in sets
 	 *
-	 * @param Names The names of the body modifiers to access. Note that if you have these in a FPhysicsControlNameArray
+	 * @param Names The names of the body modifiers or sets to update. Note that if you have these in a FPhysicsControlNameArray
 	 *        then it can be split.
+	 * @param bApplyToModifiersWithName If true, then body modifiers with the name will be modified
+	 * @param bApplyToSetsWithName If true, then body modifiers in sets with the name will be modified
 	 * @param MovementType Whether to enable/disable simulation on the bodies
 	 */
 	UFUNCTION(BlueprintCallable, Category = PhysicsControl)
 	UE_API void SetBodyModifiersMovementType(
-		const TArray<FName>&       Names,
-		const EPhysicsMovementType MovementType = EPhysicsMovementType::Simulated);
+		const TArray<FName>& Names, const EPhysicsMovementType MovementType = EPhysicsMovementType::Simulated,
+		const bool bApplyToModifiersWithName = true, const bool bApplyToSetsWithName = true);
 
 	/**
 	 * Sets the movement type for body modifiers
 	 *
-	 * @param Set The set of body modifiers to modify. Standard sets will include "All" and things like
+	 * @param Set The set of body modifiers to update. Standard sets will include "All" and things like
 	 *        "ArmLeft", depending on how body modifiers have been created.
 	 * @param MovementType Whether to enable/disable simulation on the bodies
 	 */
@@ -1159,33 +1372,38 @@ public:
 		const EPhysicsMovementType MovementType = EPhysicsMovementType::Simulated);
 
 	/**
-	 * Sets the collision type for a body modifier
+	 * Sets the collision type for a body modifier or body modifiers in a set
+	 * Standard sets will include "All" and things like "ArmLeft", depending on how body modifiers have been created.
 	 *
-	 * @param Name The name of the body modifier to access.
+	 * @param Name The name of the body modifier or set to update.
 	 * @param CollisionType Collision type to set on the body
-	 * @return true if the body modifier was found, false if not
+	 * @param bApplyToModifiersWithName If true, then body modifiers with the name will be modified
+	 * @param bApplyToSetsWithName If true, then body modifiers in sets with the name will be modified
 	 */
 	UFUNCTION(BlueprintCallable, Category = PhysicsControl, Meta = (ReturnDisplayName = "Success"))
-	UE_API bool SetBodyModifierCollisionType(
-		const FName                   Name,
-		const ECollisionEnabled::Type CollisionType = ECollisionEnabled::QueryAndPhysics);
+	UE_API void SetBodyModifierCollisionType(
+		const FName Name, const ECollisionEnabled::Type CollisionType = ECollisionEnabled::QueryAndPhysics,
+		const bool bApplyToModifiersWithName = true, const bool bApplyToSetsWithName = true);
 
 	/**
-	 * Sets the collision type for body modifiers
+	 * Sets the collision type for body modifiers or body modifiers in sets
+	 * Standard sets will include "All" and things like "ArmLeft", depending on how body modifiers have been created.
 	 *
-	 * @param Names The names of the body modifiers to access. Note that if you have these in a FPhysicsControlNameArray
+	 * @param Names The names of the body modifiers or sets to update. Note that if you have these in a FPhysicsControlNameArray
 	 *        then it can be split.
 	 * @param CollisionType Collision type to set on the bodies
+	 * @param bApplyToModifiersWithName If true, then body modifiers with the name will be modified
+	 * @param bApplyToSetsWithName If true, then body modifiers in sets with the name will be modified
 	 */
 	UFUNCTION(BlueprintCallable, Category = PhysicsControl)
 	UE_API void SetBodyModifiersCollisionType(
-		const TArray<FName>&          Names,
-		const ECollisionEnabled::Type CollisionType = ECollisionEnabled::QueryAndPhysics);
+		const TArray<FName>& Names, const ECollisionEnabled::Type CollisionType = ECollisionEnabled::QueryAndPhysics,
+		const bool bApplyToModifiersWithName = true, const bool bApplyToSetsWithName = true);
 
 	/**
 	 * Sets the collision type for body modifiers
 	 *
-	 * @param Set The set of body modifiers to modify. Standard sets will include "All" and things like
+	 * @param Set The set of body modifiers to update. Standard sets will include "All" and things like
 	 *        "ArmLeft", depending on how body modifiers have been created.
 	 * @param CollisionType Collision type to set on the bodies
 	 */
@@ -1195,33 +1413,38 @@ public:
 		const ECollisionEnabled::Type CollisionType = ECollisionEnabled::QueryAndPhysics);
 
 	/**
-	 * Sets the gravity multiplier for a body modifier
+	 * Sets the gravity multiplier for a body modifier or body modifiers in a set
+	 * Standard sets will include "All" and things like "ArmLeft", depending on how body modifiers have been created.
 	 *
-	 * @param Name The name of the body modifier to access. 
+	 * @param Name The name of the body modifier or set to update. 
 	 * @param GravityMultiplier The amount of gravity to apply when simulating
-	 * @return true if the body modifier was found, false if not
+	 * @param bApplyToModifiersWithName If true, then body modifiers with the name will be modified
+	 * @param bApplyToSetsWithName If true, then body modifiers in sets with the name will be modified
 	 */
 	UFUNCTION(BlueprintCallable, Category = PhysicsControl, Meta = (ReturnDisplayName = "Success"))
-	UE_API bool SetBodyModifierGravityMultiplier(
-		const FName Name,
-		const float GravityMultiplier = 1.0f);
+	UE_API void SetBodyModifierGravityMultiplier(
+		const FName Name, const float GravityMultiplier = 1.0f,
+		const bool bApplyToModifiersWithName = true, const bool bApplyToSetsWithName = true);
 
 	/**
-	 * Sets the gravity multiplier for body modifiers
+	 * Sets the gravity multiplier for body modifiers or body modifiers in sets
+	 * Standard sets will include "All" and things like "ArmLeft", depending on how body modifiers have been created.
 	 *
-	 * @param Names The names of the body modifiers to access. Note that if you have these in a FPhysicsControlNameArray
+	 * @param Names The names of the body modifiers or sets to update. Note that if you have these in a FPhysicsControlNameArray
 	 *              then it can be split.
 	 * @param GravityMultiplier The amount of gravity to apply when simulating
+	 * @param bApplyToModifiersWithName If true, then body modifiers with the name will be modified
+	 * @param bApplyToSetsWithName If true, then body modifiers in sets with the name will be modified
 	 */
 	UFUNCTION(BlueprintCallable, Category = PhysicsControl)
 	UE_API void SetBodyModifiersGravityMultiplier(
-		const TArray<FName>& Names,
-		const float          GravityMultiplier = 1.0f);
+		const TArray<FName>& Names, const float GravityMultiplier = 1.0f,
+		const bool bApplyToModifiersWithName = true, const bool bApplyToSetsWithName = true);
 
 	/**
 	 * Sets the gravity multiplier for body modifiers
 	 *
-	 * @param Set The set of body modifiers to modify. Standard sets will include "All" and things like
+	 * @param Set The set of body modifiers to update. Standard sets will include "All" and things like
 	 *        "ArmLeft", depending on how body modifiers have been created.
 	 * @param GravityMultiplier The amount of gravity to apply when simulating
 	 */
@@ -1231,35 +1454,40 @@ public:
 		const float GravityMultiplier = 1.0f);
 
 	/**
-	 * Sets the physics blend weight for a body modifier
+	 * Sets the physics blend weight for a body modifier or body modifiers in a set
+	 * Standard sets will include "All" and things like "ArmLeft", depending on how body modifiers have been created.
 	 *
-	 * @param Name The name of the body modifier to access. 
+	 * @param Name The name of the body modifier or set to update. 
 	 * @param PhysicsBlendWeight The blend weight between the body transform coming from 
 	 *        animation and that coming from simulation.
-	 * @return true if the body modifier was found, false if not
+	 * @param bApplyToModifiersWithName If true, then body modifiers with the name will be modified
+	 * @param bApplyToSetsWithName If true, then body modifiers in sets with the name will be modified
 	 */
 	UFUNCTION(BlueprintCallable, Category = PhysicsControl, Meta = (ReturnDisplayName = "Success"))
-	UE_API bool SetBodyModifierPhysicsBlendWeight(
-		const FName Name,
-		const float PhysicsBlendWeight = 1.0f);
+	UE_API void SetBodyModifierPhysicsBlendWeight(
+		const FName Name, const float PhysicsBlendWeight = 1.0f,
+		const bool bApplyToModifiersWithName = true, const bool bApplyToSetsWithName = true);
 
 	/**
-	 * Sets the physics blend weight for body modifiers
+	 * Sets the physics blend weight for body modifiers or body modifiers in sets
+	 * Standard sets will include "All" and things like "ArmLeft", depending on how body modifiers have been created.
 	 *
-	 * @param Names The names of the body modifiers to access. Note that if you have these in a FPhysicsControlNameArray
+	 * @param Names The names of the body modifiers or sets to update. Note that if you have these in a FPhysicsControlNameArray
 	 *        then it can be split.
 	 * @param PhysicsBlendWeight The blend weight between the body transform coming from
 	 *        animation and that coming from simulation.
+	 * @param bApplyToModifiersWithName If true, then body modifiers with the name will be modified
+	 * @param bApplyToSetsWithName If true, then body modifiers in sets with the name will be modified
 	 */
 	UFUNCTION(BlueprintCallable, Category = PhysicsControl)
 	UE_API void SetBodyModifiersPhysicsBlendWeight(
-		const TArray<FName>& Names,
-		const float          PhysicsBlendWeight = 1.0f);
+		const TArray<FName>& Names, const float PhysicsBlendWeight = 1.0f,
+		const bool bApplyToModifiersWithName = true, const bool bApplyToSetsWithName = true);
 
 	/**
 	 * Sets the physics blend weight for body modifiers
 	 *
-	 * @param Set The set of body modifiers to modify. Standard sets will include "All" and things like
+	 * @param Set The set of body modifiers to update. Standard sets will include "All" and things like
 	 *        "ArmLeft", depending on how body modifiers have been created.
 	 * @param PhysicsBlendWeight The blend weight between the body transform coming from
 	 *        animation and that coming from simulation.
@@ -1270,33 +1498,38 @@ public:
 		const float PhysicsBlendWeight = 1.0f);
 
 	/**
-	 * Sets what space kinematic targets should be applied in
+	 * Sets what space kinematic targets should be applied in for a body modifier or body modifiers in a set
+	 * Standard sets will include "All" and things like "ArmLeft", depending on how body modifiers have been created.
 	 *
-	 * @param Name The name of the body modifier to access. 
+	 * @param Name The name of the body modifier or set to update. 
 	 * @param KinematicTargetSpace The space in which the kinematic target should be applied. 
-	 * @return true if the body modifier was found, false if not
+	 * @param bApplyToModifiersWithName If true, then body modifiers with the name will be modified
+	 * @param bApplyToSetsWithName If true, then body modifiers in sets with the name will be modified
 	 */
 	UFUNCTION(BlueprintCallable, Category = PhysicsControl, Meta = (ReturnDisplayName = "Success"))
-	UE_API bool SetBodyModifierKinematicTargetSpace(
-		const FName                               Name,
-		const EPhysicsControlKinematicTargetSpace KinematicTargetSpace);
+	UE_API void SetBodyModifierKinematicTargetSpace(
+		const FName Name, const EPhysicsControlKinematicTargetSpace KinematicTargetSpace,
+		const bool bApplyToModifiersWithName = true, const bool bApplyToSetsWithName = true);
 
 	/**
-	 * Sets what space kinematic targets should be applied in
+	 * Sets what space kinematic targets should be applied in for body modifiers or body modifiers in sets
+	 * Standard sets will include "All" and things like "ArmLeft", depending on how body modifiers have been created.
 	 *
-	 * @param Names The names of the body modifiers to access. Note that if you have these in a FPhysicsControlNameArray
+	 * @param Names The names of the body modifiers or sets to update. Note that if you have these in a FPhysicsControlNameArray
 	 *        then it can be split.
 	 * @param KinematicTargetSpace The space in which the kinematic target should be applied.
+	 * @param bApplyToModifiersWithName If true, then body modifiers with the name will be modified
+	 * @param bApplyToSetsWithName If true, then body modifiers in sets with the name will be modified
 	 */
 	UFUNCTION(BlueprintCallable, Category = PhysicsControl)
 	UE_API void SetBodyModifiersKinematicTargetSpace(
-		const TArray<FName>&                      Names,
-		const EPhysicsControlKinematicTargetSpace KinematicTargetSpace);
+		const TArray<FName>& Names, const EPhysicsControlKinematicTargetSpace KinematicTargetSpace,
+		const bool bApplyToModifiersWithName = true, const bool bApplyToSetsWithName = true);
 
 	/**
 	 * Sets what space kinematic targets should be applied in
 	 *
-	 * @param Set The set of body modifiers to modify. Standard sets will include "All" and things like
+	 * @param Set The set of body modifiers to update. Standard sets will include "All" and things like
 	 *        "ArmLeft", depending on how body modifiers have been created.
 	 * @param KinematicTargetSpace The space in which the kinematic target should be applied.
 	 */
@@ -1306,39 +1539,44 @@ public:
 		const EPhysicsControlKinematicTargetSpace KinematicTargetSpace);
 
 	/**
-	 * Sets whether a body modifier should update kinematics from the simulation results
+	 * Sets whether a body modifier should update kinematics from the simulation results, for a body modifier or body modifiers in a set
+	 * Standard sets will include "All" and things like "ArmLeft", depending on how body modifiers have been created.
 	 *
-	 * @param Name The name of the body modifier to access. 
+	 * @param Name The name of the body modifier or set to update. 
 	 * @param bUpdateKinematicFromSimulation Whether the body should be updated from the simulation when
 	 *        it is kinematic, or whether it should track the kinematic target directly. This will be most
 	 *        likely useful when using async physics, in order to make kinematic parts behave the same as
 	 *        dynamic ones.
-	 * @return true if the body modifier was found, false if not
+	 * @param bApplyToModifiersWithName If true, then body modifiers with the name will be modified
+	 * @param bApplyToSetsWithName If true, then body modifiers in sets with the name will be modified
 	 */
 	UFUNCTION(BlueprintCallable, Category = PhysicsControl, Meta = (ReturnDisplayName = "Success"))
-	UE_API bool SetBodyModifierUpdateKinematicFromSimulation(
-		const FName Name,
-		const bool  bUpdateKinematicFromSimulation);
+	UE_API void SetBodyModifierUpdateKinematicFromSimulation(
+		const FName Name, const bool  bUpdateKinematicFromSimulation,
+		const bool bApplyToModifiersWithName = true, const bool bApplyToSetsWithName = true);
 
 	/**
-	 * Sets whether body modifiers should update kinematics from the simulation results
+	 * Sets whether body modifiers should update kinematics from the simulation results for body modifiers or body modifiers in sets
+	 * Standard sets will include "All" and things like "ArmLeft", depending on how body modifiers have been created.
 	 *
-	 * @param Names The names of the body modifiers to access. Note that if you have these in a FPhysicsControlNameArray
+	 * @param Names The names of the body modifiers or sets to update. Note that if you have these in a FPhysicsControlNameArray
 	 *        then it can be split.
 	 * @param bUpdateKinematicFromSimulation Whether the body should be updated from the simulation when
 	 *        it is kinematic, or whether it should track the kinematic target directly. This will be most
 	 *        likely useful when using async physics, in order to make kinematic parts behave the same as
 	 *        dynamic ones.
+	 * @param bApplyToModifiersWithName If true, then body modifiers with the name will be modified
+	 * @param bApplyToSetsWithName If true, then body modifiers in sets with the name will be modified
 	 */
 	UFUNCTION(BlueprintCallable, Category = PhysicsControl)
 	UE_API void SetBodyModifiersUpdateKinematicFromSimulation(
-		const TArray<FName>& Names,
-		const bool           bUpdateKinematicFromSimulation);
+		const TArray<FName>& Names, const bool bUpdateKinematicFromSimulation,
+		const bool bApplyToModifiersWithName = true, const bool bApplyToSetsWithName = true);
 
 	/**
 	 * Sets whether body modifiers should update kinematics from the simulation results
 	 *
-	 * @param Set The set of body modifiers to modify. Standard sets will include "All" and things like
+	 * @param Set The set of body modifiers to update. Standard sets will include "All" and things like
 	 *        "ArmLeft", depending on how body modifiers have been created.
 	 * @param bUpdateKinematicFromSimulation Whether the body should be updated from the simulation when
 	 *        it is kinematic, or whether it should track the kinematic target directly. This will be most
@@ -1450,7 +1688,7 @@ public:
 	UE_API const TArray<FName>& GetBodyModifierNamesInSet(const FName Set) const;
 
 	/**
-	 * Adds a BodyModifier to a Set. This will add a new set if necessary. For example, you might
+	 * Adds a body modifier to a Set. This will add a new set if necessary. For example, you might
 	 * make a set of body modifiers called "Feet" by calling this twice, passing in the left and right
 	 * foot body modifiers.
 	 */
@@ -1458,7 +1696,7 @@ public:
 	UE_API void AddBodyModifierToSet(FPhysicsControlNames& NewSet, const FName BodyModifier, const FName Set);
 
 	/**
-	 * Adds BodyModifiers to a Set. This will add a new set if necessary. For example, you might
+	 * Adds body modifiers to a Set. This will add a new set if necessary. For example, you might
 	 * make a set of Arm body modifiers by calling this twice, passing in the left and right
 	 * arm body modifiers.
 	 */
@@ -1571,26 +1809,30 @@ public:
 	UE_API void SetCachedBoneVelocitiesToZero();
 
 	/**
-	 * This flags the body associated with the modifier to set (using teleport) its position and velocity to 
+	 * This flags the body associated with the modifier, or modifiers in a set, to set (using teleport) its position and velocity to 
 	 * the cached animation target. This will only affect skeletal mesh component bodies. 
-	 * 
-	 * @param Name     The name of the body modifier to use to identify the body to reset.
+	 * Standard sets will include "All" and things like "ArmLeft", depending on how body modifiers have been created.
+	 *
+	 * @param Name The name of the body modifier or set of body modifiers to use to identify the body/bodies to reset.
 	 * @param Behavior When the reset should happen.
-	 * 
-	 * @return true if the body modifier is found (even if no cached target is found), and false otherwise.
+	 * @param bApplyToModifiersWithName If true, then body modifiers with the name will be modified
+	 * @param bApplyToSetsWithName If true, then body modifiers in sets with the name will be modified
 	 */
 	UFUNCTION(BlueprintCallable, Category = PhysicsControl, Meta = (ReturnDisplayName = "Success"))
-	UE_API bool ResetBodyModifierToCachedBoneTransform(
-		const FName                        Name,
-		const EResetToCachedTargetBehavior Behavior = EResetToCachedTargetBehavior::ResetImmediately);
+	UE_API void ResetBodyModifierToCachedBoneTransform(
+		const FName Name, const EResetToCachedTargetBehavior Behavior = EResetToCachedTargetBehavior::ResetImmediately,
+		const bool bApplyToModifiersWithName = true, const bool bApplyToSetsWithName = true);
 
 	/**
-	 * Calls ResetBodyModifierToCachedTarget for each of the body modifiers
+	 * Calls ResetBodyModifierToCachedTarget for each of the body modifiers or body modifiers in the sets
+	 * Standard sets will include "All" and things like "ArmLeft", depending on how body modifiers have been created.
+	 * @param bApplyToModifiersWithName If true, then body modifiers with the name will be modified
+	 * @param bApplyToSetsWithName If true, then body modifiers in sets with the name will be modified
 	 */
 	UFUNCTION(BlueprintCallable, Category = PhysicsControl)
 	UE_API void ResetBodyModifiersToCachedBoneTransforms(
-		const TArray<FName>&               Names,
-		const EResetToCachedTargetBehavior Behavior = EResetToCachedTargetBehavior::ResetImmediately);
+		const TArray<FName>& Names, const EResetToCachedTargetBehavior Behavior = EResetToCachedTargetBehavior::ResetImmediately,
+		const bool bApplyToModifiersWithName = true, const bool bApplyToSetsWithName = true);
 
 	/**
 	 * Calls ResetBodyModifierToCachedTarget for each of the body modifiers in the set
