Support using a mask when invoking control profiles

diff --git a/Engine/Plugins/Experimental/PhysicsControl/Source/PhysicsControl/Private/PhysicsControlComponent.cpp b/Engine/Plugins/Experimental/PhysicsControl/Source/PhysicsControl/Private/PhysicsControlComponent.cpp
--- a/Engine/Plugins/Experimental/PhysicsControl/Source/PhysicsControl/Private/PhysicsControlComponent.cpp
+++ b/Engine/Plugins/Experimental/PhysicsControl/Source/PhysicsControl/Private/PhysicsControlComponent.cpp
@@ -2285,7 +2285,8 @@ bool UPhysicsControlComponent::CreateControlsAndBodyModifiersFromPhysicsControlA
 }
 
 //======================================================================================================================
-bool UPhysicsControlComponent::InvokeControlProfile(FName ProfileName)
+bool UPhysicsControlComponent::InvokeControlProfile(
+	const FName ProfileName, const FName ControlSetMask, const FName BodyModifierSetMask)
 {
 	if (!PhysicsControlAsset.IsValid())
 	{
@@ -2310,31 +2311,43 @@ bool UPhysicsControlComponent::InvokeControlProfile(FName ProfileName)
 		return false;
 	}
 
-	ApplyControlAndModifierUpdates(*ControlAndModifierUpdates);
+	ApplyControlAndModifierUpdates(*ControlAndModifierUpdates, ControlSetMask, BodyModifierSetMask);
 
 	return true;
 }
 
 //======================================================================================================================
 void UPhysicsControlComponent::ApplyControlAndModifierUpdates(
-	const FPhysicsControlControlAndModifierUpdates& ControlAndModifierUpdates)
+	const FPhysicsControlControlAndModifierUpdates& ControlAndModifierUpdates,
+	const FName                                     ControlSetMask, 
+	const FName                                     BodyModifierSetMask)
 {
+	// Danny TODO Note that when we check the masks, we have to do an O(N) check, because the names
+	// are stored in an array. Would it be better to have them in a map, or at least sorted? Also
+	// note that the default empty name will not incur the search cost (i.e. leaving the name blank
+	// is better than saying "All").
+	const TArray<FName>& ControlsInMask = GetControlNamesInSet(ControlSetMask);
+	const TArray<FName>& BodyModifiersInMask = GetBodyModifierNamesInSet(BodyModifierSetMask);
+
 	for (const FPhysicsControlNamedControlParameters& ControlParameters : ControlAndModifierUpdates.ControlUpdates)
 	{
 		TArray<FName> Names = ExpandName(ControlParameters.Name, NameRecords.ControlSets);
 		for (FName Name : Names)
 		{
-			const FPhysicsControlSparseData& ControlData = ControlParameters.Data;
-			if (FPhysicsControlRecord* ControlRecord = ControlRecords.Find(Name))
-			{
-				ControlRecord->PhysicsControl.ControlData.UpdateFromSparseData(ControlData);
-			}
-			else
+			if (ControlsInMask.IsEmpty() || ControlsInMask.FindByKey(Name))
 			{
-				if (bWarnAboutInvalidNames)
+				const FPhysicsControlSparseData& ControlData = ControlParameters.Data;
+				if (FPhysicsControlRecord* ControlRecord = ControlRecords.Find(Name))
 				{
-					UE_LOG(LogPhysicsControl, Warning,
-						TEXT("ApplyControlAndModifierUpdates: Failed to find control with name %s"), *Name.ToString());
+					ControlRecord->PhysicsControl.ControlData.UpdateFromSparseData(ControlData);
+				}
+				else
+				{
+					if (bWarnAboutInvalidNames)
+					{
+						UE_LOG(LogPhysicsControl, Warning,
+							TEXT("ApplyControlAndModifierUpdates: Failed to find control with name %s"), *Name.ToString());
+					}
 				}
 			}
 		}
@@ -2346,17 +2359,20 @@ void UPhysicsControlComponent::ApplyControlAndModifierUpdates(
 		TArray<FName> Names = ExpandName(ControlMultiplierParameters.Name, NameRecords.ControlSets);
 		for (FName Name : Names)
 		{
-			const FPhysicsControlSparseMultiplier& Multiplier = ControlMultiplierParameters.Data;
-			if (FPhysicsControlRecord* ControlRecord = ControlRecords.Find(Name))
+			if (ControlsInMask.IsEmpty() || ControlsInMask.FindByKey(Name))
 			{
-				ControlRecord->PhysicsControl.ControlMultiplier.UpdateFromSparseData(Multiplier);
-			}
-			else
-			{
-				if (bWarnAboutInvalidNames)
+				const FPhysicsControlSparseMultiplier& Multiplier = ControlMultiplierParameters.Data;
+				if (FPhysicsControlRecord* ControlRecord = ControlRecords.Find(Name))
 				{
-					UE_LOG(LogPhysicsControl, Warning,
-						TEXT("ApplyControlAndModifierUpdates: Failed to find control with name %s"), *Name.ToString());
+					ControlRecord->PhysicsControl.ControlMultiplier.UpdateFromSparseData(Multiplier);
+				}
+				else
+				{
+					if (bWarnAboutInvalidNames)
+					{
+						UE_LOG(LogPhysicsControl, Warning,
+							TEXT("ApplyControlAndModifierUpdates: Failed to find control with name %s"), *Name.ToString());
+					}
 				}
 			}
 		}
@@ -2367,15 +2383,18 @@ void UPhysicsControlComponent::ApplyControlAndModifierUpdates(
 		TArray<FName> Names = ExpandName(ModifierParameters.Name, NameRecords.BodyModifierSets);
 		for (FName Name : Names)
 		{
-			const FPhysicsControlModifierSparseData& ModifierData = ModifierParameters.Data;
-			if (FPhysicsBodyModifierRecord* Record = BodyModifierRecords.Find(Name))
-			{
-				Record->BodyModifier.ModifierData.UpdateFromSparseData(ModifierData);
-			}
-			else
+			if (BodyModifiersInMask.IsEmpty() || BodyModifiersInMask.FindByKey(Name))
 			{
-				UE_LOG(LogPhysicsControl, Warning,
-					TEXT("InvokeControlProfile: Failed to find modifier with name %s"), *Name.ToString());
+				const FPhysicsControlModifierSparseData& ModifierData = ModifierParameters.Data;
+				if (FPhysicsBodyModifierRecord* Record = BodyModifierRecords.Find(Name))
+				{
+					Record->BodyModifier.ModifierData.UpdateFromSparseData(ModifierData);
+				}
+				else
+				{
+					UE_LOG(LogPhysicsControl, Warning,
+						TEXT("ApplyControlAndModifierUpdates: Failed to find modifier with name %s"), *Name.ToString());
+				}
 			}
 		}
 	}
diff --git a/Engine/Plugins/Experimental/PhysicsControl/Source/PhysicsControl/Public/PhysicsControlComponent.h b/Engine/Plugins/Experimental/PhysicsControl/Source/PhysicsControl/Public/PhysicsControlComponent.h
--- a/Engine/Plugins/Experimental/PhysicsControl/Source/PhysicsControl/Public/PhysicsControlComponent.h
+++ b/Engine/Plugins/Experimental/PhysicsControl/Source/PhysicsControl/Public/PhysicsControlComponent.h
@@ -311,7 +311,7 @@ public:
 
 	/**
 	 * Destroys a control or controls in a set
-	 * Standard set names will include "All", "WorldSpace", "ParentSpace" and things like "WorldSpace-ArmLeft",
+	 * Standard set names will include "All", "WorldSpace", "ParentSpace" and things like "WorldSpace_ArmLeft",
 	 * depending on how controls have been created.
 	 *
 	 * @param Name The name of the control or set of controls to destroy. 
@@ -324,7 +324,7 @@ public:
 
 	/**
 	 * Destroys all controls or sets of controls
-	 * Standard set names will include "All", "WorldSpace", "ParentSpace" and things like "WorldSpace-ArmLeft",
+	 * Standard set names will include "All", "WorldSpace", "ParentSpace" and things like "WorldSpace_ArmLeft",
 	 * depending on how controls have been created.
 	 *
 	 * @param Names The names of the controls or sets of controls to destroy. Note that if you have these 
@@ -340,7 +340,7 @@ public:
 	 * Destroys all controls in a set
 	 *
 	 * @param Set The set of controls to use to destroy. Standard sets will include "All", "WorldSpace",
-	 * "ParentSpace" and things like "WorldSpace-ArmLeft", depending on how controls have been created.
+	 * "ParentSpace" and things like "WorldSpace_ArmLeft", depending on how controls have been created.
 	 */
 	UFUNCTION(BlueprintCallable, Category = PhysicsControl)
 	UE_API void DestroyControlsInSet(const FName Set);
@@ -348,7 +348,7 @@ public:
 	/**
 	 * Updates the parent object part of a control or controls in a set. Note that this won't change 
 	 * the name of the control (which may subsequently be misleading), or any set it is included in, etc.
-	 * Standard set names will include "All", "WorldSpace", "ParentSpace" and things like "WorldSpace-ArmLeft",
+	 * Standard set names will include "All", "WorldSpace", "ParentSpace" and things like "WorldSpace_ArmLeft",
 	 * depending on how controls have been created.
 	 * 
 	 * @param bApplyToControlsWithName If true, then controls with the name will be modified
@@ -362,7 +362,7 @@ public:
 	/**
 	 * Updates the parent object part of controls or controls in a set. Note that this won't change 
 	 * the name of the controls (which may subsequently be misleading), or any set they are included in, etc.
-	 * Standard set names will include "All", "WorldSpace", "ParentSpace" and things like "WorldSpace-ArmLeft",
+	 * Standard set names will include "All", "WorldSpace", "ParentSpace" and things like "WorldSpace_ArmLeft",
 	 * depending on how controls have been created.
 	 * 
 	 * @param bApplyToControlsWithName If true, then controls with the name will be modified
@@ -382,7 +382,7 @@ public:
 	/**
 	 * Modifies an existing control data in a control or controls in a set - i.e. the strengths 
 	 * etc of the control driving towards the target
-	 * Standard set names will include "All", "WorldSpace", "ParentSpace" and things like "WorldSpace-ArmLeft",
+	 * Standard set names will include "All", "WorldSpace", "ParentSpace" and things like "WorldSpace_ArmLeft",
 	 * depending on how controls have been created.
 	 *
 	 * @param Name The name of the control or set of controls to update. 
@@ -398,7 +398,7 @@ public:
 	/**
 	 * Modifies existing control data in controls or controls in sets - i.e. the strengths 
 	 * etc of the controls driving towards the targets
-	 * Standard set names will include "All", "WorldSpace", "ParentSpace" and things like "WorldSpace-ArmLeft",
+	 * Standard set names will include "All", "WorldSpace", "ParentSpace" and things like "WorldSpace_ArmLeft",
 	 * depending on how controls have been created.
 	 *
 	 * @param Names The names of the controls or sets of controls to update. Note that if you have these in a FPhysicsControlNameArray
@@ -416,7 +416,7 @@ public:
 	 * Modifies existing control data - i.e. the strengths etc of the controls driving towards the targets
 	 *
 	 * @param Set The set of controls to update. Standard sets will include "All", "WorldSpace",
-	 *        "ParentSpace" and things like "WorldSpace-ArmLeft", depending on how controls have been created.
+	 *        "ParentSpace" and things like "WorldSpace_ArmLeft", depending on how controls have been created.
 	 * @param ControlData The new control data
 	 */
 	UFUNCTION(BlueprintCallable, Category = PhysicsControl)
@@ -425,7 +425,7 @@ public:
 	/**
 	 * Modifies an existing control data in a control or controls in a set - i.e. the strengths 
 	 * etc of the control driving towards the target
-	 * Standard set names will include "All", "WorldSpace", "ParentSpace" and things like "WorldSpace-ArmLeft",
+	 * Standard set names will include "All", "WorldSpace", "ParentSpace" and things like "WorldSpace_ArmLeft",
 	 * depending on how controls have been created.
 	 *
 	 * @param Name The name of the control or set of controls to update. 
@@ -441,7 +441,7 @@ public:
 	/**
 	 * Modifies existing control data in controls or controls in sets - i.e. the strengths 
 	 * etc of the controls driving towards the targets
-	 * Standard set names will include "All", "WorldSpace", "ParentSpace" and things like "WorldSpace-ArmLeft",
+	 * Standard set names will include "All", "WorldSpace", "ParentSpace" and things like "WorldSpace_ArmLeft",
 	 * depending on how controls have been created.
 	 *
 	 * @param Names The names of the controls or sets of controls to update. Note that if you have these in a FPhysicsControlNameArray
@@ -459,7 +459,7 @@ public:
 	 * Modifies existing control data - i.e. the strengths etc of the controls driving towards the targets
 	 *
 	 * @param Set The set of controls to update. Standard sets will include "All", "WorldSpace",
-	 *        "ParentSpace" and things like "WorldSpace-ArmLeft", depending on how controls have been created.
+	 *        "ParentSpace" and things like "WorldSpace_ArmLeft", depending on how controls have been created.
 	 * @param ControlData The new control data
 	 */
 	UFUNCTION(BlueprintCallable, Category = PhysicsControl)
@@ -467,7 +467,7 @@ public:
 
 	/**
 	 * Modifies existing control data in a control or controls in a set using the multipliers
-	 * Standard set names will include "All", "WorldSpace", "ParentSpace" and things like "WorldSpace-ArmLeft",
+	 * Standard set names will include "All", "WorldSpace", "ParentSpace" and things like "WorldSpace_ArmLeft",
 	 * depending on how controls have been created.
 	 *
 	 * @param Name The name of the control or set of controls to update. 
@@ -483,7 +483,7 @@ public:
 
 	/**
 	 * Modifies existing control data in controls or controls in sets using the multipliers
-	 * Standard set names will include "All", "WorldSpace", "ParentSpace" and things like "WorldSpace-ArmLeft",
+	 * Standard set names will include "All", "WorldSpace", "ParentSpace" and things like "WorldSpace_ArmLeft",
 	 * depending on how controls have been created.
 	 *
 	 * @param Names The names of the controls or sets of controls to update. Note that if you have these in a FPhysicsControlNameArray
@@ -502,7 +502,7 @@ public:
 	 * Modifies existing control data using the multipliers
 	 *
 	 * @param Set The set of controls to update. Standard sets will include "All", "WorldSpace",
-	 *        "ParentSpace" and things like "WorldSpace-ArmLeft", depending on how controls have been created.
+	 *        "ParentSpace" and things like "WorldSpace_ArmLeft", depending on how controls have been created.
 	 * @param ControlMultiplier The new control multiplier
 	 * @param bEnableControl Enables the controls if currently disabled
 	 */
@@ -512,7 +512,7 @@ public:
 
 	/**
 	 * Modifies existing control data in a control or controls in a set using the multipliers
-	 * Standard set names will include "All", "WorldSpace", "ParentSpace" and things like "WorldSpace-ArmLeft",
+	 * Standard set names will include "All", "WorldSpace", "ParentSpace" and things like "WorldSpace_ArmLeft",
 	 * depending on how controls have been created.
 	 *
 	 * @param Name The name of the control or set of controls to update. 
@@ -528,7 +528,7 @@ public:
 
 	/**
 	 * Modifies existing control data in controls or controls in sets using the multipliers
-	 * Standard set names will include "All", "WorldSpace", "ParentSpace" and things like "WorldSpace-ArmLeft",
+	 * Standard set names will include "All", "WorldSpace", "ParentSpace" and things like "WorldSpace_ArmLeft",
 	 * depending on how controls have been created.
 	 *
 	 * @param Names The names of the controls or sets of controls to update. Note that if you have these in a 
@@ -547,7 +547,7 @@ public:
 	 * Modifies existing control data using the multipliers
 	 *
 	 * @param Set The set of controls to update. Standard sets will include "All", "WorldSpace",
-	 *        "ParentSpace" and things like "WorldSpace-ArmLeft", depending on how controls have been created.
+	 *        "ParentSpace" and things like "WorldSpace_ArmLeft", depending on how controls have been created.
 	 * @param ControlMultiplier The new control multiplier
 	 * @param bEnableControl Enables the controls if currently disabled
 	 */
@@ -560,7 +560,7 @@ public:
 	/**
 	 * Modifies existing control's linear data in a control or controls in a set - i.e. the strengths 
 	 * etc of the control driving towards the target
-	 * Standard set names will include "All", "WorldSpace", "ParentSpace" and things like "WorldSpace-ArmLeft",
+	 * Standard set names will include "All", "WorldSpace", "ParentSpace" and things like "WorldSpace_ArmLeft",
 	 * depending on how controls have been created.
 	 *
 	 * @param Name The name of the control or set of controls to update. 
@@ -586,7 +586,7 @@ public:
 	/**
 	 * Modifies existing control's angular data in a control or controls in a set - i.e. the strengths 
 	 * etc of the control driving towards the target
-	 * Standard set names will include "All", "WorldSpace", "ParentSpace" and things like "WorldSpace-ArmLeft",
+	 * Standard set names will include "All", "WorldSpace", "ParentSpace" and things like "WorldSpace_ArmLeft",
 	 * depending on how controls have been created.
 	 *
 	 * @param Name The name of the control or set of controls to update. 
@@ -630,7 +630,7 @@ public:
 	/**
 	 * Modifies an existing control target in a control or controls in a set - i.e. what it is driving 
 	 * towards, relative to the parent object
-	 * Standard set names will include "All", "WorldSpace", "ParentSpace" and things like "WorldSpace-ArmLeft",
+	 * Standard set names will include "All", "WorldSpace", "ParentSpace" and things like "WorldSpace_ArmLeft",
 	 * depending on how controls have been created.
 	 *
 	 * @param Name The name of the control or set of controls to update. 
@@ -647,7 +647,7 @@ public:
 	/**
 	 * Modifies existing control targets in controls or controls in sets - i.e. what they are driving 
 	 * towards, relative to the parent objects
-	 * Standard set names will include "All", "WorldSpace", "ParentSpace" and things like "WorldSpace-ArmLeft",
+	 * Standard set names will include "All", "WorldSpace", "ParentSpace" and things like "WorldSpace_ArmLeft",
 	 * depending on how controls have been created.
 	 *
 	 * @param Names The names of the controls or sets of controls to update. Note that if you have these in a 
@@ -667,7 +667,7 @@ public:
 	 * Modifies existing control targets - i.e. what they are driving towards, relative to the parent objects
 	 *
 	 * @param Set The set of controls to update. Standard sets will include "All", "WorldSpace",
-	 *        "ParentSpace" and things like "WorldSpace-ArmLeft", depending on how controls have been created.
+	 *        "ParentSpace" and things like "WorldSpace_ArmLeft", depending on how controls have been created.
 	 * @param ControlTarget The new target for the controls
 	 * @param bEnableControl Enables the controls if currently disabled
 	 */
@@ -680,7 +680,7 @@ public:
 	/**
 	 * Modifies an existing control target in a control or controls in a set - i.e. what it is driving 
 	 * towards, relative to the parent object
-	 * Standard set names will include "All", "WorldSpace", "ParentSpace" and things like "WorldSpace-ArmLeft",
+	 * Standard set names will include "All", "WorldSpace", "ParentSpace" and things like "WorldSpace_ArmLeft",
 	 * depending on how controls have been created.
 	 *
 	 * @param Name The name of the control or set of controls to update. 
@@ -712,7 +712,7 @@ public:
 	/**
 	 * Modifies existing control targets in controls or controls in sets - i.e. what it is driving
 	 * towards, relative to the parent object
-	 * Standard set names will include "All", "WorldSpace", "ParentSpace" and things like "WorldSpace-ArmLeft",
+	 * Standard set names will include "All", "WorldSpace", "ParentSpace" and things like "WorldSpace_ArmLeft",
 	 * depending on how controls have been created.
 	 *
 	 * @param Names The name of the controls or sets of controls to update.
@@ -756,7 +756,7 @@ public:
 	/**
 	 * Modifies an existing control target of a control or controls in a set - i.e. what it is driving 
 	 * towards, relative to the parent object
-	 * Standard set names will include "All", "WorldSpace", "ParentSpace" and things like "WorldSpace-ArmLeft",
+	 * Standard set names will include "All", "WorldSpace", "ParentSpace" and things like "WorldSpace_ArmLeft",
 	 * depending on how controls have been created.
 	 *
 	 * @param Name The name of the control or set of controls to update. 
@@ -787,7 +787,7 @@ public:
 	/**
 	 * Modifies existing control targets in controls or controls in sets - i.e. what it is driving
 	 * towards, relative to the parent object
-	 * Standard set names will include "All", "WorldSpace", "ParentSpace" and things like "WorldSpace-ArmLeft",
+	 * Standard set names will include "All", "WorldSpace", "ParentSpace" and things like "WorldSpace_ArmLeft",
 	 * depending on how controls have been created.
 	 *
 	 * @param Names The name of the controls or sets of controls to update.
@@ -828,7 +828,7 @@ public:
 	/**
 	 * Modifies an existing control target in a control or controls in a set - i.e. what it is driving 
 	 * towards, relative to the parent object
-	 * Standard set names will include "All", "WorldSpace", "ParentSpace" and things like "WorldSpace-ArmLeft",
+	 * Standard set names will include "All", "WorldSpace", "ParentSpace" and things like "WorldSpace_ArmLeft",
 	 * depending on how controls have been created.
 	 *
 	 * @param Name The name of the control or set of controls to update. 
@@ -859,7 +859,7 @@ public:
 	/**
 	 * Modifies existing control targets in controls or controls in sets - i.e. what it is driving
 	 * towards, relative to the parent object
-	 * Standard set names will include "All", "WorldSpace", "ParentSpace" and things like "WorldSpace-ArmLeft",
+	 * Standard set names will include "All", "WorldSpace", "ParentSpace" and things like "WorldSpace_ArmLeft",
 	 * depending on how controls have been created.
 	 *
 	 * @param Name The name of the control or set of controls to update.
@@ -964,7 +964,7 @@ public:
 
 	/**
 	 * Sets whether or not the control or controls in a set should use skeletal animation for the targets
-	 * Standard set names will include "All", "WorldSpace", "ParentSpace" and things like "WorldSpace-ArmLeft",
+	 * Standard set names will include "All", "WorldSpace", "ParentSpace" and things like "WorldSpace_ArmLeft",
 	 * depending on how controls have been created.
 	 *
 	 * @param Name The name of the control or set of controls to update. 
@@ -988,7 +988,7 @@ public:
 
 	/**
 	 * Sets whether or not the controls or controls in sets should use skeletal animation for the targets
-	 * Standard set names will include "All", "WorldSpace", "ParentSpace" and things like "WorldSpace-ArmLeft",
+	 * Standard set names will include "All", "WorldSpace", "ParentSpace" and things like "WorldSpace_ArmLeft",
 	 * depending on how controls have been created.
 	 *
 	 * @param Names The names of the controls or sets of controls to update. Note that if you have these in a FPhysicsControlNameArray 
@@ -1015,7 +1015,7 @@ public:
 	 * Sets whether or not the controls should use skeletal animation for the targets
 	 *
 	 * @param Set The set of controls to update. Standard sets will include "All", "WorldSpace",
-	 *        "ParentSpace" and things like "WorldSpace-ArmLeft", depending on how controls have been created.
+	 *        "ParentSpace" and things like "WorldSpace_ArmLeft", depending on how controls have been created.
 	 * @param bUseSkeletalAnimation If true then the targets will be a combination of the skeletal animation (if
 	 *        there is any) and the control target that has been set
 	 * @param AngularTargetVelocityMultiplier If skeletal animation is being used, then this determines the amount of
@@ -1032,7 +1032,7 @@ public:
 
 	/**
 	 * Activates or deactivates a control or controls in a set.
-	 * Standard set names will include "All", "WorldSpace", "ParentSpace" and things like "WorldSpace-ArmLeft", 
+	 * Standard set names will include "All", "WorldSpace", "ParentSpace" and things like "WorldSpace_ArmLeft", 
 	 * depending on how controls have been created.
 	 *
 	 * @param Name     The name of the control or set of controls to update. 
@@ -1047,7 +1047,7 @@ public:
 
 	/**
 	 * Activates or deactivates controls or sets of controls
-	 * Standard set names will include "All", "WorldSpace", "ParentSpace" and things like "WorldSpace-ArmLeft",
+	 * Standard set names will include "All", "WorldSpace", "ParentSpace" and things like "WorldSpace_ArmLeft",
 	 * depending on how controls have been created.
 	 *
 	 * @param Names The names of the controls or sets of controls to update. Note that if you have these in a FPhysicsControlNameArray
@@ -1065,7 +1065,7 @@ public:
 	 * Activates or deactivates controls
 	 *
 	 * @param Set The set of controls to update. Standard sets will include "All", "WorldSpace",
-	 *        "ParentSpace" and things like "WorldSpace-ArmLeft", depending on how controls have been created.
+	 *        "ParentSpace" and things like "WorldSpace_ArmLeft", depending on how controls have been created.
 	 * @param bEnable  Whether to enable/disable the controls
 	 */
 
@@ -1074,7 +1074,7 @@ public:
 
 	/**
 	 * @param Name The name of the control or set of controls to update. Standard sets will include "All", "WorldSpace",
-	 *        "ParentSpace" and things like "WorldSpace-ArmLeft", depending on how controls have been created.
+	 *        "ParentSpace" and things like "WorldSpace_ArmLeft", depending on how controls have been created.
 	 * @param bDisableCollision If set then the control will disable collision between the bodies it connects.
 	 * @param bApplyToControlsWithName If true, then controls with the name will be modified
 	 * @param bApplyToSetsWithName If true, then controls in sets with the name will be modified
@@ -1086,7 +1086,7 @@ public:
 
 	/**
 	 * @param Names The names of the controls or sets of controls to update. Standard sets will include "All", "WorldSpace",
-	 *        "ParentSpace" and things like "WorldSpace-ArmLeft", depending on how controls have been created.
+	 *        "ParentSpace" and things like "WorldSpace_ArmLeft", depending on how controls have been created.
 	 * @param bDisableCollision If set then the control will disable collision between the bodies it connects.
 	 * @param bApplyToControlsWithName If true, then controls with the name will be modified
 	 * @param bApplyToSetsWithName If true, then controls in sets with the name will be modified
@@ -1098,7 +1098,7 @@ public:
 
 	/**
 	 * @param Set The set of controls to update. Standard sets will include "All", "WorldSpace",
-	 *        "ParentSpace" and things like "WorldSpace-ArmLeft", depending on how controls have been created.
+	 *        "ParentSpace" and things like "WorldSpace_ArmLeft", depending on how controls have been created.
 	 * @param bDisableCollision If set then the control will disable collision between the bodies it connects.
 	 */
 	UFUNCTION(BlueprintCallable, Category = PhysicsControl)
@@ -1634,10 +1634,20 @@ public:
 	/**
 	 * Looks up the profile which should exist in the registered control profile asset, and invokes it.
 	 * 
+	 * @param ProfileName The name of the control profile to apply
+	 * @param ControlSetMask The optional name of a control set to use as a mask, so only controls within that set will be modified.
+	 *        Standard control set names will include "All", "WorldSpace", "ParentSpace" and things like "WorldSpace_ArmLeft",
+	 *        depending on how controls have been created. If unset or invalid, no mask will be applied.
+	 * @param BodyModifierSetMask The optional name of a body modifier set to use as a mask, so only body modifiers within that set will be modified.
+	 *        Standard body modifier sets will include "All" and things like "ArmLeft", depending on how body modifiers 
+	 *        have been created. If unset or invalid, no mask will be applied.
 	 * @return true if successful, and false if the profile cannot be found.
 	 */
 	UFUNCTION(BlueprintCallable, Category = PhysicsControl, Meta = (ReturnDisplayName = "Success"))
-	UE_API bool InvokeControlProfile(FName ProfileName);
+	UE_API bool InvokeControlProfile(
+		const FName ProfileName, 
+		const FName ControlSetMask = TEXT(""),
+		const FName BodyModifierSetMask = TEXT(""));
 
 	/**
 	 * Adds a Control to a Set. This will add a new set if necessary. For example, you might
@@ -1668,7 +1678,7 @@ public:
 	/**
 	 * Returns a reference to all the control names that have been created and are in the specified 
 	 * set, which could be a limb, or a subsequently created set. Standard sets will include "All", "WorldSpace",
-	 * "ParentSpace" and things like "WorldSpace-ArmLeft", depending on how controls have been created.
+	 * "ParentSpace" and things like "WorldSpace_ArmLeft", depending on how controls have been created.
 	 */
 	UFUNCTION(BlueprintCallable, Category = PhysicsControl)
 	UE_API const TArray<FName>& GetControlNamesInSet(const FName Set) const;
@@ -2085,9 +2095,21 @@ protected:
 	/** Destroys the modifier. It will optionally be removed from the array of records too */
 	UE_API bool DestroyBodyModifier(const FName Name, const EDestroyBehavior DestroyBehavior);
 
-	/** Applies the updates to controls and body modifiers */
+	/** 
+	 * Applies the updates to controls and body modifiers
+	 * 
+	 * @param ControlAndModifierUpdates The updates to apply
+	 * @param ControlSetMask The optional name of a control set to use as a mask, so only controls within that set will be modified.
+	 *        Standard control set names will include "All", "WorldSpace", "ParentSpace" and things like "WorldSpace_ArmLeft",
+	 *        depending on how controls have been created.  If unset or invalid, no mask will be applied.
+	 * @param BodyModifierSetMask The optional name of a body modifier set to use as a mask, so only body modifiers within that set will be modified.
+	 *        Standard body modifier sets will include "All" and things like "ArmLeft", depending on how body modifiers 
+	 *        have been created. If unset or invalid, no mask will be applied.
+	 */
 	UE_API void ApplyControlAndModifierUpdates(
-		const FPhysicsControlControlAndModifierUpdates& ControlAndModifierUpdates);
+		const FPhysicsControlControlAndModifierUpdates& ControlAndModifierUpdates, 
+		const FName                                     ControlSetMask = TEXT(""),
+		const FName                                     BodyModifierSetMask = TEXT(""));
 
 protected:
 
