Disable diagnostic logging in RigidBodyWithControl in test/shipping and demote level to verbose.

This can be extremely expensive; upwards of 300ms so we want to ensure this never runs in live.

#rnx

diff --git a/Engine/Plugins/Experimental/PhysicsControl/Source/PhysicsControl/Private/RigidBodyWithControl.cpp b/Engine/Plugins/Experimental/PhysicsControl/Source/PhysicsControl/Private/RigidBodyWithControl.cpp
--- a/Engine/Plugins/Experimental/PhysicsControl/Source/PhysicsControl/Private/RigidBodyWithControl.cpp
+++ b/Engine/Plugins/Experimental/PhysicsControl/Source/PhysicsControl/Private/RigidBodyWithControl.cpp
@@ -491,9 +491,10 @@ void FAnimNode_RigidBodyWithControl::DestroyControlsAndBodyModifiers()
 // TODO This isn't ideal as it only dumps out the "original" values, not including the updates
 void FAnimNode_RigidBodyWithControl::LogControlsModifiersAndSets()
 {
+#if !UE_BUILD_TEST && !UE_BUILD_SHIPPING
 	SCOPE_CYCLE_COUNTER(STAT_RigidBodyNodeWithControl_LogControlsModifiersAndSets);
 
-#define RBWC_LOG_LEVEL Display
+#define RBWC_LOG_LEVEL Verbose
 
 	UE_LOG(LogPhysicsControl, RBWC_LOG_LEVEL, TEXT("Controls:"));
 	for (TMap<FName, FRigidBodyControlRecord>::ElementType& NameRecordPair : ControlRecords)
@@ -548,6 +549,8 @@ void FAnimNode_RigidBodyWithControl::LogControlsModifiersAndSets()
 		}
 	}
 #undef RBWC_LOG_LEVEL
+
+#endif // !UE_BUILD_TEST && !UE_BUILD_SHIPPING
 }
 
 //======================================================================================================================
