Sets the PhysicsControl plugin to Beta

diff --git a/Engine/Plugins/Experimental/PhysicsControl/PhysicsControl.uplugin b/Engine/Plugins/Experimental/PhysicsControl/PhysicsControl.uplugin
--- a/Engine/Plugins/Experimental/PhysicsControl/PhysicsControl.uplugin
+++ b/Engine/Plugins/Experimental/PhysicsControl/PhysicsControl.uplugin
@@ -3,7 +3,7 @@
 	"Version": 1,
 	"VersionName": "1.0",
 	"FriendlyName": "PhysicsControl",
-	"Description": "Additional Support for controlling static and skeletal meshes through physical controls",
+	"Description": "Physically control static and skeletal meshes through the Physics Control Component and the Rigid Body With Control animation graph node.",
 	"Category": "Physics",
 	"CreatedBy": "Epic Games, Inc.",
 	"CreatedByURL": "https://epicgames.com",
@@ -12,8 +12,8 @@
 	"SupportURL": "",
 	"CanContainContent": true,
 	"CanContainVerse": false,
-	"IsBetaVersion": false,
-	"IsExperimentalVersion": true,
+	"IsBetaVersion": true,
+	"IsExperimentalVersion": false,
 	"Installed": false,
 	"Modules": [
 		{
