[UE] RBwC Support for setting physics collision type for world collisions
- Add three new bool properties to the RBwC node that determine the physics collision types that a world physics object must have to be included in the immedidate physics simulation.
- Added some pretty UI to set these flags

#DisabledValidations Preflight

diff --git a/Engine/Plugins/Experimental/PhysicsControl/Source/PhysicsControl/Private/AnimNode_RigidBodyWithControl.cpp b/Engine/Plugins/Experimental/PhysicsControl/Source/PhysicsControl/Private/AnimNode_RigidBodyWithControl.cpp
--- a/Engine/Plugins/Experimental/PhysicsControl/Source/PhysicsControl/Private/AnimNode_RigidBodyWithControl.cpp
+++ b/Engine/Plugins/Experimental/PhysicsControl/Source/PhysicsControl/Private/AnimNode_RigidBodyWithControl.cpp
@@ -146,6 +146,9 @@ FAnimNode_RigidBodyWithControl::FAnimNode_RigidBodyWithControl()
 	, bMakeKinematicConstraints(true)
 	, ResetSimulatedTeleportType(ETeleportType::None)
 	, bEnableWorldGeometry(false)
+	, bWorldGeometryCollisionEnabledHasPhysics(true)
+	, bWorldGeometryCollisionEnabledHasQuery(false)
+	, bWorldGeometryCollisionEnabledHasProbe(false)
 	, bOverrideWorldGravity(false)
 	, bTransferBoneVelocities(false)
 	, bFreezeIncomingPoseOnStart(false)
@@ -1624,7 +1627,16 @@ void FAnimNode_RigidBodyWithControl::UpdateClothColliderObjects(const FTransform
 	}
 }
 
-inline bool IsComponentDesiredInSim(const UPrimitiveComponent* Component)
+inline bool IsComponentCollisionEnabled(const FAnimNode_RigidBodyWithControl* Node, const UPrimitiveComponent* const Component)
+{
+	const ECollisionEnabled::Type CollisionEnabled = Component->GetBodyInstance()->GetCollisionEnabled();
+
+	return (Node->bWorldGeometryCollisionEnabledHasPhysics && CollisionEnabledHasPhysics(CollisionEnabled))
+		|| (Node->bWorldGeometryCollisionEnabledHasQuery && CollisionEnabledHasQuery(CollisionEnabled))
+		|| (Node->bWorldGeometryCollisionEnabledHasProbe && CollisionEnabledHasProbe(CollisionEnabled));
+}
+
+inline bool IsComponentDesiredInSim(const FAnimNode_RigidBodyWithControl* Node, const UPrimitiveComponent* Component)
 {
 	if (!IsValid(Component))
 	{
@@ -1638,7 +1650,7 @@ inline bool IsComponentDesiredInSim(const UPrimitiveComponent* Component)
 	{
 		return false;
 	}
-	if (!CollisionEnabledHasPhysics(Component->GetBodyInstance()->GetCollisionEnabled()))
+	if(!IsComponentCollisionEnabled(Node, Component))
 	{
 		return false;
 	}
@@ -1671,7 +1683,7 @@ void FAnimNode_RigidBodyWithControl::CollectWorldObjects()
 				{
 					// New object - add it to the sim
 					const bool bIsSelf = (UnsafeOwner == OverlapComp->GetOwner());
-					if (!bIsSelf && IsComponentDesiredInSim(OverlapComp))
+					if (!bIsSelf && IsComponentDesiredInSim(this, OverlapComp))
 					{
 						// Note that the TM here isn't correct unless we are actually using
 						// world-space simulation, but it will be updated.
@@ -1708,7 +1720,7 @@ void FAnimNode_RigidBodyWithControl::ExpireWorldObjects()
 		// Do we need to expire this object?
 		const int32 ExpireTickCount = RBANWithControl_WorldObjectExpiry;
 		if (((ComponentsInSimTick - WorldObject.LastSeenTick) > ExpireTickCount)	// Haven't seen this object for a while
-			|| !IsComponentDesiredInSim(WorldComp))
+			|| !IsComponentDesiredInSim(this, WorldComp))
 		{
 			WorldObject.bExpired = true;
 		}
diff --git a/Engine/Plugins/Experimental/PhysicsControl/Source/PhysicsControl/Public/AnimNode_RigidBodyWithControl.h b/Engine/Plugins/Experimental/PhysicsControl/Source/PhysicsControl/Public/AnimNode_RigidBodyWithControl.h
--- a/Engine/Plugins/Experimental/PhysicsControl/Source/PhysicsControl/Public/AnimNode_RigidBodyWithControl.h
+++ b/Engine/Plugins/Experimental/PhysicsControl/Source/PhysicsControl/Public/AnimNode_RigidBodyWithControl.h
@@ -246,9 +246,18 @@ private:
 	ETeleportType ResetSimulatedTeleportType;
 
 public:
-	UPROPERTY(EditAnywhere, Category = Settings, meta = (InlineEditConditionToggle))
+	UPROPERTY(EditAnywhere, Category = Settings)
 	uint8 bEnableWorldGeometry : 1;
 
+	UPROPERTY(EditAnywhere, Category = Settings)
+	uint8 bWorldGeometryCollisionEnabledHasPhysics : 1;
+	
+	UPROPERTY(EditAnywhere, Category = Settings)
+	uint8 bWorldGeometryCollisionEnabledHasQuery : 1;
+	
+	UPROPERTY(EditAnywhere, Category = Settings)
+	uint8 bWorldGeometryCollisionEnabledHasProbe : 1;
+
 	UPROPERTY(EditAnywhere, Category = Settings, meta = (InlineEditConditionToggle))
 	uint8 bOverrideWorldGravity : 1;
 
@@ -535,7 +544,7 @@ private:
 		const MapConstraintsBehaviorType OrientationBehavior,
 		TArray<FConstraintInstance*>& ConstraintInstances);
 
-	// Gather cloth collision sources from the supplied Skeltal Mesh and add a kinematic actor
+	// Gather cloth collision sources from the supplied Skeletal Mesh and add a kinematic actor
 	// representing each one of them to the sim.
 	UE_API void CollectClothColliderObjects(const USkeletalMeshComponent* SkeletalMeshComp);
 	
diff --git a/Engine/Plugins/Experimental/PhysicsControl/Source/PhysicsControlUncookedOnly/Private/AnimGraphNode_RigidBodyWithControl.cpp b/Engine/Plugins/Experimental/PhysicsControl/Source/PhysicsControlUncookedOnly/Private/AnimGraphNode_RigidBodyWithControl.cpp
--- a/Engine/Plugins/Experimental/PhysicsControl/Source/PhysicsControlUncookedOnly/Private/AnimGraphNode_RigidBodyWithControl.cpp
+++ b/Engine/Plugins/Experimental/PhysicsControl/Source/PhysicsControlUncookedOnly/Private/AnimGraphNode_RigidBodyWithControl.cpp
@@ -18,6 +18,7 @@
 #include "DetailWidgetRow.h"
 #include "DetailCategoryBuilder.h"
 #include "Widgets/Input/SButton.h"
+#include "Widgets/Input/SCheckBox.h"
 #include "Widgets/Text/STextBlock.h"
 #include "PrimitiveDrawingUtils.h"
 
@@ -225,6 +226,94 @@ void UAnimGraphNode_RigidBodyWithControl::CustomizeDetails(IDetailLayoutBuilder&
 				]
 			];
 	}
+
+	// World Collision Enable checkbox and collision type flags.
+	{
+		IDetailCategoryBuilder& Category = DetailBuilder.EditCategory("Settings");
+
+		TSharedPtr<IPropertyHandle> NodeHandle = DetailBuilder.GetProperty(GET_MEMBER_NAME_CHECKED(UAnimGraphNode_RigidBodyWithControl, Node), UAnimGraphNode_RigidBodyWithControl::StaticClass());
+
+		if(ensure(NodeHandle.IsValid()))
+		{
+			TSharedPtr<IPropertyHandle> EnableWorldCollision = NodeHandle->GetChildHandle(GET_MEMBER_NAME_CHECKED(FAnimNode_RigidBodyWithControl, bEnableWorldGeometry));
+																																			
+			TSharedPtr<IPropertyHandle> HasPhysics = NodeHandle->GetChildHandle(GET_MEMBER_NAME_CHECKED(FAnimNode_RigidBodyWithControl, bWorldGeometryCollisionEnabledHasPhysics));
+			TSharedPtr<IPropertyHandle> HasQuery = NodeHandle->GetChildHandle(GET_MEMBER_NAME_CHECKED(FAnimNode_RigidBodyWithControl, bWorldGeometryCollisionEnabledHasQuery));
+			TSharedPtr<IPropertyHandle> HasProbe = NodeHandle->GetChildHandle(GET_MEMBER_NAME_CHECKED(FAnimNode_RigidBodyWithControl, bWorldGeometryCollisionEnabledHasProbe));
+
+			if(ensure(EnableWorldCollision) && ensure(HasPhysics) && ensure(HasQuery) && ensure(HasProbe))
+			{
+				DetailBuilder.HideProperty(EnableWorldCollision);
+				DetailBuilder.HideProperty(HasPhysics);
+				DetailBuilder.HideProperty(HasQuery);
+				DetailBuilder.HideProperty(HasProbe);
+
+				auto MakeCheckbox = [EnableWorldCollision](const FString& Name, TSharedPtr<IPropertyHandle> Property, const bool IsCollisionType = true) -> TSharedRef<SCheckBox>
+				{
+					FString ToolTipText = "Enable collision with world geometry";
+				
+					if(IsCollisionType)
+					{
+						ToolTipText += FString::Printf(TEXT(" with %s collision type"), *Name);
+					}
+					
+					return SNew(SCheckBox)
+						.IsChecked_Lambda([Property]()  -> ECheckBoxState { bool bValue = false; Property->GetValue(bValue); return (bValue)? ECheckBoxState::Checked : ECheckBoxState::Unchecked; })
+						.OnCheckStateChanged_Lambda([Property](ECheckBoxState InState) { Property->SetValue(InState == ECheckBoxState::Checked); })
+						.ToolTipText(FText::FromString(ToolTipText))
+						[
+							SNew(STextBlock)
+								.Text(FText::FromString(Name))
+						];
+				};
+
+				auto IsCheckboxEnabled = [EnableWorldCollision]() -> bool 
+				{ 
+					bool bValue = false; 
+					EnableWorldCollision->GetValue(bValue); 
+					return bValue; 
+				};
+
+				const float CheckBoxPadding = 2.0f;
+
+				TSharedRef<SCheckBox> PhysicsWidget = MakeCheckbox("Physics", HasPhysics);
+				PhysicsWidget->SetEnabled(TAttribute<bool>::CreateLambda(IsCheckboxEnabled));
+				TSharedRef<SCheckBox> QueryWidget = MakeCheckbox("Query", HasQuery);
+				QueryWidget->SetEnabled(TAttribute<bool>::CreateLambda(IsCheckboxEnabled));
+				TSharedRef<SCheckBox> ProbeWidget = MakeCheckbox("Probe", HasProbe);
+				ProbeWidget->SetEnabled(TAttribute<bool>::CreateLambda(IsCheckboxEnabled));
+
+				// Add them to the same row
+				Category.AddCustomRow(LOCTEXT("WorldCollisionRow", "WorldCollision"))
+					.NameContent()
+					[
+						MakeCheckbox("World Collision", EnableWorldCollision, false)
+					]
+					.ValueContent()
+					[
+						SNew(SHorizontalBox)
+							+ SHorizontalBox::Slot()
+							.AutoWidth()
+							.Padding(CheckBoxPadding)
+							[
+								PhysicsWidget
+							]
+							+ SHorizontalBox::Slot()
+							.AutoWidth()
+							.Padding(CheckBoxPadding)
+							[
+								QueryWidget
+							]
+							+ SHorizontalBox::Slot()							
+							.AutoWidth()
+							.Padding(CheckBoxPadding)
+							[
+								ProbeWidget
+							]
+					];
+			}
+		}
+	}
 }
 
 void UAnimGraphNode_RigidBodyWithControl::PostChange()
