Unshelved from pending changelist '47306592':

[UE] PBA Exclude unwanted PrimitiveComponents from RBwC simluation
- Exclude PrimitiveComponents that do not block the ECC_PhysicsBody channel for collisions

#rnx

diff --git a/Engine/Plugins/Experimental/PhysicsControl/Source/PhysicsControl/Private/AnimNode_RigidBodyWithControl.cpp b/Engine/Plugins/Experimental/PhysicsControl/Source/PhysicsControl/Private/AnimNode_RigidBodyWithControl.cpp
--- a/Engine/Plugins/Experimental/PhysicsControl/Source/PhysicsControl/Private/AnimNode_RigidBodyWithControl.cpp
+++ b/Engine/Plugins/Experimental/PhysicsControl/Source/PhysicsControl/Private/AnimNode_RigidBodyWithControl.cpp
@@ -1629,6 +1629,11 @@ void FAnimNode_RigidBodyWithControl::UpdateClothColliderObjects(const FTransform
 
 inline bool IsComponentCollisionEnabled(const FAnimNode_RigidBodyWithControl* Node, const UPrimitiveComponent* const Component)
 {
+	if (Component->GetCollisionResponseToChannel(ECC_PhysicsBody) != ECollisionResponse::ECR_Block)
+	{
+		return false;
+	}
+
 	const ECollisionEnabled::Type CollisionEnabled = Component->GetBodyInstance()->GetCollisionEnabled();
 
 	return (Node->bWorldGeometryCollisionEnabledHasPhysics && CollisionEnabledHasPhysics(CollisionEnabled))
@@ -1683,6 +1688,7 @@ void FAnimNode_RigidBodyWithControl::CollectWorldObjects()
 				{
 					// New object - add it to the sim
 					const bool bIsSelf = (UnsafeOwner == OverlapComp->GetOwner());
+
 					if (!bIsSelf && IsComponentDesiredInSim(this, OverlapComp))
 					{
 						// Note that the TM here isn't correct unless we are actually using
