Fixes for Physics Control Component names, so they're never called "". Add functions to log controls/modifiers. Adds checks that controls/modifiers/sets that are passed in are valid (and optionally warns if not). Note that this will change the names of some controls/modifiers, and if those are being referenced explicitly (by name) the assets will need to be patched up.

diff --git a/Engine/Plugins/Experimental/PhysicsControl/Source/PhysicsControl/Private/PhysicsControlComponent.cpp b/Engine/Plugins/Experimental/PhysicsControl/Source/PhysicsControl/Private/PhysicsControlComponent.cpp
--- a/Engine/Plugins/Experimental/PhysicsControl/Source/PhysicsControl/Private/PhysicsControlComponent.cpp
+++ b/Engine/Plugins/Experimental/PhysicsControl/Source/PhysicsControl/Private/PhysicsControlComponent.cpp
@@ -48,6 +48,24 @@ static FName GetControlTypeName(EPhysicsControlType ControlType)
 	}
 }
 
+//======================================================================================================================
+// Helper for disabling warnings about control etc names not existing. Used when looking for
+// controls or modifiers that could be referenced by name or set.
+struct FDisableNameWarnings
+{
+	FDisableNameWarnings(bool& bOriginal)
+		: bSaved(bOriginal), bWarnAboutInvalidNames(bOriginal)
+	{
+		bOriginal = false;
+	}
+	~FDisableNameWarnings()
+	{
+		bWarnAboutInvalidNames = bSaved;
+	}
+	bool bSaved;
+	bool& bWarnAboutInvalidNames;
+};
+
 //======================================================================================================================
 UPhysicsControlComponent::UPhysicsControlComponent(const FObjectInitializer& ObjectInitializer)
 	: Super(ObjectInitializer)
@@ -154,7 +172,8 @@ TMap<FName, FPhysicsControlLimbBones> UPhysicsControlComponent::GetLimbBonesFrom
 	UPhysicsAsset* PhysicsAsset = SkeletalMeshComponent ? SkeletalMeshComponent->GetPhysicsAsset() : nullptr;
 	if (!PhysicsAsset)
 	{
-		UE_LOG(LogPhysicsControl, Warning, TEXT("No physics asset in skeletal mesh"));
+		UE_LOG(LogPhysicsControl, Warning, TEXT("No physics asset in skeletal mesh in %s"), 
+			GetOwner() ? *GetOwner()->GetName() : TEXT("Unknown"));
 		return Result;
 	}
 
@@ -180,7 +199,8 @@ FName UPhysicsControlComponent::CreateControl(
 	const FName                   Set,
 	const FString                 NamePrefix)
 {
-	const FName Name = UE::PhysicsControl::GetUniqueControlName(ParentBoneName, ChildBoneName, ControlRecords, NamePrefix);
+	const FName Name = UE::PhysicsControl::GetUniqueControlName(
+		ParentComponent, ParentBoneName, ChildComponent, ChildBoneName, ControlRecords, NamePrefix);
 	if (CreateNamedControl(
 		Name, ParentComponent, ParentBoneName, ChildComponent, ChildBoneName, ControlData, ControlTarget, Set))
 	{
@@ -204,14 +224,16 @@ bool UPhysicsControlComponent::CreateNamedControl(
 	if (FindControlRecord(Name))
 	{
 		UE_LOG(LogPhysicsControl, Warning,
-			TEXT("Unable to make a Control as one with the desired name already exists"), *Name.ToString());
+			TEXT("Unable to make a Control as one with the desired name %s already exists in %s"),
+			*Name.ToString(), GetOwner() ? *GetOwner()->GetName() : TEXT("Unknown"));
 		return false;
 	}
 
 	if (!ChildComponent)
 	{
 		UE_LOG(LogPhysicsControl, Warning,
-			TEXT("Unable to make a Control as the child mesh component has not been set"));
+			TEXT("Unable to make a Control as the child mesh component has not been set in %s"), 
+			GetOwner() ? *GetOwner()->GetName() : TEXT("Unknown"));
 		return false;
 	}
 
@@ -249,7 +271,8 @@ TArray<FName> UPhysicsControlComponent::CreateControlsFromSkeletalMeshBelow(
 	UPhysicsAsset* PhysicsAsset = SkeletalMeshComponent ? SkeletalMeshComponent->GetPhysicsAsset() : nullptr;
 	if (!PhysicsAsset)
 	{
-		UE_LOG(LogPhysicsControl, Warning, TEXT("No physics asset in skeletal mesh"));
+		UE_LOG(LogPhysicsControl, Warning, TEXT("No physics asset in skeletal mesh in %s"),
+			GetOwner() ? *GetOwner()->GetName() : TEXT("Unknown"));
 		return Result;
 	}
 
@@ -289,7 +312,8 @@ TArray<FName> UPhysicsControlComponent::CreateControlsFromSkeletalMeshBelow(
 				else
 				{
 					UE_LOG(LogPhysicsControl, Warning, 
-						TEXT("Failed to make control for %s"), *ChildBoneName.ToString());
+						TEXT("Failed to make control for %s in %s"),
+						*ChildBoneName.ToString(), GetOwner() ? *GetOwner()->GetName() : TEXT("Unknown"));
 				}
 			}
 		});
@@ -311,7 +335,8 @@ TArray<FName> UPhysicsControlComponent::CreateControlsFromSkeletalMeshAndConstra
 	UPhysicsAsset* PhysicsAsset = SkeletalMeshComponent ? SkeletalMeshComponent->GetPhysicsAsset() : nullptr;
 	if (!PhysicsAsset)
 	{
-		UE_LOG(LogPhysicsControl, Warning, TEXT("No physics asset in skeletal mesh"));
+		UE_LOG(LogPhysicsControl, Warning, TEXT("No physics asset in skeletal mesh in %s"), 
+			GetOwner() ? *GetOwner()->GetName() : TEXT("Unknown"));
 		return Result;
 	}
 
@@ -344,7 +369,8 @@ TArray<FName> UPhysicsControlComponent::CreateControlsFromSkeletalMeshAndConstra
 					ProfileProperties, ChildBoneName, ConstraintProfile))
 				{
 					UE_LOG(LogPhysicsControl, Warning, 
-						TEXT("Failed get constraint profile for %s"), *ChildBoneName.ToString());
+						TEXT("Failed get constraint profile for %s in %s"),
+						*ChildBoneName.ToString(), GetOwner() ? *GetOwner()->GetName() : TEXT("Unknown"));
 					return;
 				}
 
@@ -364,7 +390,8 @@ TArray<FName> UPhysicsControlComponent::CreateControlsFromSkeletalMeshAndConstra
 				else
 				{
 					UE_LOG(LogPhysicsControl, Warning,
-						TEXT("Failed to make control for %s"), *ChildBoneName.ToString());
+						TEXT("Failed to make control for %s in %s"), 
+						*ChildBoneName.ToString(), GetOwner() ? *GetOwner()->GetName() : TEXT("Unknown"));
 				}
 			}
 		});
@@ -385,7 +412,8 @@ TArray<FName> UPhysicsControlComponent::CreateControlsFromSkeletalMesh(
 	UPhysicsAsset* PhysicsAsset = SkeletalMeshComponent ? SkeletalMeshComponent->GetPhysicsAsset() : nullptr;
 	if (!PhysicsAsset)
 	{
-		UE_LOG(LogPhysicsControl, Warning, TEXT("No physics asset in skeletal mesh"));
+		UE_LOG(LogPhysicsControl, Warning, TEXT("No physics asset in skeletal mesh in %s"),
+			GetOwner() ? *GetOwner()->GetName() : TEXT("Unknown"));
 		return Result;
 	}
 
@@ -415,7 +443,8 @@ TArray<FName> UPhysicsControlComponent::CreateControlsFromSkeletalMesh(
 		}
 		else
 		{
-			UE_LOG(LogPhysicsControl, Warning, TEXT("Failed to make control for %s"), *ChildBoneName.ToString());
+			UE_LOG(LogPhysicsControl, Warning, TEXT("Failed to make control for %s in %s"),
+				*ChildBoneName.ToString(), GetOwner() ? *GetOwner()->GetName() : TEXT("Unknown"));
 		}
 	}
 
@@ -435,7 +464,8 @@ TArray<FName> UPhysicsControlComponent::CreateControlsFromSkeletalMeshAndConstra
 	UPhysicsAsset* PhysicsAsset = SkeletalMeshComponent ? SkeletalMeshComponent->GetPhysicsAsset() : nullptr;
 	if (!PhysicsAsset)
 	{
-		UE_LOG(LogPhysicsControl, Warning, TEXT("No physics asset in skeletal mesh"));
+		UE_LOG(LogPhysicsControl, Warning, TEXT("No physics asset in skeletal mesh in %s"), 
+			GetOwner() ? *GetOwner()->GetName() : TEXT("Unknown"));
 		return Result;
 	}
 
@@ -458,7 +488,8 @@ TArray<FName> UPhysicsControlComponent::CreateControlsFromSkeletalMeshAndConstra
 			ProfileProperties, ChildBoneName, ConstraintProfile))
 		{
 			UE_LOG(LogPhysicsControl, Warning, 
-				TEXT("Failed get constraint profile for %s"), *ChildBoneName.ToString());
+				TEXT("Failed get constraint profile for %s in %s"), 
+				*ChildBoneName.ToString(), GetOwner() ? *GetOwner()->GetName() : TEXT("Unknown"));
 			continue;
 		}
 
@@ -476,7 +507,8 @@ TArray<FName> UPhysicsControlComponent::CreateControlsFromSkeletalMeshAndConstra
 		}
 		else
 		{
-			UE_LOG(LogPhysicsControl, Warning, TEXT("Failed to make control for %s"), *ChildBoneName.ToString());
+			UE_LOG(LogPhysicsControl, Warning, TEXT("Failed to make control for %s in %s"),
+				*ChildBoneName.ToString(), GetOwner() ? *GetOwner()->GetName() : TEXT("Unknown"));
 		}
 	}
 
@@ -567,7 +599,8 @@ TMap<FName, FPhysicsControlNames> UPhysicsControlComponent::CreateControlsFromLi
 			else
 			{
 				UE_LOG(LogPhysicsControl, Warning, 
-					TEXT("Failed to make control for %s"), *ChildBoneName.ToString());
+					TEXT("Failed to make control for %s in %s"), 
+					*ChildBoneName.ToString(), GetOwner() ? *GetOwner()->GetName() : TEXT("Unknown"));
 			}
 		}
 	}
@@ -591,13 +624,15 @@ TMap<FName, FPhysicsControlNames> UPhysicsControlComponent::CreateControlsFromLi
 		USkeletalMeshComponent* SkeletalMeshComponent = BonesInLimb.SkeletalMeshComponent.Get();
 		if (!SkeletalMeshComponent)
 		{
-			UE_LOG(LogPhysicsControl, Warning, TEXT("No Skeletal mesh in limb %s"), *LimbName.ToString());
+			UE_LOG(LogPhysicsControl, Warning, TEXT("No Skeletal mesh in limb %s in %s"), 
+				*LimbName.ToString(), GetOwner() ? *GetOwner()->GetName() : TEXT("Unknown"));
 			continue;
 		}
 		UPhysicsAsset* PhysicsAsset = SkeletalMeshComponent ? SkeletalMeshComponent->GetPhysicsAsset() : nullptr;
 		if (!PhysicsAsset)
 		{
-			UE_LOG(LogPhysicsControl, Warning, TEXT("No physics asset in skeletal mesh"));
+			UE_LOG(LogPhysicsControl, Warning, TEXT("No physics asset in skeletal mesh in %s"), 
+				GetOwner() ? *GetOwner()->GetName() : TEXT("Unknown"));
 			return Result;
 		}
 
@@ -635,7 +670,8 @@ TMap<FName, FPhysicsControlNames> UPhysicsControlComponent::CreateControlsFromLi
 			
 			{
 				UE_LOG(LogPhysicsControl, Warning, 
-					TEXT("Failed get constraint profile for %s"), *ChildBoneName.ToString());
+					TEXT("Failed get constraint profile for %s in %s"),
+					*ChildBoneName.ToString(), GetOwner() ? *GetOwner()->GetName() : TEXT("Unknown"));
 				continue;
 			}
 
@@ -655,7 +691,8 @@ TMap<FName, FPhysicsControlNames> UPhysicsControlComponent::CreateControlsFromLi
 			else
 			{
 				UE_LOG(LogPhysicsControl, Warning, 
-					TEXT("Failed to make control for %s"), *ChildBoneName.ToString());
+					TEXT("Failed to make control for %s in %s"),
+					*ChildBoneName.ToString(), GetOwner() ? *GetOwner()->GetName() : TEXT("Unknown"));
 			}
 		}
 	}
@@ -670,84 +707,154 @@ void UPhysicsControlComponent::DestroyAllControlsAndBodyModifiers()
 }
 
 //======================================================================================================================
-void UPhysicsControlComponent::DestroyControl(const FName Name, const bool bApplyToControlsWithName, const bool bApplyToSetsWithName)
+bool UPhysicsControlComponent::DestroyControl(const FName Name, const bool bApplyToControlsWithName, const bool bApplyToSetsWithName)
 {
+	bool bSuccessControl = true;
 	if (bApplyToControlsWithName)
 	{
-		DestroyControl(Name, EDestroyBehavior::RemoveRecord);
+		FDisableNameWarnings DisableNameWarnings(bWarnAboutInvalidNames);
+		bSuccessControl = DestroyControl(Name, EDestroyBehavior::RemoveRecord);
 	}
+
+	bool bSuccessSet = true;
 	if (bApplyToSetsWithName)
 	{
-		DestroyControlsInSet(Name);
+		FDisableNameWarnings DisableNameWarnings(bWarnAboutInvalidNames);
+		bSuccessSet = DestroyControlsInSet(Name);
+	}
+
+	bool bSuccess =
+		(bApplyToControlsWithName && bSuccessControl) ||
+		(bApplyToSetsWithName && bSuccessSet) ||
+		(!bApplyToControlsWithName && !bApplyToSetsWithName);
+	if (!bSuccess && bWarnAboutInvalidNames)
+	{
+		UE_LOG(LogPhysicsControl, Warning,
+			TEXT("DestroyControl - failed to find controls/set %s in %s"),
+			*Name.ToString(), GetOwner() ? *GetOwner()->GetName() : TEXT("Unknown"));
 	}
+	return bSuccess;
 }
 
 //======================================================================================================================
-void UPhysicsControlComponent::DestroyControls(const TArray<FName>& Names, const bool bApplyToControlsWithName, const bool bApplyToSetsWithName)
+bool UPhysicsControlComponent::DestroyControls(const TArray<FName>& Names, const bool bApplyToControlsWithName, const bool bApplyToSetsWithName)
 {
+	bool bSuccess = true;
 	for (FName Name : Names)
 	{
-		DestroyControl(Name, bApplyToControlsWithName, bApplyToSetsWithName);
+		bSuccess &= DestroyControl(Name, bApplyToControlsWithName, bApplyToSetsWithName);
 	}
+	return bSuccess;
 }
 
 //======================================================================================================================
-void UPhysicsControlComponent::DestroyControlsInSet(const FName SetName)
+bool UPhysicsControlComponent::DestroyControlsInSet(const FName SetName)
 {
 	TArray<FName> Names = GetControlNamesInSet(SetName);
-	DestroyControls(Names, true, false);
+	if (Names.IsEmpty())
+	{
+		if (bWarnAboutInvalidNames)
+		{
+			UE_LOG(LogPhysicsControl, Warning,
+				TEXT("DestroyControlsInSet - Set %s not found in %s"),
+				*SetName.ToString(), GetOwner() ? *GetOwner()->GetName() : TEXT("Unknown"));
+
+		}
+		return false;
+	}
+	else
+	{
+		return DestroyControls(Names, true, false);
+	}
 }
 
 //======================================================================================================================
-void UPhysicsControlComponent::SetControlEnabled(
+bool UPhysicsControlComponent::SetControlEnabled(
 	const FName Name, const bool bEnable, const bool bApplyToControlsWithName, const bool bApplyToSetsWithName)
 {
+	bool bSuccessControl = true;
 	if (bApplyToControlsWithName)
 	{
+		FDisableNameWarnings DisableNameWarnings(bWarnAboutInvalidNames);
+
 		FPhysicsControlRecord* Record = FindControlRecord(Name);
 		if (Record)
 		{
 			Record->PhysicsControl.ControlData.bEnabled = bEnable;
 		}
-		else if (bWarnAboutInvalidNames)
+		else
 		{
-			UE_LOG(LogPhysicsControl, Warning,
-				TEXT("SetControlEnabled - invalid name %s"), *Name.ToString());
+			bSuccessControl = false;
 		}
 	}
+	bool bSuccessSet = true;
 	if (bApplyToSetsWithName)
 	{
-		SetControlsInSetEnabled(Name, bEnable);
+		FDisableNameWarnings DisableNameWarnings(bWarnAboutInvalidNames);
+
+		bSuccessSet = SetControlsInSetEnabled(Name, bEnable);
+	}
+
+	bool bSuccess =
+		(bApplyToControlsWithName && bSuccessControl) ||
+		(bApplyToSetsWithName && bSuccessSet) ||
+		(!bApplyToControlsWithName && !bApplyToSetsWithName);
+	if (!bSuccess && bWarnAboutInvalidNames)
+	{
+		UE_LOG(LogPhysicsControl, Warning,
+			TEXT("SetControlEnabled - failed to find controls/set %s in %s"),
+			*Name.ToString(), GetOwner() ? *GetOwner()->GetName() : TEXT("Unknown"));
 	}
+
+	return bSuccess;
 }
 
 //======================================================================================================================
-void UPhysicsControlComponent::SetControlsEnabled(
+bool UPhysicsControlComponent::SetControlsEnabled(
 	const TArray<FName>& Names, const bool bEnable, const bool bApplyToControlsWithName, const bool bApplyToSetsWithName)
 {
+	bool bSuccess = true;
 	for (FName Name : Names)
 	{
-		SetControlEnabled(Name, bEnable, bApplyToControlsWithName, bApplyToSetsWithName);
+		bSuccess &= SetControlEnabled(Name, bEnable, bApplyToControlsWithName, bApplyToSetsWithName);
 	}
+	return bSuccess;
 }
 
 //======================================================================================================================
-void UPhysicsControlComponent::SetControlsInSetEnabled(FName SetName, bool bEnable)
+bool UPhysicsControlComponent::SetControlsInSetEnabled(FName SetName, bool bEnable)
 {
-	SetControlsEnabled(GetControlNamesInSet(SetName), bEnable, true, false);
+	const TArray<FName>& Names = GetControlNamesInSet(SetName);
+	if (Names.IsEmpty())
+	{
+		if (bWarnAboutInvalidNames)
+		{
+			UE_LOG(LogPhysicsControl, Warning,
+				TEXT("SetControlsInSetEnabled - Set %s not found in %s"),
+				*SetName.ToString(), GetOwner() ? *GetOwner()->GetName() : TEXT("Unknown"));
+
+		}
+		return false;
+	}
+	else
+	{
+		return SetControlsEnabled(GetControlNamesInSet(SetName), bEnable, true, false);
+	}
 }
 
 //======================================================================================================================
-void UPhysicsControlComponent::SetControlParent(
+bool UPhysicsControlComponent::SetControlParent(
 	const FName          Name,
 	UPrimitiveComponent* ParentComponent,
 	const FName          ParentBoneName,
 	const bool           bApplyToControlsWithName, 
 	const bool           bApplyToSetsWithName)
 {
-
+	bool bSuccessControl = true;
 	if (bApplyToControlsWithName)
 	{
+		FDisableNameWarnings DisableNameWarnings(bWarnAboutInvalidNames);
+
 		FPhysicsControlRecord* Record = FindControlRecord(Name);
 		if (Record)
 		{
@@ -766,133 +873,253 @@ void UPhysicsControlComponent::SetControlParent(
 			Record->PhysicsControl.ParentBoneName = ParentBoneName;
 			Record->InitConstraint(this, Name, bWarnAboutInvalidNames);
 		}
+		else
+		{
+			bSuccessControl = false;
+		}
 	}
+	bool bSuccessSet = true;
 	if (bApplyToSetsWithName)
 	{
-		SetControlParentsInSet(Name, ParentComponent, ParentBoneName);
+		FDisableNameWarnings DisableNameWarnings(bWarnAboutInvalidNames);
+
+		bSuccessSet = SetControlParentsInSet(Name, ParentComponent, ParentBoneName);
+	}
+
+	bool bSuccess =
+		(bApplyToControlsWithName && bSuccessControl) ||
+		(bApplyToSetsWithName && bSuccessSet) ||
+		(!bApplyToControlsWithName && !bApplyToSetsWithName);
+	if (!bSuccess && bWarnAboutInvalidNames)
+	{
+		UE_LOG(LogPhysicsControl, Warning,
+			TEXT("SetControlParent - failed to find controls/set %s in %s"),
+			*Name.ToString(), GetOwner() ? *GetOwner()->GetName() : TEXT("Unknown"));
 	}
+
+	return bSuccess;
 }
 
 //======================================================================================================================
-void UPhysicsControlComponent::SetControlParents(
+bool UPhysicsControlComponent::SetControlParents(
 	const TArray<FName>& Names,
 	UPrimitiveComponent* ParentComponent,
 	const FName          ParentBoneName,
 	const bool           bApplyToControlsWithName,
 	const bool           bApplyToSetsWithName)
 {
+	bool bSuccess = true;
 	for (FName Name : Names)
 	{
-		SetControlParent(Name, ParentComponent, ParentBoneName, bApplyToControlsWithName, bApplyToSetsWithName);
+		bSuccess &= SetControlParent(Name, ParentComponent, ParentBoneName, bApplyToControlsWithName, bApplyToSetsWithName);
 	}
+	return bSuccess;
 }
 
 //======================================================================================================================
-void UPhysicsControlComponent::SetControlParentsInSet(
+bool UPhysicsControlComponent::SetControlParentsInSet(
 	const FName          SetName,
 	UPrimitiveComponent* ParentComponent,
 	const FName          ParentBoneName)
 {
-	SetControlParents(GetControlNamesInSet(SetName), ParentComponent, ParentBoneName, true, false);
+	const TArray<FName>& Names = GetControlNamesInSet(SetName);
+	if (Names.IsEmpty())
+	{
+		if (bWarnAboutInvalidNames)
+		{
+			UE_LOG(LogPhysicsControl, Warning,
+				TEXT("SetControlParentsInSet - Set %s not found in %s"),
+				*SetName.ToString(), GetOwner() ? *GetOwner()->GetName() : TEXT("Unknown"));
+
+		}
+		return false;
+	}
+	else
+	{
+		return SetControlParents(Names, ParentComponent, ParentBoneName, true, false);
+	}
 }
 
 //======================================================================================================================
 // Note - some params passed by value to allow them to be unconnected and enable inline BP editing
-void UPhysicsControlComponent::SetControlData(
+bool UPhysicsControlComponent::SetControlData(
 	const FName               Name, 
 	const FPhysicsControlData ControlData,
 	const bool                bApplyToControlsWithName,
 	const bool                bApplyToSetsWithName)
 {
+	bool bSuccessControl = true;
 	if (bApplyToControlsWithName)
 	{
+		FDisableNameWarnings DisableNameWarnings(bWarnAboutInvalidNames);
+
 		FPhysicsControlRecord* Record = FindControlRecord(Name);
 		if (Record)
 		{
 			Record->PhysicsControl.ControlData = ControlData;
 		}
+		else
+		{
+			bSuccessControl = false;
+		}
 	}
+
+	bool bSuccessSet = true;
 	if (bApplyToSetsWithName)
 	{
-		SetControlDatasInSet(Name, ControlData);
+		FDisableNameWarnings DisableNameWarnings(bWarnAboutInvalidNames);
+
+		bSuccessSet = SetControlDatasInSet(Name, ControlData);
+	}
+
+	bool bSuccess =
+		(bApplyToControlsWithName && bSuccessControl) ||
+		(bApplyToSetsWithName && bSuccessSet) ||
+		(!bApplyToControlsWithName && !bApplyToSetsWithName);
+	if (!bSuccess && bWarnAboutInvalidNames)
+	{
+		UE_LOG(LogPhysicsControl, Warning,
+			TEXT("SetControlData - failed to find controls/set %s in %s"),
+			*Name.ToString(), GetOwner() ? *GetOwner()->GetName() : TEXT("Unknown"));
 	}
+
+	return bSuccess;
 }
 
 //======================================================================================================================
-void UPhysicsControlComponent::SetControlDatas(
+bool UPhysicsControlComponent::SetControlDatas(
 	const TArray<FName>&      Names, 
 	const FPhysicsControlData ControlData,
 	const bool                bApplyToControlsWithName,
 	const bool                bApplyToSetsWithName)
 {
+	bool bSuccess = true;
 	for (FName Name : Names)
 	{
-		SetControlData(Name, ControlData, bApplyToControlsWithName, bApplyToSetsWithName);
+		bSuccess &= SetControlData(Name, ControlData, bApplyToControlsWithName, bApplyToSetsWithName);
 	}
+	return bSuccess;
 }
 
 //======================================================================================================================
-void UPhysicsControlComponent::SetControlDatasInSet(
+bool UPhysicsControlComponent::SetControlDatasInSet(
 	const FName               SetName,
 	const FPhysicsControlData ControlData)
 {
-	SetControlDatas(GetControlNamesInSet(SetName), ControlData, true, false);
+	const TArray<FName>& Names = GetControlNamesInSet(SetName);
+	if (Names.IsEmpty())
+	{
+		if (bWarnAboutInvalidNames)
+		{
+			UE_LOG(LogPhysicsControl, Warning,
+				TEXT("SetControlDatasInSet - Set %s not found in %s"),
+				*SetName.ToString(), GetOwner() ? *GetOwner()->GetName() : TEXT("Unknown"));
+
+		}
+		return false;
+	}
+	else
+	{
+		return SetControlDatas(Names, ControlData, true, false);
+	}
 }
 
 //======================================================================================================================
 // Note - some params passed by value to allow them to be unconnected and enable inline BP editing
-void UPhysicsControlComponent::SetControlSparseData(
+bool UPhysicsControlComponent::SetControlSparseData(
 	const FName                     Name, 
 	const FPhysicsControlSparseData ControlData,
 	const bool                      bApplyToControlsWithName,
 	const bool                      bApplyToSetsWithName)
 {
+	bool bSuccessControl = true;
 	if (bApplyToControlsWithName)
 	{
+		FDisableNameWarnings DisableNameWarnings(bWarnAboutInvalidNames);
+
 		FPhysicsControlRecord* Record = FindControlRecord(Name);
 		if (Record)
 		{
 			Record->PhysicsControl.ControlData.UpdateFromSparseData(ControlData);
 		}
+		else
+		{
+			bSuccessControl = false;
+		}
 	}
+	
+	bool bSuccessSet = true;
 	if (bApplyToSetsWithName)
 	{
-		SetControlSparseDatasInSet(Name, ControlData);
+		FDisableNameWarnings DisableNameWarnings(bWarnAboutInvalidNames);
+
+		bSuccessSet = SetControlSparseDatasInSet(Name, ControlData);
+	}
+
+	bool bSuccess =
+		(bApplyToControlsWithName && bSuccessControl) ||
+		(bApplyToSetsWithName && bSuccessSet) ||
+		(!bApplyToControlsWithName && !bApplyToSetsWithName);
+	if (!bSuccess && bWarnAboutInvalidNames)
+	{
+		UE_LOG(LogPhysicsControl, Warning,
+			TEXT("SetControlSparseData - failed to find controls/set %s in %s"),
+			*Name.ToString(), GetOwner() ? *GetOwner()->GetName() : TEXT("Unknown"));
 	}
+	return bSuccess;
 }
 
 //======================================================================================================================
-void UPhysicsControlComponent::SetControlSparseDatas(
+bool UPhysicsControlComponent::SetControlSparseDatas(
 	const TArray<FName>&            Names, 
 	const FPhysicsControlSparseData ControlData,
 	const bool                      bApplyToControlsWithName,
 	const bool                      bApplyToSetsWithName)
 {
+	bool bSuccess = true;
 	for (FName Name : Names)
 	{
-		SetControlSparseData(Name, ControlData, bApplyToControlsWithName, bApplyToSetsWithName);
+		bSuccess &= SetControlSparseData(Name, ControlData, bApplyToControlsWithName, bApplyToSetsWithName);
 	}
+	return bSuccess;
 }
 
 //======================================================================================================================
-void UPhysicsControlComponent::SetControlSparseDatasInSet(
+bool UPhysicsControlComponent::SetControlSparseDatasInSet(
 	const FName                     SetName,
 	const FPhysicsControlSparseData ControlData)
 {
-	SetControlSparseDatas(GetControlNamesInSet(SetName), ControlData, true, false);
+	const TArray<FName>& Names = GetControlNamesInSet(SetName);
+	if (Names.IsEmpty())
+	{
+		if (bWarnAboutInvalidNames)
+		{
+			UE_LOG(LogPhysicsControl, Warning,
+				TEXT("SetControlSparseDatasInSet - Set %s not found in %s"),
+				*SetName.ToString(), GetOwner() ? *GetOwner()->GetName() : TEXT("Unknown"));
+
+		}
+		return false;
+	}
+	else
+	{
+		return SetControlSparseDatas(Names, ControlData, true, false);
+	}
 }
 
 //======================================================================================================================
 // Note - some params passed by value to allow them to be unconnected and enable inline BP editing
-void UPhysicsControlComponent::SetControlMultiplier(
+bool UPhysicsControlComponent::SetControlMultiplier(
 	const FName                      Name, 
 	const FPhysicsControlMultiplier  ControlMultiplier, 
 	const bool                       bEnableControl,
 	const bool                       bApplyToControlsWithName,
 	const bool                       bApplyToSetsWithName)
 {
+	bool bSuccessControl = true;
 	if (bApplyToControlsWithName)
 	{
+		FDisableNameWarnings DisableNameWarnings(bWarnAboutInvalidNames);
 		FPhysicsControlRecord* Record = FindControlRecord(Name);
 		if (Record)
 		{
@@ -902,47 +1129,85 @@ void UPhysicsControlComponent::SetControlMultiplier(
 				Record->PhysicsControl.ControlData.bEnabled = true;
 			}
 		}
+		else
+		{
+			bSuccessControl = false;
+		}
 	}
+
+	bool bSuccessSet = true;
 	if (bApplyToSetsWithName)
 	{
-		SetControlMultipliersInSet(Name, ControlMultiplier, bEnableControl);
+		FDisableNameWarnings DisableNameWarnings(bWarnAboutInvalidNames);
+		bSuccessSet = SetControlMultipliersInSet(Name, ControlMultiplier, bEnableControl);
 	}
+
+	bool bSuccess =
+		(bApplyToControlsWithName && bSuccessControl) ||
+		(bApplyToSetsWithName && bSuccessSet) ||
+		(!bApplyToControlsWithName && !bApplyToSetsWithName);
+	if (!bSuccess && bWarnAboutInvalidNames)
+	{
+		UE_LOG(LogPhysicsControl, Warning,
+			TEXT("SetControlMultiplier - failed to find controls/set %s in %s"),
+			*Name.ToString(), GetOwner() ? *GetOwner()->GetName() : TEXT("Unknown"));
+	}
+	return bSuccess;
 }
 
 //======================================================================================================================
-void UPhysicsControlComponent::SetControlMultipliers(
+bool UPhysicsControlComponent::SetControlMultipliers(
 	const TArray<FName>&             Names,
 	const FPhysicsControlMultiplier  ControlMultiplier, 
 	const bool                       bEnableControl,
 	const bool                       bApplyToControlsWithName,
 	const bool                       bApplyToSetsWithName)
 {
+	bool bSuccess = true;
 	for (FName Name : Names)
 	{
-		SetControlMultiplier(Name, ControlMultiplier, bEnableControl, bApplyToControlsWithName, bApplyToSetsWithName);
+		bSuccess &= SetControlMultiplier(Name, ControlMultiplier, bEnableControl, bApplyToControlsWithName, bApplyToSetsWithName);
 	}
+	return bSuccess;
 }
 
 //======================================================================================================================
-void UPhysicsControlComponent::SetControlMultipliersInSet(
+bool UPhysicsControlComponent::SetControlMultipliersInSet(
 	const FName                      SetName,
 	const FPhysicsControlMultiplier  ControlMultiplier, 
 	const bool                       bEnableControl)
 {
-	SetControlMultipliers(GetControlNamesInSet(SetName), ControlMultiplier, bEnableControl, true, false);
+	TArray<FName> Names = GetControlNamesInSet(SetName);
+	if (Names.IsEmpty())
+	{
+		if (bWarnAboutInvalidNames)
+		{
+			UE_LOG(LogPhysicsControl, Warning,
+				TEXT("SetControlMultipliersInSet - Set %s not found in %s"),
+				*SetName.ToString(), GetOwner() ? *GetOwner()->GetName() : TEXT("Unknown"));
+
+		}
+		return false;
+	}
+	else
+	{
+		return SetControlMultipliers(Names, ControlMultiplier, bEnableControl, true, false);
+	}
 }
 
 //======================================================================================================================
 // Note - some params passed by value to allow them to be unconnected and enable inline BP editing
-void UPhysicsControlComponent::SetControlSparseMultiplier(
+bool UPhysicsControlComponent::SetControlSparseMultiplier(
 	const FName                            Name,
 	const FPhysicsControlSparseMultiplier  ControlMultiplier,
 	const bool                             bEnableControl,
 	const bool                             bApplyToControlsWithName,
 	const bool                             bApplyToSetsWithName)
 {
+	bool bSuccessControl = true;
 	if (bApplyToControlsWithName)
 	{
+		FDisableNameWarnings DisableNameWarnings(bWarnAboutInvalidNames);
 		FPhysicsControlRecord* Record = FindControlRecord(Name);
 		if (Record)
 		{
@@ -952,38 +1217,74 @@ void UPhysicsControlComponent::SetControlSparseMultiplier(
 				Record->PhysicsControl.ControlData.bEnabled = true;
 			}
 		}
+		else
+		{
+			bSuccessControl = false;
+		}
 	}
+
+	bool bSuccessSet = true;
 	if (bApplyToSetsWithName)
 	{
-		SetControlSparseMultipliersInSet(Name, ControlMultiplier, bEnableControl);
+		FDisableNameWarnings DisableNameWarnings(bWarnAboutInvalidNames);
+		bSuccessSet = SetControlSparseMultipliersInSet(Name, ControlMultiplier, bEnableControl);
+	}
+
+	bool bSuccess =
+		(bApplyToControlsWithName && bSuccessControl) ||
+		(bApplyToSetsWithName && bSuccessSet) ||
+		(!bApplyToControlsWithName && !bApplyToSetsWithName);
+	if (!bSuccess && bWarnAboutInvalidNames)
+	{
+		UE_LOG(LogPhysicsControl, Warning,
+			TEXT("SetControlSparseMultiplier - failed to find controls/set %s in %s"),
+			*Name.ToString(), GetOwner() ? *GetOwner()->GetName() : TEXT("Unknown"));
 	}
+	return bSuccess;
 }
 
 //======================================================================================================================
-void UPhysicsControlComponent::SetControlSparseMultipliers(
+bool UPhysicsControlComponent::SetControlSparseMultipliers(
 	const TArray<FName>&                  Names,
 	const FPhysicsControlSparseMultiplier ControlMultiplier,
 	const bool                            bEnableControl,
 	const bool                            bApplyToControlsWithName,
 	const bool                            bApplyToSetsWithName)
 {
+	bool bSuccess = true;
 	for (FName Name : Names)
 	{
-		SetControlSparseMultiplier(Name, ControlMultiplier, bEnableControl, bApplyToControlsWithName, bApplyToSetsWithName);
+		bSuccess &= SetControlSparseMultiplier(Name, ControlMultiplier, bEnableControl, bApplyToControlsWithName, bApplyToSetsWithName);
 	}
+	return bSuccess;
 }
 
 //======================================================================================================================
-void UPhysicsControlComponent::SetControlSparseMultipliersInSet(
+bool UPhysicsControlComponent::SetControlSparseMultipliersInSet(
 	const FName                           SetName,
 	const FPhysicsControlSparseMultiplier ControlMultiplier,
 	const bool                            bEnableControl)
 {
-	SetControlSparseMultipliers(GetControlNamesInSet(SetName), ControlMultiplier, bEnableControl, true, false);
+	TArray<FName> Names = GetControlNamesInSet(SetName);
+	if (Names.IsEmpty())
+	{
+		if (bWarnAboutInvalidNames)
+		{
+			UE_LOG(LogPhysicsControl, Warning,
+				TEXT("SetControlSparseMultipliersInSet - Set %s not found in %s"),
+				*SetName.ToString(), GetOwner() ? *GetOwner()->GetName() : TEXT("Unknown"));
+
+		}
+		return false;
+	}
+	else
+	{
+		return SetControlSparseMultipliers(Names, ControlMultiplier, bEnableControl, true, false);
+	}
 }
 
 //======================================================================================================================
-void UPhysicsControlComponent::SetControlLinearData(
+bool UPhysicsControlComponent::SetControlLinearData(
 	const FName Name, 
 	const float Strength, 
 	const float DampingRatio, 
@@ -993,8 +1294,10 @@ void UPhysicsControlComponent::SetControlLinearData(
 	const bool  bApplyToControlsWithName,
 	const bool  bApplyToSetsWithName)
 {
+	bool bSuccessControl = true;
 	if (bApplyToControlsWithName)
 	{
+		FDisableNameWarnings DisableNameWarnings(bWarnAboutInvalidNames);
 		FPhysicsControlRecord* Record = FindControlRecord(Name);
 		if (Record)
 		{
@@ -1007,19 +1310,39 @@ void UPhysicsControlComponent::SetControlLinearData(
 				Record->PhysicsControl.ControlData.bEnabled = true;
 			}
 		}
+		else
+		{
+			bSuccessControl = false;
+		}
 	}
+
+	bool bSuccessSet = true;
 	if (bApplyToSetsWithName)
 	{
+		FDisableNameWarnings DisableNameWarnings(bWarnAboutInvalidNames);
 		TArray<FName> Names = GetControlNamesInSet(Name);
 		for (FName ControlName : Names)
 		{
-			SetControlLinearData(ControlName, Strength, DampingRatio, ExtraDamping, MaxForce, bEnableControl, true, false);
+			bSuccessSet &= SetControlLinearData(ControlName, Strength, DampingRatio, ExtraDamping, MaxForce, bEnableControl, true, false);
 		}
 	}
+
+	bool bSuccess =
+		(bApplyToControlsWithName && bSuccessControl) ||
+		(bApplyToSetsWithName && bSuccessSet) ||
+		(!bApplyToControlsWithName && !bApplyToSetsWithName);
+	if (!bSuccess && bWarnAboutInvalidNames)
+	{
+		UE_LOG(LogPhysicsControl, Warning,
+			TEXT("SetControlLinearData - failed to find controls/set %s in %s"),
+			*Name.ToString(), GetOwner() ? *GetOwner()->GetName() : TEXT("Unknown"));
+	}
+	return bSuccess;
+
 }
 
 //======================================================================================================================
-void UPhysicsControlComponent::SetControlAngularData(
+bool UPhysicsControlComponent::SetControlAngularData(
 	const FName Name, 
 	const float Strength, 
 	const float DampingRatio, 
@@ -1029,8 +1352,10 @@ void UPhysicsControlComponent::SetControlAngularData(
 	const bool  bApplyToControlsWithName,
 	const bool  bApplyToSetsWithName)
 {
+	bool bSuccessControl = true;
 	if (bApplyToControlsWithName)
 	{
+		FDisableNameWarnings DisableNameWarnings(bWarnAboutInvalidNames);
 		FPhysicsControlRecord* Record = FindControlRecord(Name);
 		if (Record)
 		{
@@ -1043,15 +1368,34 @@ void UPhysicsControlComponent::SetControlAngularData(
 				Record->PhysicsControl.ControlData.bEnabled = true;
 			}
 		}
+		else
+		{
+			bSuccessControl = false;
+		}
 	}
+
+	bool bSuccessSet = true;
 	if (bApplyToSetsWithName)
 	{
+		FDisableNameWarnings DisableNameWarnings(bWarnAboutInvalidNames);
 		TArray<FName> Names = GetControlNamesInSet(Name);
 		for (FName ControlName : Names)
 		{
-			SetControlAngularData(ControlName, Strength, DampingRatio, ExtraDamping, MaxTorque, bEnableControl, true, false);
+			bSuccessSet &= SetControlAngularData(ControlName, Strength, DampingRatio, ExtraDamping, MaxTorque, bEnableControl, true, false);
 		}
 	}
+
+	bool bSuccess =
+		(bApplyToControlsWithName && bSuccessControl) ||
+		(bApplyToSetsWithName && bSuccessSet) ||
+		(!bApplyToControlsWithName && !bApplyToSetsWithName);
+	if (!bSuccess && bWarnAboutInvalidNames)
+	{
+		UE_LOG(LogPhysicsControl, Warning,
+			TEXT("SetControlAngularData - failed to find controls/set %s in %s"),
+			*Name.ToString(), GetOwner() ? *GetOwner()->GetName() : TEXT("Unknown"));
+	}
+	return bSuccess;
 }
 
 //======================================================================================================================
@@ -1067,8 +1411,8 @@ bool UPhysicsControlComponent::SetControlPoint(const FName Name, const FVector P
 	}
 	if (bWarnAboutInvalidNames)
 	{
-		UE_LOG(LogPhysicsControl, Warning,
-			TEXT("SetControlPoint - invalid name %s"), *Name.ToString());
+		UE_LOG(LogPhysicsControl, Warning, TEXT("SetControlPoint - invalid control name %s for %s"),
+			*Name.ToString(), GetOwner() ? *GetOwner()->GetName() : TEXT("Unknown"));
 	}
 	return false;
 }
@@ -1084,22 +1428,24 @@ bool UPhysicsControlComponent::ResetControlPoint(const FName Name)
 	}
 	if (bWarnAboutInvalidNames)
 	{
-		UE_LOG(LogPhysicsControl, Warning,
-			TEXT("ResetControlPoint - invalid name %s"), *Name.ToString());
+		UE_LOG(LogPhysicsControl, Warning, TEXT("ResetControlPoint - invalid control name %s for %s"),
+			*Name.ToString(), GetOwner() ? *GetOwner()->GetName() : TEXT("Unknown"));
 	}
 	return false;
 }
 
 //======================================================================================================================
-void UPhysicsControlComponent::SetControlTarget(
+bool UPhysicsControlComponent::SetControlTarget(
 	const FName                 Name, 
 	const FPhysicsControlTarget ControlTarget, 
 	const bool                  bEnableControl,
 	const bool                  bApplyToControlsWithName,
 	const bool                  bApplyToSetsWithName)
 {
+	bool bSuccessControl = true;
 	if (bApplyToControlsWithName)
 	{
+		FDisableNameWarnings DisableNameWarnings(bWarnAboutInvalidNames);
 		FPhysicsControlRecord* Record = FindControlRecord(Name);
 		if (Record)
 		{
@@ -1109,38 +1455,74 @@ void UPhysicsControlComponent::SetControlTarget(
 				Record->PhysicsControl.ControlData.bEnabled = true;
 			}
 		}
+		else
+		{
+			bSuccessControl = false;
+		}
 	}
+
+	bool bSuccessSet = true;
 	if (bApplyToSetsWithName)
 	{
-		SetControlTargetsInSet(Name, ControlTarget, bEnableControl);
+		FDisableNameWarnings DisableNameWarnings(bWarnAboutInvalidNames);
+		bSuccessSet = SetControlTargetsInSet(Name, ControlTarget, bEnableControl);
+	}
+
+	bool bSuccess =
+		(bApplyToControlsWithName && bSuccessControl) ||
+		(bApplyToSetsWithName && bSuccessSet) ||
+		(!bApplyToControlsWithName && !bApplyToSetsWithName);
+	if (!bSuccess && bWarnAboutInvalidNames)
+	{
+		UE_LOG(LogPhysicsControl, Warning,
+			TEXT("SetControlTarget - failed to find controls/set %s in %s"),
+			*Name.ToString(), GetOwner() ? *GetOwner()->GetName() : TEXT("Unknown"));
 	}
+	return bSuccess;
 }
 
 //======================================================================================================================
-void UPhysicsControlComponent::SetControlTargets(
+bool UPhysicsControlComponent::SetControlTargets(
 	const TArray<FName>&        Names, 
 	const FPhysicsControlTarget ControlTarget, 
 	const bool                  bEnableControl,
 	const bool                  bApplyToControlsWithName,
 	const bool                  bApplyToSetsWithName)
 {
+	bool bSuccess = true;
 	for (FName Name : Names)
 	{
-		SetControlTarget(Name, ControlTarget, bEnableControl, bApplyToControlsWithName, bApplyToSetsWithName);
+		bSuccess &= SetControlTarget(Name, ControlTarget, bEnableControl, bApplyToControlsWithName, bApplyToSetsWithName);
 	}
+	return bSuccess;
 }
 
 //======================================================================================================================
-void UPhysicsControlComponent::SetControlTargetsInSet(
+bool UPhysicsControlComponent::SetControlTargetsInSet(
 	const FName                 SetName, 
 	const FPhysicsControlTarget ControlTarget, 
 	const bool                  bEnableControl)
 {
-	SetControlTargets(GetControlNamesInSet(SetName), ControlTarget, bEnableControl, true, false);
+	TArray<FName> Names = GetControlNamesInSet(SetName);
+	if (Names.IsEmpty())
+	{
+		if (bWarnAboutInvalidNames)
+		{
+			UE_LOG(LogPhysicsControl, Warning,
+				TEXT("SetControlTargetsInSet - Set %s not found in %s"),
+				*SetName.ToString(), GetOwner() ? *GetOwner()->GetName() : TEXT("Unknown"));
+
+		}
+		return false;
+	}
+	else
+	{
+		return SetControlTargets(Names, ControlTarget, bEnableControl, true, false);
+	}
 }
 
 //======================================================================================================================
-void UPhysicsControlComponent::SetControlTargetPositionAndOrientation(
+bool UPhysicsControlComponent::SetControlTargetPositionAndOrientation(
 	const FName    Name, 
 	const FVector  Position, 
 	const FRotator Orientation, 
@@ -1150,26 +1532,41 @@ void UPhysicsControlComponent::SetControlTargetPositionAndOrientation(
 	const bool     bApplyToControlsWithName,
 	const bool     bApplyToSetsWithName)
 {
+	bool bSuccessControl = true;
 	if (bApplyToControlsWithName)
 	{
-		FPhysicsControlRecord* Record = FindControlRecord(Name);
-		if (Record)
-		{
-			SetControlTargetPosition(
-				Name, Position, VelocityDeltaTime, bEnableControl, bApplyControlPointToTarget);
-			SetControlTargetOrientation(
-				Name, Orientation, VelocityDeltaTime, bEnableControl, bApplyControlPointToTarget);
-		}
+		FDisableNameWarnings DisableNameWarnings(bWarnAboutInvalidNames);
+
+		bSuccessControl &= SetControlTargetPosition(
+			Name, Position, VelocityDeltaTime, bEnableControl, bApplyControlPointToTarget, true, false);
+		bSuccessControl &= SetControlTargetOrientation(
+			Name, Orientation, VelocityDeltaTime, bEnableControl, bApplyControlPointToTarget, true, false);
 	}
+
+	bool bSuccessSet = true;
 	if (bApplyToSetsWithName)
 	{
-		SetControlTargetPositionsAndOrientationsInSet(
+		FDisableNameWarnings DisableNameWarnings(bWarnAboutInvalidNames);
+
+		bSuccessSet = SetControlTargetPositionsAndOrientationsInSet(
 			Name, Position, Orientation, VelocityDeltaTime, bEnableControl, bApplyControlPointToTarget);
 	}
+
+	bool bSuccess =
+		(bApplyToControlsWithName && bSuccessControl) ||
+		(bApplyToSetsWithName && bSuccessSet) ||
+		(!bApplyToControlsWithName && !bApplyToSetsWithName);
+	if (!bSuccess && bWarnAboutInvalidNames)
+	{
+		UE_LOG(LogPhysicsControl, Warning,
+			TEXT("SetControlTargetPositionAndOrientation - failed to find controls/set %s in %s"),
+			*Name.ToString(), GetOwner() ? *GetOwner()->GetName() : TEXT("Unknown"));
+	}
+	return bSuccess;
 }
 
 //======================================================================================================================
-void UPhysicsControlComponent::SetControlTargetPositionsAndOrientations(
+bool UPhysicsControlComponent::SetControlTargetPositionsAndOrientations(
 	const TArray<FName>& Names,
 	const FVector        Position,
 	const FRotator       Orientation,
@@ -1179,16 +1576,18 @@ void UPhysicsControlComponent::SetControlTargetPositionsAndOrientations(
 	const bool           bApplyToControlsWithName,
 	const bool           bApplyToSetsWithName)
 {
+	bool bSuccess = true;
 	for (FName Name : Names)
 	{
-		SetControlTargetPositionAndOrientation(
+		bSuccess &= SetControlTargetPositionAndOrientation(
 			Name, Position, Orientation, VelocityDeltaTime, bEnableControl, bApplyControlPointToTarget,
 			bApplyToControlsWithName, bApplyToSetsWithName);
 	}
+	return bSuccess;
 }
 
 //======================================================================================================================
-void UPhysicsControlComponent::SetControlTargetPositionsAndOrientationsInSet(
+bool UPhysicsControlComponent::SetControlTargetPositionsAndOrientationsInSet(
 	const FName          SetName,
 	const FVector        Position,
 	const FRotator       Orientation,
@@ -1196,12 +1595,27 @@ void UPhysicsControlComponent::SetControlTargetPositionsAndOrientationsInSet(
 	const bool           bEnableControl,
 	const bool           bApplyControlPointToTarget)
 {
-	SetControlTargetPositionsAndOrientations(GetControlNamesInSet(SetName), 
-		Position, Orientation, VelocityDeltaTime, bEnableControl, bApplyControlPointToTarget, true, false);
+	const TArray<FName>& Names = GetControlNamesInSet(SetName);
+	if (Names.IsEmpty())
+	{
+		if (bWarnAboutInvalidNames)
+		{
+			UE_LOG(LogPhysicsControl, Warning,
+				TEXT("SetControlTargetPositionsAndOrientationsInSet - Set %s not found in %s"),
+				*SetName.ToString(), GetOwner() ? *GetOwner()->GetName() : TEXT("Unknown"));
+
+		}
+		return false;
+	}
+	else
+	{
+		return SetControlTargetPositionsAndOrientations(
+			Names, Position, Orientation, VelocityDeltaTime, bEnableControl, bApplyControlPointToTarget, true, false);
+	}
 }
 
 //======================================================================================================================
-void UPhysicsControlComponent::SetControlTargetPosition(
+bool UPhysicsControlComponent::SetControlTargetPosition(
 	const FName   Name, 
 	const FVector Position, 
 	const float   VelocityDeltaTime, 
@@ -1210,8 +1624,11 @@ void UPhysicsControlComponent::SetControlTargetPosition(
 	const bool    bApplyToControlsWithName, 
 	const bool    bApplyToSetsWithName)
 {
+	bool bSuccessControl = true;
 	if (bApplyToControlsWithName)
 	{
+		FDisableNameWarnings DisableNameWarnings(bWarnAboutInvalidNames);
+
 		FPhysicsControlRecord* Record = FindControlRecord(Name);
 		if (Record)
 		{
@@ -1231,15 +1648,37 @@ void UPhysicsControlComponent::SetControlTargetPosition(
 				Record->PhysicsControl.ControlData.bEnabled = true;
 			}
 		}
+		else
+		{
+			bSuccessControl = false;
+		}
 	}
+
+	bool bSuccessSet = true;
 	if (bApplyToSetsWithName)
 	{
-		SetControlTargetPositionsInSet(Name, Position, VelocityDeltaTime, bEnableControl, bApplyControlPointToTarget);
+		FDisableNameWarnings DisableNameWarnings(bWarnAboutInvalidNames);
+
+		bSuccessSet = SetControlTargetPositionsInSet(
+			Name, Position, VelocityDeltaTime, bEnableControl, bApplyControlPointToTarget);
+	}
+
+	bool bSuccess = 
+		(bApplyToControlsWithName && bSuccessControl) || 
+		(bApplyToSetsWithName && bSuccessSet) ||
+		(!bApplyToControlsWithName && !bApplyToSetsWithName);
+	if (!bSuccess && bWarnAboutInvalidNames)
+	{
+		UE_LOG(LogPhysicsControl, Warning,
+			TEXT("SetControlTargetPosition - failed to find controls/set %s in %s"),
+			*Name.ToString(), GetOwner() ? *GetOwner()->GetName() : TEXT("Unknown"));
 	}
+
+	return bSuccess;
 }
 
 //======================================================================================================================
-void UPhysicsControlComponent::SetControlTargetPositions(
+bool UPhysicsControlComponent::SetControlTargetPositions(
 	const TArray<FName>& Names,
 	const FVector        Position,
 	const float          VelocityDeltaTime,
@@ -1248,28 +1687,45 @@ void UPhysicsControlComponent::SetControlTargetPositions(
 	const bool           bApplyToControlsWithName,
 	const bool           bApplyToSetsWithName)
 {
+	bool bSuccess = true;
 	for (FName Name : Names)
 	{
-		SetControlTargetPosition(
+		bSuccess &= SetControlTargetPosition(
 			Name, Position, VelocityDeltaTime, bEnableControl, bApplyControlPointToTarget,
 			bApplyToControlsWithName, bApplyToSetsWithName);
 	}
+	return bSuccess;
 }
 
 //======================================================================================================================
-void UPhysicsControlComponent::SetControlTargetPositionsInSet(
+bool UPhysicsControlComponent::SetControlTargetPositionsInSet(
 	const FName   SetName,
 	const FVector Position,
 	const float   VelocityDeltaTime,
 	const bool    bEnableControl,
 	const bool    bApplyControlPointToTarget)
 {
-	SetControlTargetPositions(GetControlNamesInSet(SetName),
-		Position, VelocityDeltaTime, bEnableControl, bApplyControlPointToTarget, true, false);
+	const TArray<FName>& Names = GetControlNamesInSet(SetName);
+	if (Names.IsEmpty())
+	{
+		if (bWarnAboutInvalidNames)
+		{
+			UE_LOG(LogPhysicsControl, Warning,
+				TEXT("SetControlTargetPositionsInSet - Set %s not found in %s"),
+				*SetName.ToString(), GetOwner() ? *GetOwner()->GetName() : TEXT("Unknown"));
+
+		}
+		return false;
+	}
+	else
+	{
+		return SetControlTargetPositions(
+			Names, Position, VelocityDeltaTime, bEnableControl, bApplyControlPointToTarget, true, false);
+	}
 }
 
 //======================================================================================================================
-void UPhysicsControlComponent::SetControlTargetOrientation(
+bool UPhysicsControlComponent::SetControlTargetOrientation(
 	const FName    Name, 
 	const FRotator Orientation, 
 	const float    AngularVelocityDeltaTime, 
@@ -1278,8 +1734,11 @@ void UPhysicsControlComponent::SetControlTargetOrientation(
 	const bool     bApplyToControlsWithName,
 	const bool     bApplyToSetsWithName)
 {
+	bool bSuccessControl = true;
 	if (bApplyToControlsWithName)
 	{
+		FDisableNameWarnings DisableNameWarnings(bWarnAboutInvalidNames);
+
 		FPhysicsControlRecord* Record = FindControlRecord(Name);
 		if (Record)
 		{
@@ -1303,16 +1762,37 @@ void UPhysicsControlComponent::SetControlTargetOrientation(
 				Record->PhysicsControl.ControlData.bEnabled = true;
 			}
 		}
+		else
+		{
+			bSuccessControl = false;
+		}
 	}
+
+	bool bSuccessSet = true;
 	if (bApplyToSetsWithName)
 	{
-		SetControlTargetOrientationsInSet(
+		FDisableNameWarnings DisableNameWarnings(bWarnAboutInvalidNames);
+
+		bSuccessSet &= SetControlTargetOrientationsInSet(
 			Name, Orientation, AngularVelocityDeltaTime, bEnableControl, bApplyControlPointToTarget);
 	}
+
+	bool bSuccess =
+		(bApplyToControlsWithName && bSuccessControl) ||
+		(bApplyToSetsWithName && bSuccessSet) ||
+		(!bApplyToControlsWithName && !bApplyToSetsWithName);
+	if (!bSuccess && bWarnAboutInvalidNames)
+	{
+		UE_LOG(LogPhysicsControl, Warning,
+			TEXT("SetControlTargetOrientation - failed to find controls/set %s in %s"),
+			*Name.ToString(), GetOwner() ? *GetOwner()->GetName() : TEXT("Unknown"));
+	}
+
+	return bSuccess;
 }
 
 //======================================================================================================================
-void UPhysicsControlComponent::SetControlTargetOrientations(
+bool UPhysicsControlComponent::SetControlTargetOrientations(
 	const TArray<FName>& Names,
 	const FRotator Orientation,
 	const float    AngularVelocityDeltaTime,
@@ -1321,24 +1801,41 @@ void UPhysicsControlComponent::SetControlTargetOrientations(
 	const bool     bApplyToControlsWithName,
 	const bool     bApplyToSetsWithName)
 {
+	bool bSuccess = true;
 	for (FName Name : Names)
 	{
-		SetControlTargetOrientation(
+		bSuccess &= SetControlTargetOrientation(
 			Name, Orientation, AngularVelocityDeltaTime, bEnableControl, bApplyControlPointToTarget,
 			bApplyToControlsWithName, bApplyToSetsWithName);
 	}
+	return bSuccess;
 }
 
 //======================================================================================================================
-void UPhysicsControlComponent::SetControlTargetOrientationsInSet(
+bool UPhysicsControlComponent::SetControlTargetOrientationsInSet(
 	const FName    SetName,
 	const FRotator Orientation,
 	const float    AngularVelocityDeltaTime,
 	const bool     bEnableControl,
 	const bool     bApplyControlPointToTarget)
 {
-	SetControlTargetOrientations(GetControlNamesInSet(SetName),
-		Orientation, AngularVelocityDeltaTime, bEnableControl, bApplyControlPointToTarget, true, false);
+	const TArray<FName>& Names = GetControlNamesInSet(SetName);
+	if (Names.IsEmpty())
+	{
+		if (bWarnAboutInvalidNames)
+		{
+			UE_LOG(LogPhysicsControl, Warning,
+				TEXT("SetControlTargetOrientationsInSet - Set %s not found in %s"),
+				*SetName.ToString(), GetOwner() ? *GetOwner()->GetName() : TEXT("Unknown"));
+
+		}
+		return false;
+	}
+	else
+	{
+		return SetControlTargetOrientations(
+			Names, Orientation, AngularVelocityDeltaTime, bEnableControl, bApplyControlPointToTarget, true, false);
+	}
 }
 
 //======================================================================================================================
@@ -1354,7 +1851,8 @@ bool UPhysicsControlComponent::SetControlTargetPositionsFromArray(
 	if (NumControlNames != NumPositions)
 	{
 		UE_LOG(LogPhysicsControl, Warning,
-			TEXT("SetControlTargetPositionsFromArray - names and positions arrays sizes do not match"));
+			TEXT("SetControlTargetPositionsFromArray - names and positions arrays sizes do not match in %s"), 
+			GetOwner() ? *GetOwner()->GetName() : TEXT("Unknown"));
 		return false;
 	}
 	for (int32 Index = 0 ; Index != NumControlNames ; ++Index)
@@ -1378,7 +1876,8 @@ bool UPhysicsControlComponent::SetControlTargetOrientationsFromArray(
 	if (NumControlNames != NumOrientations)
 	{
 		UE_LOG(LogPhysicsControl, Warning,
-			TEXT("SetControlTargetOrientationsFromArray - names and orientations arrays sizes do not match"));
+			TEXT("SetControlTargetOrientationsFromArray - names and orientations arrays sizes do not match in %s"),
+			GetOwner() ? *GetOwner()->GetName() : TEXT("Unknown"));
 		return false;
 	}
 	for (int32 Index = 0 ; Index != NumControlNames ; ++Index)
@@ -1404,7 +1903,8 @@ bool UPhysicsControlComponent::SetControlTargetPositionsAndOrientationsFromArray
 	if (NumControlNames != NumPositions || NumControlNames != NumOrientations)
 	{
 		UE_LOG(LogPhysicsControl, Warning,
-			TEXT("SetControlTargetPositionsAndOrientationsFromArray - names and positions/orientation arrays sizes do not match"));
+			TEXT("SetControlTargetPositionsAndOrientationsFromArray - names and positions/orientation arrays sizes do not match in %s"), 
+			GetOwner() ? *GetOwner()->GetName() : TEXT("Unknown"));
 		return false;
 	}
 	for (int32 Index = 0; Index != NumControlNames; ++Index)
@@ -1461,14 +1961,14 @@ bool UPhysicsControlComponent::SetControlTargetPoses(
 	}
 	if (bWarnAboutInvalidNames)
 	{
-		UE_LOG(LogPhysicsControl, Warning,
-			TEXT("SetControlTargetPoses - invalid name %s"), *Name.ToString());
+		UE_LOG(LogPhysicsControl, Warning, TEXT("SetControlTargetPoses - invalid control name %s for %s"),
+			*Name.ToString(), GetOwner() ? *GetOwner()->GetName() : TEXT("Unknown"));
 	}
 	return false;
 }
 
 //======================================================================================================================
-void UPhysicsControlComponent::SetControlUseSkeletalAnimation(
+bool UPhysicsControlComponent::SetControlUseSkeletalAnimation(
 	const FName Name,
 	const bool  bUseSkeletalAnimation,
 	const float AngularTargetVelocityMultiplier,
@@ -1476,8 +1976,11 @@ void UPhysicsControlComponent::SetControlUseSkeletalAnimation(
 	const bool  bApplyToControlsWithName,
 	const bool  bApplyToSetsWithName)
 {
+	bool bSuccessControl = true;
 	if (bApplyToControlsWithName)
 	{
+		FDisableNameWarnings DisableNameWarnings(bWarnAboutInvalidNames);
+
 		FPhysicsControl* PhysicsControl = FindControl(Name);
 		if (PhysicsControl)
 		{
@@ -1485,16 +1988,36 @@ void UPhysicsControlComponent::SetControlUseSkeletalAnimation(
 			PhysicsControl->ControlData.AngularTargetVelocityMultiplier = AngularTargetVelocityMultiplier;
 			PhysicsControl->ControlData.LinearTargetVelocityMultiplier = LinearTargetVelocityMultiplier;
 		}
+		else
+		{
+			bSuccessControl = false;
+		}
 	}
+
+	bool bSuccessSet = true;
 	if (bApplyToSetsWithName)
 	{
+		FDisableNameWarnings DisableNameWarnings(bWarnAboutInvalidNames);
+
 		SetControlsInSetUseSkeletalAnimation(
 			Name, bUseSkeletalAnimation, AngularTargetVelocityMultiplier, LinearTargetVelocityMultiplier);
 	}
+
+	bool bSuccess =
+		(bApplyToControlsWithName && bSuccessControl) ||
+		(bApplyToSetsWithName && bSuccessSet) ||
+		(!bApplyToControlsWithName && !bApplyToSetsWithName);
+	if (!bSuccess && bWarnAboutInvalidNames)
+	{
+		UE_LOG(LogPhysicsControl, Warning,
+			TEXT("SetControlUseSkeletalAnimation - failed to find controls/set %s in %s"),
+			*Name.ToString(), GetOwner() ? *GetOwner()->GetName() : TEXT("Unknown"));
+	}
+	return bSuccess;
 }
 
 //======================================================================================================================
-void UPhysicsControlComponent::SetControlsUseSkeletalAnimation(
+bool UPhysicsControlComponent::SetControlsUseSkeletalAnimation(
 	const TArray<FName>& Names,
 	const bool           bUseSkeletalAnimation,
 	const float          AngularTargetVelocityMultiplier,
@@ -1502,64 +2025,120 @@ void UPhysicsControlComponent::SetControlsUseSkeletalAnimation(
 	const bool           bApplyToControlsWithName,
 	const bool           bApplyToSetsWithName)
 {
+	bool bSuccess = true;
 	for (FName Name : Names)
 	{
-		SetControlUseSkeletalAnimation(
+		bSuccess &= SetControlUseSkeletalAnimation(
 			Name, bUseSkeletalAnimation, AngularTargetVelocityMultiplier, LinearTargetVelocityMultiplier,
 			bApplyToControlsWithName, bApplyToSetsWithName);
 	}
+	return bSuccess;
 }
 
 //======================================================================================================================
-void UPhysicsControlComponent::SetControlsInSetUseSkeletalAnimation(
+bool UPhysicsControlComponent::SetControlsInSetUseSkeletalAnimation(
 	const FName SetName,
 	const bool  bUseSkeletalAnimation,
 	const float AngularTargetVelocityMultiplier,
 	const float LinearTargetVelocityMultiplier)
 {
-	SetControlsUseSkeletalAnimation(
-		GetControlNamesInSet(SetName), bUseSkeletalAnimation, AngularTargetVelocityMultiplier, 
-		LinearTargetVelocityMultiplier, true, false);
+	const TArray<FName>& Names = GetControlNamesInSet(SetName);
+	if (Names.IsEmpty())
+	{
+		if (bWarnAboutInvalidNames)
+		{
+			UE_LOG(LogPhysicsControl, Warning,
+				TEXT("SetControlsInSetUseSkeletalAnimation - Set %s not found in %s"),
+				*SetName.ToString(), GetOwner() ? *GetOwner()->GetName() : TEXT("Unknown"));
+
+		}
+		return false;
+	}
+	else
+	{
+		return SetControlsUseSkeletalAnimation(
+			Names, bUseSkeletalAnimation, AngularTargetVelocityMultiplier,
+			LinearTargetVelocityMultiplier, true, false);
+	}
 }
 
 //======================================================================================================================
-void UPhysicsControlComponent::SetControlDisableCollision(
+bool UPhysicsControlComponent::SetControlDisableCollision(
 	const FName   Name, 
 	const bool    bDisableCollision,
 	const bool    bApplyToControlsWithName,
 	const bool    bApplyToSetsWithName)
 {
+	bool bSuccessControl = true;
 	if (bApplyToControlsWithName)
 	{
+		FDisableNameWarnings DisableNameWarnings(bWarnAboutInvalidNames);
+
 		FPhysicsControl* PhysicsControl = FindControl(Name);
 		if (PhysicsControl)
 		{
 			PhysicsControl->ControlData.bDisableCollision = bDisableCollision;
 		}
+		else
+		{
+			bSuccessControl = false;
+		}
 	}
+
+	bool bSuccessSet = true;
 	if (bApplyToSetsWithName)
 	{
-		SetControlsInSetDisableCollision(Name, bDisableCollision);
+		FDisableNameWarnings DisableNameWarnings(bWarnAboutInvalidNames);
+
+		bSuccessSet = SetControlsInSetDisableCollision(Name, bDisableCollision);
+	}
+	bool bSuccess =
+		(bApplyToControlsWithName && bSuccessControl) ||
+		(bApplyToSetsWithName && bSuccessSet) ||
+		(!bApplyToControlsWithName && !bApplyToSetsWithName);
+	if (!bSuccess && bWarnAboutInvalidNames)
+	{
+		UE_LOG(LogPhysicsControl, Warning,
+			TEXT("SetControlDisableCollision - failed to find controls/set %s in %s"),
+			*Name.ToString(), GetOwner() ? *GetOwner()->GetName() : TEXT("Unknown"));
 	}
+	return bSuccess;
 }
 
 //======================================================================================================================
-void UPhysicsControlComponent::SetControlsDisableCollision(
+bool UPhysicsControlComponent::SetControlsDisableCollision(
 	const TArray<FName>& Names, 
 	const bool           bDisableCollision,
 	const bool           bApplyToControlsWithName,
 	const bool           bApplyToSetsWithName)
 {
+	bool bSuccess = true;
 	for (FName Name : Names)
 	{
-		SetControlDisableCollision(Name, bDisableCollision, bApplyToControlsWithName, bApplyToSetsWithName);
+		bSuccess &= SetControlDisableCollision(Name, bDisableCollision, bApplyToControlsWithName, bApplyToSetsWithName);
 	}
+	return bSuccess;
 }
 
 //======================================================================================================================
-void UPhysicsControlComponent::SetControlsInSetDisableCollision(const FName SetName, const bool bDisableCollision)
+bool UPhysicsControlComponent::SetControlsInSetDisableCollision(const FName SetName, const bool bDisableCollision)
 {
-	SetControlsDisableCollision(GetControlNamesInSet(SetName), bDisableCollision, true, false);
+	const TArray<FName>& Names = GetControlNamesInSet(SetName);
+	if (Names.IsEmpty())
+	{
+		if (bWarnAboutInvalidNames)
+		{
+			UE_LOG(LogPhysicsControl, Warning,
+				TEXT("SetControlsInSetDisableCollision - Set %s not found in %s"),
+				*SetName.ToString(), GetOwner() ? *GetOwner()->GetName() : TEXT("Unknown"));
+
+		}
+		return false;
+	}
+	else
+	{
+		return SetControlsDisableCollision(GetControlNamesInSet(SetName), bDisableCollision, true, false);
+	}
 }
 
 //======================================================================================================================
@@ -1573,8 +2152,8 @@ bool UPhysicsControlComponent::GetControlData(const FName Name, FPhysicsControlD
 	}
 	if (bWarnAboutInvalidNames)
 	{
-		UE_LOG(LogPhysicsControl, Warning,
-			TEXT("GetControlData - invalid name %s"), *Name.ToString());
+		UE_LOG(LogPhysicsControl, Warning, TEXT("GetControlData - invalid control name %s for %s"),
+			*Name.ToString(), GetOwner() ? *GetOwner()->GetName() : TEXT("Unknown"));
 	}
 	return false;
 }
@@ -1590,8 +2169,8 @@ bool UPhysicsControlComponent::GetControlMultiplier(const FName Name, FPhysicsCo
 	}
 	if (bWarnAboutInvalidNames)
 	{
-		UE_LOG(LogPhysicsControl, Warning,
-			TEXT("GetControlMultiplier - invalid name %s"), *Name.ToString());
+		UE_LOG(LogPhysicsControl, Warning, TEXT("GetControlMultiplier - invalid control name %s for %s"),
+			*Name.ToString(), GetOwner() ? *GetOwner()->GetName() : TEXT("Unknown"));
 	}
 	return false;
 }
@@ -1607,8 +2186,8 @@ bool UPhysicsControlComponent::GetControlTarget(const FName Name, FPhysicsContro
 	}
 	if (bWarnAboutInvalidNames)
 	{
-		UE_LOG(LogPhysicsControl, Warning,
-			TEXT("GetControlTarget - invalid name %s"), *Name.ToString());
+		UE_LOG(LogPhysicsControl, Warning, TEXT("GetControlTarget - invalid control name %s for %s"),
+			*Name.ToString(), GetOwner() ? *GetOwner()->GetName() : TEXT("Unknown"));
 	}
 	return false;
 }
@@ -1623,8 +2202,8 @@ bool UPhysicsControlComponent::GetControlEnabled(const FName Name) const
 	}
 	if (bWarnAboutInvalidNames)
 	{
-		UE_LOG(LogPhysicsControl, Warning,
-			TEXT("GetControlEnabled - invalid name %s"), *Name.ToString());
+		UE_LOG(LogPhysicsControl, Warning, TEXT("GetControlEnabled - invalid control name %s for %s"), 
+			*Name.ToString(), GetOwner() ? *GetOwner()->GetName() : TEXT("Unknown"));
 	}
 	return false;
 }
@@ -1636,7 +2215,8 @@ FName UPhysicsControlComponent::CreateBodyModifier(
 	const FName                       Set,
 	const FPhysicsControlModifierData BodyModifierData)
 {
-	const FName Name = UE::PhysicsControl::GetUniqueBodyModifierName(BoneName, BodyModifierRecords, TEXT(""));
+	const FName Name = UE::PhysicsControl::GetUniqueBodyModifierName(
+		Component, BoneName, BodyModifierRecords, TEXT(""));
 	if (CreateNamedBodyModifier(Name, Component, BoneName, Set, BodyModifierData))
 	{
 		return Name;
@@ -1655,14 +2235,16 @@ bool UPhysicsControlComponent::CreateNamedBodyModifier(
 	if (FindBodyModifierRecord(Name))
 	{
 		UE_LOG(LogPhysicsControl, Warning,
-			TEXT("CreateNamedBodyModifier - modifier with name %s already exists"), *Name.ToString());
+			TEXT("CreateNamedBodyModifier - modifier with name %s already exists in %s"),
+			*Name.ToString(), GetOwner() ? *GetOwner()->GetName() : TEXT("Unknown"));
 		return false;
 	}
 
 	if (!Component)
 	{
 		UE_LOG(LogPhysicsControl, Warning,
-			TEXT("Unable to make a PhysicsBodyModifier as the mesh component has not been set"));
+			TEXT("Unable to make a PhysicsBodyModifier as the mesh component has not been set in %s"), 
+			GetOwner() ? *GetOwner()->GetName() : TEXT("Unknown"));
 		return false;
 	}
 
@@ -1694,7 +2276,8 @@ TArray<FName> UPhysicsControlComponent::CreateBodyModifiersFromSkeletalMeshBelow
 	if (!PhysicsAsset)
 	{
 		UE_LOG(LogPhysicsControl, Warning,
-			TEXT("CreateBodyModifiersFromSkeletalMeshBelow - No physics asset available"));
+			TEXT("CreateBodyModifiersFromSkeletalMeshBelow - No physics asset available in %s"), 
+			GetOwner() ? *GetOwner()->GetName() : TEXT("Unknown"));
 		return Result;
 	}
 
@@ -1752,7 +2335,8 @@ TMap<FName, FPhysicsControlNames> UPhysicsControlComponent::CreateBodyModifiersF
 			else
 			{
 				UE_LOG(LogPhysicsControl, Warning, 
-					TEXT("Failed to make body modifier for %s"), *BoneName.ToString());
+					TEXT("Failed to make body modifier for %s in %s"),
+					*BoneName.ToString(), GetOwner() ? *GetOwner()->GetName() : TEXT("Unknown"));
 			}
 		}
 	}
@@ -1760,120 +2344,230 @@ TMap<FName, FPhysicsControlNames> UPhysicsControlComponent::CreateBodyModifiersF
 }
 
 //======================================================================================================================
-void UPhysicsControlComponent::DestroyBodyModifier(
+bool UPhysicsControlComponent::DestroyBodyModifier(
 	const FName Name, const bool bApplyToModifiersWithName, const bool bApplyToSetsWithName)
 {
+	bool bSuccessModifier = true;
 	if (bApplyToModifiersWithName)
 	{
-		DestroyBodyModifier(Name, EDestroyBehavior::RemoveRecord);
+		FDisableNameWarnings DisableNameWarnings(bWarnAboutInvalidNames);
+		bSuccessModifier &= DestroyBodyModifier(Name, EDestroyBehavior::RemoveRecord);
 	}
+
+	bool bSuccessSet = true;
 	if (bApplyToSetsWithName)
 	{
-		DestroyBodyModifiersInSet(Name);
+		FDisableNameWarnings DisableNameWarnings(bWarnAboutInvalidNames);
+		bSuccessSet = DestroyBodyModifiersInSet(Name);
 	}
+
+	bool bSuccess =
+		(bApplyToModifiersWithName && bSuccessModifier) ||
+		(bApplyToSetsWithName && bSuccessSet) ||
+		(!bApplyToModifiersWithName && !bApplyToSetsWithName);
+	if (!bSuccess && bWarnAboutInvalidNames)
+	{
+		UE_LOG(LogPhysicsControl, Warning,
+			TEXT("DestroyBodyModifier - failed to find modifier/set %s in %s"),
+			*Name.ToString(), GetOwner() ? *GetOwner()->GetName() : TEXT("Unknown"));
+	}
+	return bSuccess;
 }
 
 //======================================================================================================================
-void UPhysicsControlComponent::DestroyBodyModifiers(
+bool UPhysicsControlComponent::DestroyBodyModifiers(
 	const TArray<FName>& Names, const bool bApplyToModifiersWithName, const bool bApplyToSetsWithName)
 {
+	bool bSuccess = true;
 	for (FName Name : Names)
 	{
-		DestroyBodyModifier(Name, bApplyToModifiersWithName, bApplyToSetsWithName);
+		bSuccess &= DestroyBodyModifier(Name, bApplyToModifiersWithName, bApplyToSetsWithName);
 	}
+	return bSuccess;
 }
 
 //======================================================================================================================
-void UPhysicsControlComponent::DestroyBodyModifiersInSet(const FName SetName)
+bool UPhysicsControlComponent::DestroyBodyModifiersInSet(const FName SetName)
 {
 	TArray<FName> Names = GetBodyModifierNamesInSet(SetName);
-	DestroyBodyModifiers(Names, true, false);
+	if (Names.IsEmpty())
+	{
+		if (bWarnAboutInvalidNames)
+		{
+			UE_LOG(LogPhysicsControl, Warning,
+				TEXT("DestroyBodyModifiersInSet - Set %s not found in %s"),
+				*SetName.ToString(), GetOwner() ? *GetOwner()->GetName() : TEXT("Unknown"));
+
+		}
+		return false;
+	}
+	else
+	{
+		return DestroyBodyModifiers(Names, true, false);
+	}
 }
 
 //======================================================================================================================
 // Note - some params passed by value to allow them to be unconnected and enable inline BP editing
-void UPhysicsControlComponent::SetBodyModifierData(
+bool UPhysicsControlComponent::SetBodyModifierData(
 	const FName                       Name,
 	const FPhysicsControlModifierData ModifierData, 
 	const bool                        bApplyToModifiersWithName, 
 	const bool                        bApplyToSetsWithName)
 {
+	bool bSuccessModifier = true;
 	if (bApplyToModifiersWithName)
 	{
+		FDisableNameWarnings DisableNameWarnings(bWarnAboutInvalidNames);
 		FPhysicsBodyModifierRecord* Record = FindBodyModifierRecord(Name);
 		if (Record)
 		{
 			Record->BodyModifier.ModifierData = ModifierData;
 		}
+		else
+		{
+			bSuccessModifier = false;
+		}
 	}
+
+	bool bSuccessSet = true;
 	if (bApplyToSetsWithName)
 	{
-		SetBodyModifierDatasInSet(Name, ModifierData);
+		FDisableNameWarnings DisableNameWarnings(bWarnAboutInvalidNames);
+		bSuccessSet = SetBodyModifierDatasInSet(Name, ModifierData);
+	}
+
+	bool bSuccess =
+		(bApplyToModifiersWithName && bSuccessModifier) ||
+		(bApplyToSetsWithName && bSuccessSet) ||
+		(!bApplyToModifiersWithName && !bApplyToSetsWithName);
+	if (!bSuccess && bWarnAboutInvalidNames)
+	{
+		UE_LOG(LogPhysicsControl, Warning,
+			TEXT("SetBodyModifierData - failed to find modifier/set %s in %s"),
+			*Name.ToString(), GetOwner() ? *GetOwner()->GetName() : TEXT("Unknown"));
 	}
+	return bSuccess;
 }
 
 //======================================================================================================================
-void UPhysicsControlComponent::SetBodyModifierDatas(
+bool UPhysicsControlComponent::SetBodyModifierDatas(
 	const TArray<FName>&              Names,
 	const FPhysicsControlModifierData ModifierData,
 	const bool                        bApplyToModifiersWithName,
 	const bool                        bApplyToSetsWithName)
 {
+	bool bSuccess = true;
 	for (FName Name : Names)
 	{
-		SetBodyModifierData(Name, ModifierData, bApplyToModifiersWithName, bApplyToSetsWithName);
+		bSuccess = SetBodyModifierData(Name, ModifierData, bApplyToModifiersWithName, bApplyToSetsWithName);
 	}
+	return bSuccess;
 }
 
 //======================================================================================================================
-void UPhysicsControlComponent::SetBodyModifierDatasInSet(
+bool UPhysicsControlComponent::SetBodyModifierDatasInSet(
 	const FName                       SetName,
 	const FPhysicsControlModifierData ModifierData)
 {
-	SetBodyModifierDatas(GetBodyModifierNamesInSet(SetName), ModifierData, true, false);
+	TArray<FName> Names = GetBodyModifierNamesInSet(SetName);
+	if (Names.IsEmpty())
+	{
+		if (bWarnAboutInvalidNames)
+		{
+			UE_LOG(LogPhysicsControl, Warning,
+				TEXT("SetBodyModifierDatasInSet - Set %s not found in %s"),
+				*SetName.ToString(), GetOwner() ? *GetOwner()->GetName() : TEXT("Unknown"));
+
+		}
+		return false;
+	}
+	else
+	{
+		return SetBodyModifierDatas(Names, ModifierData, true, false);
+	}
 }
 
 //======================================================================================================================
 // Note - some params passed by value to allow them to be unconnected and enable inline BP editing
-void UPhysicsControlComponent::SetBodyModifierSparseData(
+bool UPhysicsControlComponent::SetBodyModifierSparseData(
 	const FName                             Name,
 	const FPhysicsControlModifierSparseData ModifierData,
 	const bool                              bApplyToModifiersWithName,
 	const bool                              bApplyToSetsWithName)
 {
+	bool bSuccessModifier = true;
 	if (bApplyToModifiersWithName)
 	{
+		FDisableNameWarnings DisableNameWarnings(bWarnAboutInvalidNames);
 		FPhysicsBodyModifierRecord* Record = FindBodyModifierRecord(Name);
 		if (Record)
 		{
 			Record->BodyModifier.ModifierData.UpdateFromSparseData(ModifierData);
 		}
+		else
+		{
+			bSuccessModifier = false;
+		}
 	}
+
+	bool bSuccessSet = true;
 	if (bApplyToSetsWithName)
 	{
-		SetBodyModifierSparseDatasInSet(Name, ModifierData);
+		FDisableNameWarnings DisableNameWarnings(bWarnAboutInvalidNames);
+		bSuccessSet = SetBodyModifierSparseDatasInSet(Name, ModifierData);
+	}
+
+	bool bSuccess =
+		(bApplyToModifiersWithName && bSuccessModifier) ||
+		(bApplyToSetsWithName && bSuccessSet) ||
+		(!bApplyToModifiersWithName && !bApplyToSetsWithName);
+	if (!bSuccess && bWarnAboutInvalidNames)
+	{
+		UE_LOG(LogPhysicsControl, Warning,
+			TEXT("SetBodyModifierSparseData - failed to find modifier/set %s in %s"),
+			*Name.ToString(), GetOwner() ? *GetOwner()->GetName() : TEXT("Unknown"));
 	}
+	return bSuccess;
+
 }
 
 //======================================================================================================================
-void UPhysicsControlComponent::SetBodyModifierSparseDatas(
+bool UPhysicsControlComponent::SetBodyModifierSparseDatas(
 	const TArray<FName>&                    Names,
 	const FPhysicsControlModifierSparseData ModifierData,
 	const bool                              bApplyToModifiersWithName,
 	const bool                              bApplyToSetsWithName)
 {
+	bool bSuccess = true;
 	for (FName Name : Names)
 	{
-		SetBodyModifierSparseData(Name, ModifierData, bApplyToModifiersWithName, bApplyToSetsWithName);
+		bSuccess &= SetBodyModifierSparseData(Name, ModifierData, bApplyToModifiersWithName, bApplyToSetsWithName);
 	}
+	return bSuccess;
 }
 
 //======================================================================================================================
-void UPhysicsControlComponent::SetBodyModifierSparseDatasInSet(
+bool UPhysicsControlComponent::SetBodyModifierSparseDatasInSet(
 	const FName                             SetName,
 	const FPhysicsControlModifierSparseData ModifierData)
 {
-	SetBodyModifierSparseDatas(GetBodyModifierNamesInSet(SetName), ModifierData, true, false);
+	TArray<FName> Names = GetBodyModifierNamesInSet(SetName);
+	if (Names.IsEmpty())
+	{
+		if (bWarnAboutInvalidNames)
+		{
+			UE_LOG(LogPhysicsControl, Warning,
+				TEXT("SetBodyModifierSparseDatasInSet - Set %s not found in %s"),
+				*SetName.ToString(), GetOwner() ? *GetOwner()->GetName() : TEXT("Unknown"));
+
+		}
+		return false;
+	}
+	else
+	{
+		return SetBodyModifierSparseDatas(Names, ModifierData, true, false);
+	}
 }
 
 //======================================================================================================================
@@ -1896,264 +2590,495 @@ bool UPhysicsControlComponent::SetBodyModifierKinematicTarget(
 	}
 	if (bWarnAboutInvalidNames)
 	{
-		UE_LOG(LogPhysicsControl, Warning,
-			TEXT("SetBodyModifierKinematicTarget - invalid name %s"), *Name.ToString());
+		UE_LOG(LogPhysicsControl, Warning, TEXT("SetBodyModifierKinematicTarget - invalid modifier name %s for %s"),
+			*Name.ToString(), GetOwner() ? *GetOwner()->GetName() : TEXT("Unknown"));
 	}
 	return false;
 }
 
 //======================================================================================================================
-void UPhysicsControlComponent::SetBodyModifierMovementType(
+bool UPhysicsControlComponent::SetBodyModifierMovementType(
 	const FName                Name,
 	const EPhysicsMovementType MovementType,
 	const bool                 bApplyToModifiersWithName,
 	const bool                 bApplyToSetsWithName)
 {
+	bool bSuccessModifier = true;
 	if (bApplyToModifiersWithName)
 	{
+		FDisableNameWarnings DisableNameWarnings(bWarnAboutInvalidNames);
+
 		FPhysicsBodyModifierRecord* Record = FindBodyModifierRecord(Name);
 		if (Record)
 		{
 			Record->BodyModifier.ModifierData.MovementType = MovementType;
 		}
+		else
+		{
+			bSuccessModifier = false;
+		}
 	}
+
+	bool bSuccessSet = true;
 	if (bApplyToSetsWithName)
 	{
-		SetBodyModifiersInSetMovementType(Name, MovementType);
+		FDisableNameWarnings DisableNameWarnings(bWarnAboutInvalidNames);
+
+		bSuccessSet = SetBodyModifiersInSetMovementType(Name, MovementType);
+	}
+
+	bool bSuccess =
+		(bApplyToModifiersWithName && bSuccessModifier) ||
+		(bApplyToSetsWithName && bSuccessSet) ||
+		(!bApplyToModifiersWithName && !bApplyToSetsWithName);
+	if (!bSuccess && bWarnAboutInvalidNames)
+	{
+		UE_LOG(LogPhysicsControl, Warning,
+			TEXT("SetBodyModifierMovementType - failed to find modifier/set %s in %s"),
+			*Name.ToString(), GetOwner() ? *GetOwner()->GetName() : TEXT("Unknown"));
 	}
+	return bSuccess;
 }
 
 //======================================================================================================================
-void UPhysicsControlComponent::SetBodyModifiersMovementType(
+bool UPhysicsControlComponent::SetBodyModifiersMovementType(
 	const TArray<FName>&       Names,
 	const EPhysicsMovementType MovementType,
 	const bool                 bApplyToModifiersWithName,
 	const bool                 bApplyToSetsWithName)
 {
+	bool bSuccess = true;
 	for (FName Name : Names)
 	{
-		SetBodyModifierMovementType(Name, MovementType, bApplyToModifiersWithName, bApplyToSetsWithName);
+		bSuccess &= SetBodyModifierMovementType(Name, MovementType, bApplyToModifiersWithName, bApplyToSetsWithName);
 	}
+	return bSuccess;
 }
 
 //======================================================================================================================
-void UPhysicsControlComponent::SetBodyModifiersInSetMovementType(
+bool UPhysicsControlComponent::SetBodyModifiersInSetMovementType(
 	const FName                SetName,
 	const EPhysicsMovementType MovementType)
 {
-	SetBodyModifiersMovementType(GetBodyModifierNamesInSet(SetName), MovementType, true, false);
+	TArray<FName> Names = GetBodyModifierNamesInSet(SetName);
+	if (Names.IsEmpty())
+	{
+		if (bWarnAboutInvalidNames)
+		{
+			UE_LOG(LogPhysicsControl, Warning,
+				TEXT("SetBodyModifiersInSetMovementType - Set %s not found in %s"),
+				*SetName.ToString(), GetOwner() ? *GetOwner()->GetName() : TEXT("Unknown"));
+
+		}
+		return false;
+	}
+	else
+	{
+		return SetBodyModifiersMovementType(Names, MovementType, true, false);
+	}
 }
 
 //======================================================================================================================
-void UPhysicsControlComponent::SetBodyModifierCollisionType(
+bool UPhysicsControlComponent::SetBodyModifierCollisionType(
 	const FName                   Name,
 	const ECollisionEnabled::Type CollisionType,
 	const bool                    bApplyToModifiersWithName,
 	const bool                    bApplyToSetsWithName)
 {
+	bool bSuccessModifier = true;
 	if (bApplyToModifiersWithName)
 	{
+		FDisableNameWarnings DisableNameWarnings(bWarnAboutInvalidNames);
 		FPhysicsBodyModifierRecord* Record = FindBodyModifierRecord(Name);
 		if (Record)
 		{
 			Record->BodyModifier.ModifierData.CollisionType = CollisionType;
 		}
+		else
+		{
+			bSuccessModifier = false;
+		}
 	}
+
+	bool bSuccessSet = true;
 	if (bApplyToSetsWithName)
 	{
-		SetBodyModifiersInSetCollisionType(Name, CollisionType);
+		FDisableNameWarnings DisableNameWarnings(bWarnAboutInvalidNames);
+		bSuccessSet = SetBodyModifiersInSetCollisionType(Name, CollisionType);
+	}
+
+	bool bSuccess =
+		(bApplyToModifiersWithName && bSuccessModifier) ||
+		(bApplyToSetsWithName && bSuccessSet) ||
+		(!bApplyToModifiersWithName && !bApplyToSetsWithName);
+	if (!bSuccess && bWarnAboutInvalidNames)
+	{
+		UE_LOG(LogPhysicsControl, Warning,
+			TEXT("SetBodyModifierCollisionType - failed to find modifier/set %s in %s"),
+			*Name.ToString(), GetOwner() ? *GetOwner()->GetName() : TEXT("Unknown"));
 	}
+	return bSuccess;
 }
 
 //======================================================================================================================
-void UPhysicsControlComponent::SetBodyModifiersCollisionType(
+bool UPhysicsControlComponent::SetBodyModifiersCollisionType(
 	const TArray<FName>&          Names,
 	const ECollisionEnabled::Type CollisionType,
 	const bool                    bApplyToModifiersWithName,
 	const bool                    bApplyToSetsWithName)
 {
+	bool bSuccess = true;
 	for (FName Name : Names)
 	{
-		SetBodyModifierCollisionType(Name, CollisionType, bApplyToModifiersWithName, bApplyToSetsWithName);
+		bSuccess &= SetBodyModifierCollisionType(Name, CollisionType, bApplyToModifiersWithName, bApplyToSetsWithName);
 	}
+	return bSuccess;
 }
 
 //======================================================================================================================
-void UPhysicsControlComponent::SetBodyModifiersInSetCollisionType(
+bool UPhysicsControlComponent::SetBodyModifiersInSetCollisionType(
 	const FName                   SetName,
 	const ECollisionEnabled::Type CollisionType)
 {
-	SetBodyModifiersCollisionType(GetBodyModifierNamesInSet(SetName), CollisionType, true, false);
+	TArray<FName> Names = GetBodyModifierNamesInSet(SetName);
+	if (Names.IsEmpty())
+	{
+		if (bWarnAboutInvalidNames)
+		{
+			UE_LOG(LogPhysicsControl, Warning,
+				TEXT("SetBodyModifiersInSetCollisionType - Set %s not found in %s"),
+				*SetName.ToString(), GetOwner() ? *GetOwner()->GetName() : TEXT("Unknown"));
+
+		}
+		return false;
+	}
+	else
+	{
+		return SetBodyModifiersCollisionType(Names, CollisionType, true, false);
+	}
 }
 
 //======================================================================================================================
-void UPhysicsControlComponent::SetBodyModifierGravityMultiplier(
+bool UPhysicsControlComponent::SetBodyModifierGravityMultiplier(
 	const FName Name,
 	const float GravityMultiplier,
 	const bool  bApplyToModifiersWithName,
 	const bool  bApplyToSetsWithName)
 {
+	bool bSuccessModifier = true;
 	if (bApplyToModifiersWithName)
 	{
+		FDisableNameWarnings DisableNameWarnings(bWarnAboutInvalidNames);
 		FPhysicsBodyModifierRecord* Record = FindBodyModifierRecord(Name);
 		if (Record)
 		{
 			Record->BodyModifier.ModifierData.GravityMultiplier = GravityMultiplier;
 		}
+		else
+		{
+			bSuccessModifier = false;
+		}
 	}
+
+	bool bSuccessSet = true;
 	if (bApplyToSetsWithName)
 	{
-		SetBodyModifiersInSetGravityMultiplier(Name, GravityMultiplier);
+		FDisableNameWarnings DisableNameWarnings(bWarnAboutInvalidNames);
+		bSuccessSet = SetBodyModifiersInSetGravityMultiplier(Name, GravityMultiplier);
+	}
+
+	bool bSuccess =
+		(bApplyToModifiersWithName && bSuccessModifier) ||
+		(bApplyToSetsWithName && bSuccessSet) ||
+		(!bApplyToModifiersWithName && !bApplyToSetsWithName);
+	if (!bSuccess && bWarnAboutInvalidNames)
+	{
+		UE_LOG(LogPhysicsControl, Warning,
+			TEXT("SetBodyModifierGravityMultiplier - failed to find modifier/set %s in %s"),
+			*Name.ToString(), GetOwner() ? *GetOwner()->GetName() : TEXT("Unknown"));
 	}
+	return bSuccess;
+
 }
 
 //======================================================================================================================
-void UPhysicsControlComponent::SetBodyModifiersGravityMultiplier(
+bool UPhysicsControlComponent::SetBodyModifiersGravityMultiplier(
 	const TArray<FName>& Names,
 	const float          GravityMultiplier,
 	const bool           bApplyToModifiersWithName,
 	const bool           bApplyToSetsWithName)
 {
+	bool bSuccess = true;
 	for (FName Name : Names)
 	{
-		SetBodyModifierGravityMultiplier(Name, GravityMultiplier, bApplyToModifiersWithName, bApplyToSetsWithName);
+		bSuccess = SetBodyModifierGravityMultiplier(Name, GravityMultiplier, bApplyToModifiersWithName, bApplyToSetsWithName);
 	}
+	return bSuccess;
 }
 
 //======================================================================================================================
-void UPhysicsControlComponent::SetBodyModifiersInSetGravityMultiplier(
+bool UPhysicsControlComponent::SetBodyModifiersInSetGravityMultiplier(
 	const FName SetName,
 	const float GravityMultiplier)
 {
-	SetBodyModifiersGravityMultiplier(GetBodyModifierNamesInSet(SetName), GravityMultiplier, true, false);
+	TArray<FName> Names = GetBodyModifierNamesInSet(SetName);
+	if (Names.IsEmpty())
+	{
+		if (bWarnAboutInvalidNames)
+		{
+			UE_LOG(LogPhysicsControl, Warning,
+				TEXT("SetBodyModifiersInSetGravityMultiplier - Set %s not found in %s"),
+				*SetName.ToString(), GetOwner() ? *GetOwner()->GetName() : TEXT("Unknown"));
+
+		}
+		return false;
+	}
+	else
+	{
+		return SetBodyModifiersGravityMultiplier(Names, GravityMultiplier, true, false);
+	}
 }
 
 //======================================================================================================================
-void UPhysicsControlComponent::SetBodyModifierPhysicsBlendWeight(
+bool UPhysicsControlComponent::SetBodyModifierPhysicsBlendWeight(
 	const FName Name,
 	const float PhysicsBlendWeight,
 	const bool  bApplyToModifiersWithName,
 	const bool  bApplyToSetsWithName)
 {
+	bool bSuccessModifier = true;
 	if (bApplyToModifiersWithName)
 	{
+		FDisableNameWarnings DisableNameWarnings(bWarnAboutInvalidNames);
 		FPhysicsBodyModifierRecord* Record = FindBodyModifierRecord(Name);
 		if (Record)
 		{
 			Record->BodyModifier.ModifierData.PhysicsBlendWeight = PhysicsBlendWeight;
 		}
+		else
+		{
+			bSuccessModifier = false;
+		}
 	}
+
+	bool bSuccessSet = true;
 	if (bApplyToSetsWithName)
 	{
-		SetBodyModifiersInSetPhysicsBlendWeight(Name, PhysicsBlendWeight);
+		FDisableNameWarnings DisableNameWarnings(bWarnAboutInvalidNames);
+		bSuccessSet = SetBodyModifiersInSetPhysicsBlendWeight(Name, PhysicsBlendWeight);
+	}
+
+	bool bSuccess =
+		(bApplyToModifiersWithName && bSuccessModifier) ||
+		(bApplyToSetsWithName && bSuccessSet) ||
+		(!bApplyToModifiersWithName && !bApplyToSetsWithName);
+	if (!bSuccess && bWarnAboutInvalidNames)
+	{
+		UE_LOG(LogPhysicsControl, Warning,
+			TEXT("SetBodyModifierPhysicsBlendWeight - failed to find modifier/set %s in %s"),
+			*Name.ToString(), GetOwner() ? *GetOwner()->GetName() : TEXT("Unknown"));
 	}
+	return bSuccess;
 }
 
 //======================================================================================================================
-void UPhysicsControlComponent::SetBodyModifiersPhysicsBlendWeight(
+bool UPhysicsControlComponent::SetBodyModifiersPhysicsBlendWeight(
 	const TArray<FName>& Names,
 	const float          PhysicsBlendWeight,
 	const bool           bApplyToModifiersWithName,
 	const bool           bApplyToSetsWithName)
 {
+	bool bSuccess = true;
 	for (FName Name : Names)
 	{
-		SetBodyModifierPhysicsBlendWeight(Name, PhysicsBlendWeight, bApplyToModifiersWithName, bApplyToSetsWithName);
+		bSuccess &= SetBodyModifierPhysicsBlendWeight(Name, PhysicsBlendWeight, bApplyToModifiersWithName, bApplyToSetsWithName);
 	}
+	return bSuccess;
 }
 
 //======================================================================================================================
-void UPhysicsControlComponent::SetBodyModifiersInSetPhysicsBlendWeight(
+bool UPhysicsControlComponent::SetBodyModifiersInSetPhysicsBlendWeight(
 	const FName SetName,
 	const float PhysicsBlendWeight)
 {
-	SetBodyModifiersPhysicsBlendWeight(GetBodyModifierNamesInSet(SetName), PhysicsBlendWeight, true, false);
+	TArray<FName> Names = GetBodyModifierNamesInSet(SetName);
+	if (Names.IsEmpty())
+	{
+		if (bWarnAboutInvalidNames)
+		{
+			UE_LOG(LogPhysicsControl, Warning,
+				TEXT("SetBodyModifiersInSetPhysicsBlendWeight - Set %s not found in %s"),
+				*SetName.ToString(), GetOwner() ? *GetOwner()->GetName() : TEXT("Unknown"));
+
+		}
+		return false;
+	}
+	else
+	{
+		return SetBodyModifiersPhysicsBlendWeight(Names, PhysicsBlendWeight, true, false);
+	}
 }
 
 //======================================================================================================================
-void UPhysicsControlComponent::SetBodyModifierKinematicTargetSpace(
+bool UPhysicsControlComponent::SetBodyModifierKinematicTargetSpace(
 	const FName                               Name,
 	const EPhysicsControlKinematicTargetSpace KinematicTargetSpace,
 	const bool                                bApplyToModifiersWithName,
 	const bool                                bApplyToSetsWithName)
 {
+	bool bSuccessModifier = true;
 	if (bApplyToModifiersWithName)
 	{
+		FDisableNameWarnings DisableNameWarnings(bWarnAboutInvalidNames);
 		FPhysicsBodyModifierRecord* Record = FindBodyModifierRecord(Name);
 		if (Record)
 		{
 			Record->BodyModifier.ModifierData.KinematicTargetSpace = KinematicTargetSpace;
 		}
+		else
+		{
+			bSuccessModifier = false;
+		}
 	}
+
+	bool bSuccessSet = true;
 	if (bApplyToSetsWithName)
 	{
-		SetBodyModifiersInSetKinematicTargetSpace(Name, KinematicTargetSpace);
+		FDisableNameWarnings DisableNameWarnings(bWarnAboutInvalidNames);
+		bSuccessSet = SetBodyModifiersInSetKinematicTargetSpace(Name, KinematicTargetSpace);
+	}
+
+	bool bSuccess =
+		(bApplyToModifiersWithName && bSuccessModifier) ||
+		(bApplyToSetsWithName && bSuccessSet) ||
+		(!bApplyToModifiersWithName && !bApplyToSetsWithName);
+	if (!bSuccess && bWarnAboutInvalidNames)
+	{
+		UE_LOG(LogPhysicsControl, Warning,
+			TEXT("SetBodyModifierKinematicTargetSpace - failed to find modifier/set %s in %s"),
+			*Name.ToString(), GetOwner() ? *GetOwner()->GetName() : TEXT("Unknown"));
 	}
+	return bSuccess;
 }
 
 //======================================================================================================================
-void UPhysicsControlComponent::SetBodyModifiersKinematicTargetSpace(
+bool UPhysicsControlComponent::SetBodyModifiersKinematicTargetSpace(
 	const TArray<FName>&                      Names,
 	const EPhysicsControlKinematicTargetSpace KinematicTargetSpace,
 	const bool                                bApplyToModifiersWithName,
 	const bool                                bApplyToSetsWithName)
 {
+	bool bSuccess = true;
 	for (FName Name : Names)
 	{
-		SetBodyModifierKinematicTargetSpace(Name, KinematicTargetSpace, bApplyToModifiersWithName, bApplyToSetsWithName);
+		bSuccess &= SetBodyModifierKinematicTargetSpace(Name, KinematicTargetSpace, bApplyToModifiersWithName, bApplyToSetsWithName);
 	}
+	return bSuccess;
 }
 
 //======================================================================================================================
-void UPhysicsControlComponent::SetBodyModifiersInSetKinematicTargetSpace(
+bool UPhysicsControlComponent::SetBodyModifiersInSetKinematicTargetSpace(
 	const FName                               SetName,
 	const EPhysicsControlKinematicTargetSpace KinematicTargetSpace)
 {
-	SetBodyModifiersKinematicTargetSpace(GetBodyModifierNamesInSet(SetName), KinematicTargetSpace, true, false);
+	TArray<FName> Names = GetBodyModifierNamesInSet(SetName);
+	if (Names.IsEmpty())
+	{
+		if (bWarnAboutInvalidNames)
+		{
+			UE_LOG(LogPhysicsControl, Warning,
+				TEXT("SetBodyModifiersInSetKinematicTargetSpace - Set %s not found in %s"),
+				*SetName.ToString(), GetOwner() ? *GetOwner()->GetName() : TEXT("Unknown"));
+
+		}
+		return false;
+	}
+	else
+	{
+		return SetBodyModifiersKinematicTargetSpace(Names, KinematicTargetSpace, true, false);
+	}
 }
 
 //======================================================================================================================
-void UPhysicsControlComponent::SetBodyModifierUpdateKinematicFromSimulation(
+bool UPhysicsControlComponent::SetBodyModifierUpdateKinematicFromSimulation(
 	const FName Name,
 	const bool  bUpdateKinematicFromSimulation,
 	const bool  bApplyToModifiersWithName,
 	const bool  bApplyToSetsWithName)
 {
+	bool bSuccessModifier = true;
 	if (bApplyToModifiersWithName)
 	{
+		FDisableNameWarnings DisableNameWarnings(bWarnAboutInvalidNames);
 		FPhysicsBodyModifierRecord* Record = FindBodyModifierRecord(Name);
 		if (Record)
 		{
 			Record->BodyModifier.ModifierData.bUpdateKinematicFromSimulation = bUpdateKinematicFromSimulation;
 		}
+		else
+		{
+			bSuccessModifier = false;
+		}
 	}
+
+	bool bSuccessSet = true;
 	if (bApplyToSetsWithName)
 	{
-		SetBodyModifiersInSetUpdateKinematicFromSimulation(Name, bUpdateKinematicFromSimulation);
+		FDisableNameWarnings DisableNameWarnings(bWarnAboutInvalidNames);
+		bSuccessSet = SetBodyModifiersInSetUpdateKinematicFromSimulation(Name, bUpdateKinematicFromSimulation);
+	}
+
+	bool bSuccess =
+		(bApplyToModifiersWithName && bSuccessModifier) ||
+		(bApplyToSetsWithName && bSuccessSet) ||
+		(!bApplyToModifiersWithName && !bApplyToSetsWithName);
+	if (!bSuccess && bWarnAboutInvalidNames)
+	{
+		UE_LOG(LogPhysicsControl, Warning,
+			TEXT("SetBodyModifierUpdateKinematicFromSimulation - failed to find modifier/set %s in %s"),
+			*Name.ToString(), GetOwner() ? *GetOwner()->GetName() : TEXT("Unknown"));
 	}
+	return bSuccess;
+
 }
 
 //======================================================================================================================
-void UPhysicsControlComponent::SetBodyModifiersUpdateKinematicFromSimulation(
+bool UPhysicsControlComponent::SetBodyModifiersUpdateKinematicFromSimulation(
 	const TArray<FName>& Names,
 	const bool           bUpdateKinematicFromSimulation,
 	const bool           bApplyToModifiersWithName,
 	const bool           bApplyToSetsWithName)
 {
+	bool bSuccess = true;
 	for (FName Name : Names)
 	{
-		SetBodyModifierUpdateKinematicFromSimulation(
+		bSuccess &= SetBodyModifierUpdateKinematicFromSimulation(
 			Name, bUpdateKinematicFromSimulation, bApplyToModifiersWithName, bApplyToSetsWithName);
 	}
+	return bSuccess;
 }
 
 //======================================================================================================================
-void UPhysicsControlComponent::SetBodyModifiersInSetUpdateKinematicFromSimulation(
+bool UPhysicsControlComponent::SetBodyModifiersInSetUpdateKinematicFromSimulation(
 	const FName SetName,
 	const bool  bUpdateKinematicFromSimulation)
 {
-	SetBodyModifiersUpdateKinematicFromSimulation(
-		GetBodyModifierNamesInSet(SetName), bUpdateKinematicFromSimulation, true, false);
+	TArray<FName> Names = GetBodyModifierNamesInSet(SetName);
+	if (Names.IsEmpty())
+	{
+		if (bWarnAboutInvalidNames)
+		{
+			UE_LOG(LogPhysicsControl, Warning,
+				TEXT("SetBodyModifiersInSetUpdateKinematicFromSimulation - Set %s not found in %s"),
+				*SetName.ToString(), GetOwner() ? *GetOwner()->GetName() : TEXT("Unknown"));
+
+		}
+		return false;
+	}
+	else
+	{
+		return SetBodyModifiersUpdateKinematicFromSimulation(Names, bUpdateKinematicFromSimulation, true, false);
+	}
 }
 
 //======================================================================================================================
@@ -2210,7 +3135,8 @@ bool UPhysicsControlComponent::CreateControlsAndBodyModifiersFromPhysicsControlA
 	if (!PhysicsControlAsset.LoadSynchronous())
 	{
 		UE_LOG(LogPhysicsControl, Warning,
-			TEXT("CreateControlsAndBodyModifiersFromPhysicsControlAsset - unable to get/load the control profile asset"));
+			TEXT("CreateControlsAndBodyModifiersFromPhysicsControlAsset - unable to get/load the control profile asset in %s"), 
+			GetOwner() ? *GetOwner()->GetName() : TEXT("Unknown"));
 		return false;
 	}
 
@@ -2292,8 +3218,9 @@ bool UPhysicsControlComponent::InvokeControlProfile(
 	{
 		if (bWarnAboutInvalidNames)
 		{
-			UE_LOG(LogPhysicsControl, Warning,
-				TEXT("InvokeControlProfile - control profile asset is invalid or missing"));
+			UE_LOG(LogPhysicsControl, Warning, 
+				TEXT("InvokeControlProfile - control profile asset is invalid or missing for %s"),
+				GetOwner() ? *GetOwner()->GetName() : TEXT("Unknown"));
 		}
 		return false;
 	}
@@ -2306,7 +3233,8 @@ bool UPhysicsControlComponent::InvokeControlProfile(
 		if (bWarnAboutInvalidNames)
 		{
 			UE_LOG(LogPhysicsControl, Warning,
-				TEXT("InvokeControlProfile - control profile %s not found"), *ProfileName.ToString());
+				TEXT("InvokeControlProfile - control profile %s not found for %s"), 
+				*ProfileName.ToString(), GetOwner() ? *GetOwner()->GetName() : TEXT("Unknown"));
 		}
 		return false;
 	}
@@ -2346,7 +3274,8 @@ void UPhysicsControlComponent::ApplyControlAndModifierUpdates(
 					if (bWarnAboutInvalidNames)
 					{
 						UE_LOG(LogPhysicsControl, Warning,
-							TEXT("ApplyControlAndModifierUpdates: Failed to find control with name %s"), *Name.ToString());
+							TEXT("ApplyControlAndModifierUpdates: Failed to find control with name %s for %s"),
+							*Name.ToString(), GetOwner() ? *GetOwner()->GetName() : TEXT("Unknown"));
 					}
 				}
 			}
@@ -2371,7 +3300,8 @@ void UPhysicsControlComponent::ApplyControlAndModifierUpdates(
 					if (bWarnAboutInvalidNames)
 					{
 						UE_LOG(LogPhysicsControl, Warning,
-							TEXT("ApplyControlAndModifierUpdates: Failed to find control with name %s"), *Name.ToString());
+							TEXT("ApplyControlAndModifierUpdates: Failed to find control with name %s for %s"),
+							*Name.ToString(), GetOwner() ? *GetOwner()->GetName() : TEXT("Unknown"));
 					}
 				}
 			}
@@ -2393,7 +3323,8 @@ void UPhysicsControlComponent::ApplyControlAndModifierUpdates(
 				else
 				{
 					UE_LOG(LogPhysicsControl, Warning,
-						TEXT("ApplyControlAndModifierUpdates: Failed to find modifier with name %s"), *Name.ToString());
+						TEXT("ApplyControlAndModifierUpdates: Failed to find modifier with name %s in %s"),
+						*Name.ToString(), GetOwner() ? *GetOwner()->GetName() : TEXT("Unknown"));
 				}
 			}
 		}
@@ -2521,7 +3452,8 @@ TArray<FTransform> UPhysicsControlComponent::GetCachedBoneTransforms(
 			if (bWarnAboutInvalidNames)
 			{
 				UE_LOG(LogPhysicsControl, Warning,
-					TEXT("GetCachedBoneTransforms - unable to get bone data for %s"), *BoneName.ToString());
+					TEXT("GetCachedBoneTransforms - unable to get bone %s data for %s"),
+					*BoneName.ToString(), GetOwner() ? *GetOwner()->GetName() : TEXT("Unknown"));
 			}
 			Result.Add(FTransform::Identity);
 		}
@@ -2550,7 +3482,8 @@ TArray<FVector> UPhysicsControlComponent::GetCachedBonePositions(
 			if (bWarnAboutInvalidNames)
 			{
 				UE_LOG(LogPhysicsControl, Warning,
-					TEXT("GetCachedBonePositions - unable to get bone data for %s"), *BoneName.ToString());
+					TEXT("GetCachedBonePositions - unable to get bone %s data for %s"),
+					*BoneName.ToString(), GetOwner() ? *GetOwner()->GetName() : TEXT("Unknown"));
 			}
 			Result.Add(FVector::ZeroVector);
 		}
@@ -2578,8 +3511,9 @@ TArray<FRotator> UPhysicsControlComponent::GetCachedBoneOrientations(
 		{
 			if (bWarnAboutInvalidNames)
 			{
-				UE_LOG(LogPhysicsControl, Warning,
-					TEXT("GetCachedBoneOrientations - unable to get bone data for %s"), *BoneName.ToString());
+				UE_LOG(LogPhysicsControl, Warning, 
+					TEXT("GetCachedBoneOrientations - unable to get bone %s data for %s"), 
+					*BoneName.ToString(), GetOwner() ? *GetOwner()->GetName() : TEXT("Unknown"));
 			}
 			Result.Add(FRotator::ZeroRotator);
 		}
@@ -2601,7 +3535,8 @@ FTransform UPhysicsControlComponent::GetCachedBoneTransform(
 	if (bWarnAboutInvalidNames)
 	{
 		UE_LOG(LogPhysicsControl, Warning,
-			TEXT("GetCachedBoneTransform - invalid bone name %s"), *BoneName.ToString());
+			TEXT("GetCachedBoneTransform - invalid bone name %s for %s"),
+			*BoneName.ToString(), GetOwner() ? *GetOwner()->GetName() : TEXT("Unknown"));
 	}
 	return FTransform();
 }
@@ -2620,7 +3555,8 @@ FVector UPhysicsControlComponent::GetCachedBonePosition(
 	if (bWarnAboutInvalidNames)
 	{
 		UE_LOG(LogPhysicsControl, Warning,
-			TEXT("GetCachedBonePosition - invalid bone name %s"), *BoneName.ToString());
+			TEXT("GetCachedBonePosition - invalid bone name %s for %s"),
+			*BoneName.ToString(), GetOwner() ? *GetOwner()->GetName() : TEXT("Unknown"));
 	}
 	return FVector::ZeroVector;
 }
@@ -2639,7 +3575,8 @@ FRotator UPhysicsControlComponent::GetCachedBoneOrientation(
 	if (bWarnAboutInvalidNames)
 	{
 		UE_LOG(LogPhysicsControl, Warning,
-			TEXT("GetCachedBoneOrientation - invalid bone name %s"), *BoneName.ToString());
+			TEXT("GetCachedBoneOrientation - invalid bone name %s for %s"),
+			*BoneName.ToString(), GetOwner() ? *GetOwner()->GetName() : TEXT("Unknown"));
 	}
 	return FRotator::ZeroRotator;
 }
@@ -2659,7 +3596,8 @@ bool UPhysicsControlComponent::SetCachedBoneData(
 	if (bWarnAboutInvalidNames)
 	{
 		UE_LOG(LogPhysicsControl, Warning,
-			TEXT("SetCachedBoneData - invalid bone name %s"), *BoneName.ToString());
+			TEXT("SetCachedBoneData - invalid bone name %s for %s"),
+			*BoneName.ToString(), GetOwner() ? *GetOwner()->GetName() : TEXT("Unknown"));
 	}
 	return false;
 }
@@ -2736,6 +3674,99 @@ bool UPhysicsControlComponent::GetBodyModifierExists(const FName Name) const
 	return FindBodyModifierRecord(Name) != nullptr;
 }
 
+#define PCC_LOG_LEVEL Display
+
+static FString GetComponentName(TWeakObjectPtr<UPrimitiveComponent> Component)
+{
+	if (Component.Get())
+	{
+		return Component.Get()->GetName();
+	}
+	return TEXT("None");
+}
+
+//======================================================================================================================
+void UPhysicsControlComponent::LogControls() const
+{
+#if !UE_BUILD_TEST && !UE_BUILD_SHIPPING
+
+	UE_LOG(LogPhysicsControl, PCC_LOG_LEVEL, TEXT("Controls for %s:"),
+		GetOwner() ? *GetOwner()->GetName() : TEXT("Unknown"));
+	for (const TPair<FName, FPhysicsControlRecord>& RecordPair : ControlRecords)
+	{
+		UE_LOG(LogPhysicsControl, PCC_LOG_LEVEL, TEXT("  %s:"), *RecordPair.Key.ToString());
+		const FPhysicsControlRecord& Record = RecordPair.Value;
+		UE_LOG(LogPhysicsControl, PCC_LOG_LEVEL, TEXT("    Parent: %s (%s) -> Child: %s (%s)"),
+			*GetComponentName(Record.ParentComponent),
+			*Record.PhysicsControl.ParentBoneName.ToString(), 
+			*GetComponentName(Record.ChildComponent),
+			*Record.PhysicsControl.ChildBoneName.ToString());
+		UE_LOG(LogPhysicsControl, PCC_LOG_LEVEL, TEXT("    Enabled %d"),
+			Record.PhysicsControl.IsEnabled());
+		UE_LOG(LogPhysicsControl, PCC_LOG_LEVEL, TEXT("    Linear: Strength %f DampingRatio %f ExtraDamping %f"),
+			Record.PhysicsControl.ControlData.LinearStrength,
+			Record.PhysicsControl.ControlData.LinearDampingRatio,
+			Record.PhysicsControl.ControlData.LinearExtraDamping);
+		UE_LOG(LogPhysicsControl, PCC_LOG_LEVEL, TEXT("    Angular: Strength %f DampingRatio %f ExtraDamping %f"),
+			Record.PhysicsControl.ControlData.AngularStrength,
+			Record.PhysicsControl.ControlData.AngularDampingRatio,
+			Record.PhysicsControl.ControlData.AngularExtraDamping);
+	}
+
+	UE_LOG(LogPhysicsControl, PCC_LOG_LEVEL, TEXT("Control sets:"));
+	for (const TMap<FName, TArray<FName>>::ElementType& Pair : NameRecords.ControlSets)
+	{
+		UE_LOG(LogPhysicsControl, PCC_LOG_LEVEL, TEXT("  %s:"), *Pair.Key.ToString());
+		const TArray<FName>& Names = Pair.Value;
+		for (FName Name : Names)
+		{
+			UE_LOG(LogPhysicsControl, PCC_LOG_LEVEL, TEXT("    %s"), *Name.ToString());
+		}
+	}
+#endif
+}
+
+//======================================================================================================================
+void UPhysicsControlComponent::LogBodyModifiers() const
+{
+#if !UE_BUILD_TEST && !UE_BUILD_SHIPPING
+
+	UE_LOG(LogPhysicsControl, PCC_LOG_LEVEL, TEXT("Body Modifiers for %s:"),
+		GetOwner() ? *GetOwner()->GetName() : TEXT("Unknown"));
+	for (const TPair<FName, FPhysicsBodyModifierRecord>& RecordPair : BodyModifierRecords)
+	{
+		UE_LOG(LogPhysicsControl, PCC_LOG_LEVEL, TEXT("  %s:"), *RecordPair.Key.ToString());
+		const FPhysicsBodyModifierRecord& Record = RecordPair.Value;
+		UE_LOG(LogPhysicsControl, PCC_LOG_LEVEL, TEXT("    Body: %s (%s)"),
+			*GetComponentName(Record.Component),
+			*Record.BodyModifier.BoneName.ToString());
+		UE_LOG(LogPhysicsControl, PCC_LOG_LEVEL, TEXT("    Movement: %s GravityMultiplier %f"),
+			*GetPhysicsMovementTypeName(Record.BodyModifier.ModifierData.MovementType).ToString(),
+			Record.BodyModifier.ModifierData.GravityMultiplier);
+	}
+
+	UE_LOG(LogPhysicsControl, PCC_LOG_LEVEL, TEXT("Body Modifier sets:"));
+	for (const TMap<FName, TArray<FName>>::ElementType& Pair : NameRecords.BodyModifierSets)
+	{
+		UE_LOG(LogPhysicsControl, PCC_LOG_LEVEL, TEXT("  %s:"), *Pair.Key.ToString());
+		const TArray<FName>& Names = Pair.Value;
+		for (FName Name : Names)
+		{
+			UE_LOG(LogPhysicsControl, PCC_LOG_LEVEL, TEXT("    %s"), *Name.ToString());
+		}
+	}
+#endif
+}
+
+#undef PCC_LOG_LEVEL
+
+//======================================================================================================================
+void UPhysicsControlComponent::LogControlsAndBodyModifiers() const
+{
+	LogControls();
+	LogBodyModifiers();
+}
+
 //======================================================================================================================
 bool UPhysicsControlComponent::ShouldCreatePhysicsState() const
 {
diff --git a/Engine/Plugins/Experimental/PhysicsControl/Source/PhysicsControl/Private/PhysicsControlComponentImpl.cpp b/Engine/Plugins/Experimental/PhysicsControl/Source/PhysicsControl/Private/PhysicsControlComponentImpl.cpp
--- a/Engine/Plugins/Experimental/PhysicsControl/Source/PhysicsControl/Private/PhysicsControlComponentImpl.cpp
+++ b/Engine/Plugins/Experimental/PhysicsControl/Source/PhysicsControl/Private/PhysicsControlComponentImpl.cpp
@@ -50,14 +50,18 @@ bool UPhysicsControlComponent::GetBoneData(
 	const USkeletalMeshComponent*                       InSkeletalMeshComponent,
 	const FName                                         InBoneName) const
 {
-	check(InSkeletalMeshComponent);
+	if (!ensure(InSkeletalMeshComponent))
+	{
+		return false;
+	}
 	OutPoseData = nullptr;
 	const FReferenceSkeleton& RefSkeleton = InSkeletalMeshComponent->GetSkeletalMeshAsset()->GetRefSkeleton();
 	const int32 BoneIndex = RefSkeleton.FindBoneIndex(InBoneName);
 
 	if (BoneIndex == INDEX_NONE)
 	{
-		UE_LOG(LogPhysicsControl, Warning, TEXT("Failed to find BoneIndex for %s"), *InBoneName.ToString());
+		UE_LOG(LogPhysicsControl, Warning, TEXT("Failed to find BoneIndex for %s in %s"), 
+			*InBoneName.ToString(), GetOwner() ? *GetOwner()->GetName() : TEXT("Unknown"));
 		return false;
 	}
 
@@ -72,10 +76,12 @@ bool UPhysicsControlComponent::GetBoneData(
 			OutPoseData = CachedSkeletalMeshData;
 			return true;
 		}
-		UE_LOG(LogPhysicsControl, Warning, TEXT("BoneIndex is out of range"));
+		UE_LOG(LogPhysicsControl, Warning, 
+			TEXT("BoneIndex is out of range in %s"), GetOwner() ? *GetOwner()->GetName() : TEXT("Unknown"));
 
 	}
-	UE_LOG(LogPhysicsControl, Warning, TEXT("Failed to find bone data for %s"), *InBoneName.ToString());
+	UE_LOG(LogPhysicsControl, Warning, TEXT("Failed to find bone data for %s in %s"),
+		*InBoneName.ToString(), GetOwner() ? *GetOwner()->GetName() : TEXT("Unknown"));
 	return false;
 }
 
@@ -85,13 +91,17 @@ bool UPhysicsControlComponent::GetModifiableBoneData(
 	const USkeletalMeshComponent*   InSkeletalMeshComponent,
 	const FName                     InBoneName)
 {
-	check(InSkeletalMeshComponent);
+	if (!ensure(InSkeletalMeshComponent))
+	{
+		return false;
+	}
 	const FReferenceSkeleton& RefSkeleton = InSkeletalMeshComponent->GetSkeletalMeshAsset()->GetRefSkeleton();
 	const int32 BoneIndex = RefSkeleton.FindBoneIndex(InBoneName);
 
 	if (BoneIndex == INDEX_NONE)
 	{
-		UE_LOG(LogPhysicsControl, Warning, TEXT("Failed to find BoneIndex for %s"), *InBoneName.ToString());
+		UE_LOG(LogPhysicsControl, Warning, TEXT("Failed to find BoneIndex for %s in %s"),
+			*InBoneName.ToString(), GetOwner() ? *GetOwner()->GetName() : TEXT("Unknown"));
 		return false;
 	}
 
@@ -106,14 +116,15 @@ bool UPhysicsControlComponent::GetModifiableBoneData(
 			OutBoneData = &CachedSkeletalMeshData->BoneDatas[BoneIndex];
 			return true;
 		}
-		UE_LOG(LogPhysicsControl, Warning, TEXT("BoneIndex is out of range"));
+		UE_LOG(LogPhysicsControl, Warning, TEXT("BoneIndex is out of range in %s"), 
+			GetOwner() ? *GetOwner()->GetName() : TEXT("Unknown"));
 
 	}
-	UE_LOG(LogPhysicsControl, Warning, TEXT("Failed to find bone data for %s"), *InBoneName.ToString());
+	UE_LOG(LogPhysicsControl, Warning, TEXT("Failed to find bone data for %s in %s"),
+		*InBoneName.ToString(), GetOwner() ? *GetOwner()->GetName() : TEXT("Unknown"));
 	return false;
 }
 
-
 //======================================================================================================================
 FPhysicsControlRecord* UPhysicsControlComponent::FindControlRecord(const FName Name)
 {
@@ -358,7 +369,8 @@ bool UPhysicsControlComponent::RemoveSkeletalMeshReferenceForCaching(
 	check(InSkeletalMeshComponent);
 	if (!InSkeletalMeshComponent)
 	{
-		UE_LOG(LogPhysicsControl, Warning, TEXT("Invalid skeletal mesh component"));
+		UE_LOG(LogPhysicsControl, Warning, TEXT("Invalid skeletal mesh component in %s"), 
+			GetOwner() ? *GetOwner()->GetName() : TEXT("Unknown"));
 		return false;
 	}
 
@@ -711,8 +723,9 @@ void UPhysicsControlComponent::ApplyBodyModifier(FPhysicsBodyModifierRecord& Rec
 			// Default means do nothing, so let's do exactly that
 			break;
 		default:
-			UE_LOG(LogPhysicsControl, Warning, TEXT("Invalid movement type %d"),
-				int(Record.BodyModifier.ModifierData.MovementType));
+			UE_LOG(LogPhysicsControl, Warning, TEXT("Invalid movement type %d in %s"),
+				int(Record.BodyModifier.ModifierData.MovementType),
+				GetOwner() ? *GetOwner()->GetName() : TEXT("Unknown"));
 			break;
 		}
 
@@ -788,8 +801,8 @@ bool UPhysicsControlComponent::DestroyControl(const FName Name, const EDestroyBe
 	}
 	if (bWarnAboutInvalidNames)
 	{
-		UE_LOG(LogPhysicsControl, Warning,
-			TEXT("DestroyControl - invalid name %s"), *Name.ToString());
+		UE_LOG(LogPhysicsControl, Warning, TEXT("DestroyControl - invalid control name %s for %s"),
+			*Name.ToString(), GetOwner() ? *GetOwner()->GetName() : TEXT("Unknown"));
 	}
 	return false;
 }
@@ -814,8 +827,8 @@ bool UPhysicsControlComponent::DestroyBodyModifier(const FName Name, const EDest
 	}
 	if (bWarnAboutInvalidNames)
 	{
-		UE_LOG(LogPhysicsControl, Warning,
-			TEXT("DestroyBodyModifier - invalid name %s"), *Name.ToString());
+		UE_LOG(LogPhysicsControl, Warning, TEXT("DestroyBodyModifier - invalid modifier name %s for %s"),
+			*Name.ToString(), GetOwner() ? *GetOwner()->GetName() : TEXT("Unknown"));
 	}
 	return false;
 }
diff --git a/Engine/Plugins/Experimental/PhysicsControl/Source/PhysicsControl/Private/PhysicsControlNameRecords.cpp b/Engine/Plugins/Experimental/PhysicsControl/Source/PhysicsControl/Private/PhysicsControlNameRecords.cpp
--- a/Engine/Plugins/Experimental/PhysicsControl/Source/PhysicsControl/Private/PhysicsControlNameRecords.cpp
+++ b/Engine/Plugins/Experimental/PhysicsControl/Source/PhysicsControl/Private/PhysicsControlNameRecords.cpp
@@ -4,18 +4,23 @@
 //======================================================================================================================
 void AddUniqueNamesToSet(const TArray<FName>& Names, const FName SetName, TMap<FName, TArray<FName>>& SetCollection)
 {
-	TArray<FName>& SetElementNames = SetCollection.FindOrAdd(SetName);
-
-	for (const FName Name : Names)
+	if (!SetName.IsNone())
 	{
-		SetElementNames.AddUnique(Name);
+		TArray<FName>& SetElementNames = SetCollection.FindOrAdd(SetName);
+		for (const FName Name : Names)
+		{
+			SetElementNames.AddUnique(Name);
+		}
 	}
 }
 
 //======================================================================================================================
 void FPhysicsControlNameRecords::AddControl(FName Name, FName SetName)
 {
-	ControlSets.FindOrAdd(SetName).AddUnique(Name);
+	if (!SetName.IsNone())
+	{
+		ControlSets.FindOrAdd(SetName).AddUnique(Name);
+	}
 	ControlSets.FindOrAdd(TEXT("All")).AddUnique(Name);
 }
 
@@ -24,7 +29,10 @@ void FPhysicsControlNameRecords::AddControl(FName Name, const TArray<FName>& Set
 {
 	for (const FName SetName : SetNames)
 	{
-		ControlSets.FindOrAdd(SetName).AddUnique(Name);
+		if (!SetName.IsNone())
+		{
+			ControlSets.FindOrAdd(SetName).AddUnique(Name);
+		}
 	}
 	ControlSets.FindOrAdd(TEXT("All")).AddUnique(Name);
 }
@@ -61,7 +69,10 @@ const TArray<FName>& FPhysicsControlNameRecords::GetControlNamesInSet(FName SetN
 //======================================================================================================================
 void FPhysicsControlNameRecords::AddBodyModifier(FName Name, FName SetName)
 {
-	BodyModifierSets.FindOrAdd(SetName).AddUnique(Name);
+	if (!SetName.IsNone())
+	{
+		BodyModifierSets.FindOrAdd(SetName).AddUnique(Name);
+	}
 	BodyModifierSets.FindOrAdd(TEXT("All")).AddUnique(Name);
 }
 
@@ -70,7 +81,10 @@ void FPhysicsControlNameRecords::AddBodyModifier(FName Name, const TArray<FName>
 {
 	for (FName SetName : SetNames)
 	{
-		BodyModifierSets.FindOrAdd(SetName).AddUnique(Name);
+		if (!SetName.IsNone())
+		{
+			BodyModifierSets.FindOrAdd(SetName).AddUnique(Name);
+		}
 	}
 	BodyModifierSets.FindOrAdd(TEXT("All")).AddUnique(Name);
 }
@@ -129,11 +143,9 @@ TArray<FName> ExpandName(const FName InName, const TMap<FName, TArray<FName>>& S
 TArray<FName> ExpandNames(const TArray<FName>& InNames, const TMap<FName, TArray<FName>>& SetNames)
 {
 	TArray<FName> OutputNames;
-
 	for (const FName Name : InNames)
 	{
 		OutputNames.Append(ExpandName(Name, SetNames));
 	}
-
 	return OutputNames;
 }
diff --git a/Engine/Plugins/Experimental/PhysicsControl/Source/PhysicsControl/Private/RigidBodyWithControl.cpp b/Engine/Plugins/Experimental/PhysicsControl/Source/PhysicsControl/Private/RigidBodyWithControl.cpp
--- a/Engine/Plugins/Experimental/PhysicsControl/Source/PhysicsControl/Private/RigidBodyWithControl.cpp
+++ b/Engine/Plugins/Experimental/PhysicsControl/Source/PhysicsControl/Private/RigidBodyWithControl.cpp
@@ -499,11 +499,11 @@ void FAnimNode_RigidBodyWithControl::LogControlsModifiersAndSets()
 #define RBWC_LOG_LEVEL Verbose
 
 	UE_LOG(LogPhysicsControl, RBWC_LOG_LEVEL, TEXT("Controls:"));
-	for (TMap<FName, FRigidBodyControlRecord>::ElementType& NameRecordPair : ControlRecords)
+	for (const TMap<FName, FRigidBodyControlRecord>::ElementType& NameRecordPair : ControlRecords)
 	{
 		UE_LOG(LogPhysicsControl, RBWC_LOG_LEVEL, TEXT("  %s:"), *NameRecordPair.Key.ToString());
 		const FRigidBodyControlRecord& Record = NameRecordPair.Value;
-		UE_LOG(LogPhysicsControl, RBWC_LOG_LEVEL, TEXT("    Parent bone: %s Child bone: %s"), 
+		UE_LOG(LogPhysicsControl, RBWC_LOG_LEVEL, TEXT("    Parent bone: %s -> Child bone: %s"), 
 			*Record.Control.ParentBoneName.ToString(), *Record.Control.ChildBoneName.ToString());
 		UE_LOG(LogPhysicsControl, RBWC_LOG_LEVEL, TEXT("    Enabled %d"),
 			Record.IsEnabled());
@@ -518,7 +518,7 @@ void FAnimNode_RigidBodyWithControl::LogControlsModifiersAndSets()
 	}
 
 	UE_LOG(LogPhysicsControl, RBWC_LOG_LEVEL, TEXT("Body Modifiers:"));
-	for (TMap<FName, FRigidBodyModifierRecord>::ElementType& NameRecordPair : ModifierRecords)
+	for (const TMap<FName, FRigidBodyModifierRecord>::ElementType& NameRecordPair : ModifierRecords)
 	{
 		UE_LOG(LogPhysicsControl, RBWC_LOG_LEVEL, TEXT("  %s:"), *NameRecordPair.Key.ToString());
 		const FRigidBodyModifierRecord& Record = NameRecordPair.Value;
@@ -530,24 +530,24 @@ void FAnimNode_RigidBodyWithControl::LogControlsModifiersAndSets()
 	}
 
 	UE_LOG(LogPhysicsControl, RBWC_LOG_LEVEL, TEXT("Control sets:"));
-	for (TMap<FName, TArray<FName>>::ElementType& Pair : NameRecords.ControlSets)
+	for (const TMap<FName, TArray<FName>>::ElementType& Pair : NameRecords.ControlSets)
 	{
 		UE_LOG(LogPhysicsControl, RBWC_LOG_LEVEL, TEXT("  %s:"), *Pair.Key.ToString());
 		const TArray<FName>& Names = Pair.Value;
 		for (FName Name : Names)
 		{
-			UE_LOG(LogPhysicsControl, RBWC_LOG_LEVEL, TEXT("    %s:"), *Name.ToString());
+			UE_LOG(LogPhysicsControl, RBWC_LOG_LEVEL, TEXT("    %s"), *Name.ToString());
 		}
 	}
 
 	UE_LOG(LogPhysicsControl, RBWC_LOG_LEVEL, TEXT("Body Modifier sets:"));
-	for (TMap<FName, TArray<FName>>::ElementType& Pair : NameRecords.BodyModifierSets)
+	for (const TMap<FName, TArray<FName>>::ElementType& Pair : NameRecords.BodyModifierSets)
 	{
 		UE_LOG(LogPhysicsControl, RBWC_LOG_LEVEL, TEXT("  %s:"), *Pair.Key.ToString());
 		const TArray<FName>& Names = Pair.Value;
 		for (FName Name : Names)
 		{
-			UE_LOG(LogPhysicsControl, RBWC_LOG_LEVEL, TEXT("    %s:"), *Name.ToString());
+			UE_LOG(LogPhysicsControl, RBWC_LOG_LEVEL, TEXT("    %s"), *Name.ToString());
 		}
 	}
 #undef RBWC_LOG_LEVEL
diff --git a/Engine/Plugins/Experimental/PhysicsControl/Source/PhysicsControl/Public/PhysicsControlComponent.h b/Engine/Plugins/Experimental/PhysicsControl/Source/PhysicsControl/Public/PhysicsControlComponent.h
--- a/Engine/Plugins/Experimental/PhysicsControl/Source/PhysicsControl/Public/PhysicsControlComponent.h
+++ b/Engine/Plugins/Experimental/PhysicsControl/Source/PhysicsControl/Public/PhysicsControlComponent.h
@@ -9,10 +9,11 @@
 #include "PhysicsControlRecord.h"
 #include "PhysicsControlPoseData.h"
 
-#include "UObject/ObjectMacros.h"
 #include "Components/SceneComponent.h"
 #include "EngineDefines.h"
+#include "Logging/LogVerbosity.h"
 #include "Templates/PimplPtr.h"
+#include "UObject/ObjectMacros.h"
 
 #include "PhysicsControlComponent.generated.h"
 
@@ -319,7 +320,7 @@ public:
 	 * @param bApplyToSetsWithName If true, then controls in sets with the name will be modified
 	 */
 	UFUNCTION(BlueprintCallable, Category = PhysicsControl, Meta = (ReturnDisplayName = "Success"))
-	UE_API void DestroyControl(
+	UE_API bool DestroyControl(
 		const FName Name, const bool bApplyToControlsWithName = true, const bool bApplyToSetsWithName = true);
 
 	/**
@@ -333,7 +334,7 @@ public:
 	 * @param bApplyToSetsWithName If true, then controls in sets with the name will be modified
 	 */
 	UFUNCTION(BlueprintCallable, Category = PhysicsControl)
-	UE_API void DestroyControls(
+	UE_API bool DestroyControls(
 		const TArray<FName>& Names, const bool bApplyToControlsWithName = true, const bool bApplyToSetsWithName = true);
 
 	/**
@@ -343,7 +344,7 @@ public:
 	 * "ParentSpace" and things like "WorldSpace_ArmLeft", depending on how controls have been created.
 	 */
 	UFUNCTION(BlueprintCallable, Category = PhysicsControl)
-	UE_API void DestroyControlsInSet(const FName Set);
+	UE_API bool DestroyControlsInSet(const FName Set);
 
 	/**
 	 * Updates the parent object part of a control or controls in a set. Note that this won't change 
@@ -355,7 +356,7 @@ public:
 	 * @param bApplyToSetsWithName If true, then controls in sets with the name will be modified
 	 */
 	UFUNCTION(BlueprintCallable, Category = PhysicsControl, Meta = (ReturnDisplayName = "Success"))
-	UE_API void SetControlParent(
+	UE_API bool SetControlParent(
 		const FName Name, UPrimitiveComponent* ParentComponent, const FName ParentBoneName,
 		const bool bApplyToControlsWithName = true, const bool bApplyToSetsWithName = true);
 
@@ -369,12 +370,12 @@ public:
 	 * @param bApplyToSetsWithName If true, then controls in sets with the name will be modified
 	 */
 	UFUNCTION(BlueprintCallable, Category = PhysicsControl)
-	UE_API void SetControlParents(
+	UE_API bool SetControlParents(
 		const TArray<FName>& Names, UPrimitiveComponent* ParentComponent, const FName ParentBoneName,
 		const bool bApplyToControlsWithName = true, const bool bApplyToSetsWithName = true);
 
 	UFUNCTION(BlueprintCallable, Category = PhysicsControl)
-	UE_API void SetControlParentsInSet(
+	UE_API bool SetControlParentsInSet(
 		const FName          Set,
 		UPrimitiveComponent* ParentComponent,
 		const FName          ParentBoneName);
@@ -391,7 +392,7 @@ public:
 	 * @param bApplyToSetsWithName If true, then controls in sets with the name will be modified
 	 */
 	UFUNCTION(BlueprintCallable, Category = PhysicsControl, Meta = (ReturnDisplayName = "Success"))
-	UE_API void SetControlData(
+	UE_API bool SetControlData(
 		const FName Name, const FPhysicsControlData ControlData, 
 		const bool bApplyToControlsWithName = true, const bool bApplyToSetsWithName = true);
 
@@ -408,7 +409,7 @@ public:
 	 * @param bApplyToSetsWithName If true, then controls in sets with the name will be modified
 	 */
 	UFUNCTION(BlueprintCallable, Category = PhysicsControl)
-	UE_API void SetControlDatas(
+	UE_API bool SetControlDatas(
 		const TArray<FName>& Names, const FPhysicsControlData ControlData,
 		const bool bApplyToControlsWithName = true, const bool bApplyToSetsWithName = true);
 
@@ -420,7 +421,7 @@ public:
 	 * @param ControlData The new control data
 	 */
 	UFUNCTION(BlueprintCallable, Category = PhysicsControl)
-	UE_API void SetControlDatasInSet(const FName Set, const FPhysicsControlData ControlData);
+	UE_API bool SetControlDatasInSet(const FName Set, const FPhysicsControlData ControlData);
 
 	/**
 	 * Modifies an existing control data in a control or controls in a set - i.e. the strengths 
@@ -434,7 +435,7 @@ public:
 	 * @param bApplyToSetsWithName If true, then controls in sets with the name will be modified
 	 */
 	UFUNCTION(BlueprintCallable, Category = PhysicsControl, Meta = (ReturnDisplayName = "Success"))
-	UE_API void SetControlSparseData(
+	UE_API bool SetControlSparseData(
 		const FName Name, const FPhysicsControlSparseData ControlData,
 		const bool bApplyToControlsWithName = true, const bool bApplyToSetsWithName = true);
 
@@ -451,7 +452,7 @@ public:
 	 * @param bApplyToSetsWithName If true, then controls in sets with the name will be modified
 	 */
 	UFUNCTION(BlueprintCallable, Category = PhysicsControl)
-	UE_API void SetControlSparseDatas(
+	UE_API bool SetControlSparseDatas(
 		const TArray<FName>& Names, const FPhysicsControlSparseData ControlData,
 		const bool bApplyToControlsWithName = true, const bool bApplyToSetsWithName = true);
 
@@ -463,7 +464,7 @@ public:
 	 * @param ControlData The new control data
 	 */
 	UFUNCTION(BlueprintCallable, Category = PhysicsControl)
-	UE_API void SetControlSparseDatasInSet(const FName Set, const FPhysicsControlSparseData ControlData);
+	UE_API bool SetControlSparseDatasInSet(const FName Set, const FPhysicsControlSparseData ControlData);
 
 	/**
 	 * Modifies existing control data in a control or controls in a set using the multipliers
@@ -477,7 +478,7 @@ public:
 	 * @param bApplyToSetsWithName If true, then controls in sets with the name will be modified
 	 */
 	UFUNCTION(BlueprintCallable, Category = PhysicsControl, Meta = (ReturnDisplayName = "Success"))
-	UE_API void SetControlMultiplier(
+	UE_API bool SetControlMultiplier(
 		const FName Name, const FPhysicsControlMultiplier ControlMultiplier, const bool bEnableControl = true,
 		const bool bApplyToControlsWithName = true, const bool bApplyToSetsWithName = true);
 
@@ -494,7 +495,7 @@ public:
 	 * @param bApplyToSetsWithName If true, then controls in sets with the name will be modified
 	 */
 	UFUNCTION(BlueprintCallable, Category = PhysicsControl)
-	UE_API void SetControlMultipliers(
+	UE_API bool SetControlMultipliers(
 		const TArray<FName>& Names, const FPhysicsControlMultiplier ControlMultiplier, const bool bEnableControl = true,
 		const bool bApplyToControlsWithName = true, const bool bApplyToSetsWithName = true);
 
@@ -507,7 +508,7 @@ public:
 	 * @param bEnableControl Enables the controls if currently disabled
 	 */
 	UFUNCTION(BlueprintCallable, Category = PhysicsControl)
-	UE_API void SetControlMultipliersInSet(
+	UE_API bool SetControlMultipliersInSet(
 		const FName Set, const FPhysicsControlMultiplier  ControlMultiplier, const bool bEnableControl = true);
 
 	/**
@@ -522,7 +523,7 @@ public:
 	 * @param bApplyToSetsWithName If true, then controls in sets with the name will be modified
 	 */
 	UFUNCTION(BlueprintCallable, Category = PhysicsControl, Meta = (ReturnDisplayName = "Success"))
-	UE_API void SetControlSparseMultiplier(
+	UE_API bool SetControlSparseMultiplier(
 		const FName Name, const FPhysicsControlSparseMultiplier ControlMultiplier, const bool bEnableControl = true,
 		const bool bApplyToControlsWithName = true, const bool bApplyToSetsWithName = true);
 
@@ -539,7 +540,7 @@ public:
 	 * @param bApplyToSetsWithName If true, then controls in sets with the name will be modified
 	 */
 	UFUNCTION(BlueprintCallable, Category = PhysicsControl)
-	UE_API void SetControlSparseMultipliers(
+	UE_API bool SetControlSparseMultipliers(
 		const TArray<FName>& Names, const FPhysicsControlSparseMultiplier ControlMultiplier, const bool bEnableControl = true,
 		const bool bApplyToControlsWithName = true, const bool bApplyToSetsWithName = true);
 
@@ -552,7 +553,7 @@ public:
 	 * @param bEnableControl Enables the controls if currently disabled
 	 */
 	UFUNCTION(BlueprintCallable, Category = PhysicsControl)
-	UE_API void SetControlSparseMultipliersInSet(
+	UE_API bool SetControlSparseMultipliersInSet(
 		const FName                           Set, 
 		const FPhysicsControlSparseMultiplier ControlMultiplier, 
 		const bool                            bEnableControl = true);
@@ -573,7 +574,7 @@ public:
 	 * @param bApplyToSetsWithName If true, then controls in sets with the name will be modified
 	 */
 	UFUNCTION(BlueprintCallable, Category = PhysicsControl, Meta = (ReturnDisplayName = "Success"))
-	UE_API void SetControlLinearData(
+	UE_API bool SetControlLinearData(
 		const FName Name,
 		const float Strength = 1.0f, 
 		const float DampingRatio = 1.0f, 
@@ -599,7 +600,7 @@ public:
 	 * @param bApplyToSetsWithName If true, then controls in sets with the name will be modified
 	 */
 	UFUNCTION(BlueprintCallable, Category = PhysicsControl, Meta = (ReturnDisplayName = "Success"))
-	UE_API void SetControlAngularData(
+	UE_API bool SetControlAngularData(
 		const FName Name,
 		const float Strength = 1.0f, 
 		const float DampingRatio = 1.0f, 
@@ -640,7 +641,7 @@ public:
 	 * @param bApplyToSetsWithName If true, then controls in sets with the name will be modified
 	 */
 	UFUNCTION(BlueprintCallable, Category = PhysicsControl, Meta = (ReturnDisplayName = "Success"))
-	UE_API void SetControlTarget(
+	UE_API bool SetControlTarget(
 		const FName Name, const FPhysicsControlTarget ControlTarget, const bool bEnableControl = true,
 		const bool bApplyToControlsWithName = true, const bool bApplyToSetsWithName = true);
 
@@ -659,7 +660,7 @@ public:
 	 * @param bApplyToSetsWithName If true, then controls in sets with the name will be modified
 	 */
 	UFUNCTION(BlueprintCallable, Category = PhysicsControl)
-	UE_API void SetControlTargets(
+	UE_API bool SetControlTargets(
 		const TArray<FName>& Names, const FPhysicsControlTarget ControlTarget, const bool bEnableControl = true,
 		const bool bApplyToControlsWithName = true, const bool bApplyToSetsWithName = true);
 
@@ -672,7 +673,7 @@ public:
 	 * @param bEnableControl Enables the controls if currently disabled
 	 */
 	UFUNCTION(BlueprintCallable, Category = PhysicsControl)
-	UE_API void SetControlTargetsInSet(
+	UE_API bool SetControlTargetsInSet(
 		const FName                 Set, 
 		const FPhysicsControlTarget ControlTarget, 
 		const bool                  bEnableControl = true);
@@ -697,9 +698,10 @@ public:
 	 *        transform.
 	 * @param bApplyToControlsWithName If true, then controls with the name will be modified
 	 * @param bApplyToSetsWithName If true, then controls in sets with the name will be modified
+	 * @return Returns true if the control/set could be found
 	 */
 	UFUNCTION(BlueprintCallable, Category = PhysicsControl, Meta = (ReturnDisplayName = "Success"))
-	UE_API void SetControlTargetPositionAndOrientation(
+	UE_API bool SetControlTargetPositionAndOrientation(
 		const FName    Name,
 		const FVector  Position,
 		const FRotator Orientation,
@@ -729,9 +731,10 @@ public:
 	 *        transform.
 	 * @param bApplyToControlsWithName If true, then controls with the name will be modified
 	 * @param bApplyToSetsWithName If true, then controls in sets with the name will be modified
+	 * @return Returns true if all the control names were found
 	 */
 	UFUNCTION(BlueprintCallable, Category = PhysicsControl)
-	UE_API void SetControlTargetPositionsAndOrientations(
+	UE_API bool SetControlTargetPositionsAndOrientations(
 		const TArray<FName>& Names,
 		const FVector        Position,
 		const FRotator       Orientation,
@@ -743,9 +746,10 @@ public:
 
 	/**
 	 * Calls SetControlTargetPositionAndOrientation for each control in the set
+	 * @return Returns true if the set was be found
 	 */
 	UFUNCTION(BlueprintCallable, Category = PhysicsControl)
-	UE_API void SetControlTargetPositionsAndOrientationsInSet(
+	UE_API bool SetControlTargetPositionsAndOrientationsInSet(
 		const FName          SetName,
 		const FVector        Position,
 		const FRotator       Orientation,
@@ -773,9 +777,10 @@ public:
 	 *        transform.
 	 * @param bApplyToControlsWithName If true, then controls with the name will be modified
 	 * @param bApplyToSetsWithName If true, then controls in sets with the name will be modified
+	 * @return Returns true if the control/set can be found
 	 */
 	UFUNCTION(BlueprintCallable, Category = PhysicsControl, Meta = (ReturnDisplayName = "Success"))
-	UE_API void SetControlTargetPosition(
+	UE_API bool SetControlTargetPosition(
 		const FName   Name,
 		const FVector Position, 
 		const float   VelocityDeltaTime, 
@@ -803,9 +808,10 @@ public:
 	 *        transform.
 	 * @param bApplyToControlsWithName If true, then controls with the name will be modified
 	 * @param bApplyToSetsWithName If true, then controls in sets with the name will be modified
+	 * @return Returns true if all the controls/sets can be found
 	 */
 	UFUNCTION(BlueprintCallable, Category = PhysicsControl)
-	UE_API void SetControlTargetPositions(
+	UE_API bool SetControlTargetPositions(
 		const TArray<FName>& Names,
 		const FVector        Position, 
 		const float          VelocityDeltaTime, 
@@ -816,9 +822,10 @@ public:
 
 	/**
 	 * Calls SetControlTargetPosition for each of the controls in the set
+	 * @return Returns true if the set name and its controls can be found
 	 */
 	UFUNCTION(BlueprintCallable, Category = PhysicsControl)
-	UE_API void SetControlTargetPositionsInSet(
+	UE_API bool SetControlTargetPositionsInSet(
 		const FName   SetName,
 		const FVector Position, 
 		const float   VelocityDeltaTime, 
@@ -845,9 +852,10 @@ public:
 	 *        transform.
 	 * @param bApplyToControlsWithName If true, then controls with the name will be modified
 	 * @param bApplyToSetsWithName If true, then controls in sets with the name will be modified
+	 * @return Returns true if the control/set can be found
 	 */
 	UFUNCTION(BlueprintCallable, Category = PhysicsControl, Meta = (ReturnDisplayName = "Success"))
-	UE_API void SetControlTargetOrientation(
+	UE_API bool SetControlTargetOrientation(
 		const FName    Name,
 		const FRotator Orientation, 
 		const float    AngularVelocityDeltaTime, 
@@ -876,9 +884,10 @@ public:
 	 *        transform.
 	 * @param bApplyToControlsWithName If true, then controls with the name will be modified
 	 * @param bApplyToSetsWithName If true, then controls in sets with the name will be modified
+	 * @return Returns true if all the controls can be found
 	 */
 	UFUNCTION(BlueprintCallable, Category = PhysicsControl)
-	UE_API void SetControlTargetOrientations(
+	UE_API bool SetControlTargetOrientations(
 		const TArray<FName>& Names,
 		const FRotator Orientation,
 		const float    AngularVelocityDeltaTime, 
@@ -889,9 +898,10 @@ public:
 
 	/**
 	 * Calls SetControlTargetOrientation for each of the controls in the set
+	 * @return Returns true if the set and all the controls can be found
 	 */
 	UFUNCTION(BlueprintCallable, Category = PhysicsControl)
-	UE_API void SetControlTargetOrientationsInSet(
+	UE_API bool SetControlTargetOrientationsInSet(
 		const FName    SetName,
 		const FRotator Orientation, 
 		const float    AngularVelocityDeltaTime, 
@@ -978,7 +988,7 @@ public:
 	 * @param bApplyToSetsWithName If true, then controls in sets with the name will be modified
 	 */
 	UFUNCTION(BlueprintCallable, Category = PhysicsControl, Meta = (ReturnDisplayName = "Success"))
-	UE_API void SetControlUseSkeletalAnimation(
+	UE_API bool SetControlUseSkeletalAnimation(
 		const FName Name,
 		const bool  bUseSkeletalAnimation = true,
 		const float AngularTargetVelocityMultiplier = 1.0f,
@@ -1003,7 +1013,7 @@ public:
 	 * @param bApplyToSetsWithName If true, then controls in sets with the name will be modified
 	 */
 	UFUNCTION(BlueprintCallable, Category = PhysicsControl)
-	UE_API void SetControlsUseSkeletalAnimation(
+	UE_API bool SetControlsUseSkeletalAnimation(
 		const TArray<FName>& Names,
 		const bool           bUseSkeletalAnimation = true,
 		const float          AngularTargetVelocityMultiplier = 1.0f,
@@ -1024,7 +1034,7 @@ public:
 	 *        linear velocity extracted from the animation that is used as targets for the controls
 	 */
 	UFUNCTION(BlueprintCallable, Category = PhysicsControl)
-	UE_API void SetControlsInSetUseSkeletalAnimation(
+	UE_API bool SetControlsInSetUseSkeletalAnimation(
 		const FName Set,
 		const bool  bUseSkeletalAnimation = true,
 		const float AngularTargetVelocityMultiplier = 1.0f,
@@ -1041,7 +1051,7 @@ public:
 	 * @param bApplyToSetsWithName If true, then controls in sets with the name will be modified
 	 */
 	UFUNCTION(BlueprintCallable, Category = PhysicsControl, Meta = (ReturnDisplayName = "Success"))
-	UE_API void SetControlEnabled(
+	UE_API bool SetControlEnabled(
 		const FName Name, const bool bEnable = true, 
 		const bool bApplyToControlsWithName = true, const bool bApplyToSetsWithName = true);
 
@@ -1057,7 +1067,7 @@ public:
 	 * @param bApplyToSetsWithName If true, then controls in sets with the name will be modified
 	 */
 	UFUNCTION(BlueprintCallable, Category = PhysicsControl)
-	UE_API void SetControlsEnabled(
+	UE_API bool SetControlsEnabled(
 		const TArray<FName>& Names, const bool bEnable = true, 
 		const bool bApplyToControlsWithName = true, const bool bApplyToSetsWithName = true);
 
@@ -1070,7 +1080,7 @@ public:
 	 */
 
 	UFUNCTION(BlueprintCallable, Category = PhysicsControl)
-	UE_API void SetControlsInSetEnabled(const FName Set, const bool bEnable = true);
+	UE_API bool SetControlsInSetEnabled(const FName Set, const bool bEnable = true);
 
 	/**
 	 * @param Name The name of the control or set of controls to update. Standard sets will include "All", "WorldSpace",
@@ -1080,7 +1090,7 @@ public:
 	 * @param bApplyToSetsWithName If true, then controls in sets with the name will be modified
 	 */
 	UFUNCTION(BlueprintCallable, Category = PhysicsControl, Meta = (ReturnDisplayName = "Success"))
-	UE_API void SetControlDisableCollision(
+	UE_API bool SetControlDisableCollision(
 		const FName Name, const bool bDisableCollision,
 		const bool bApplyToControlsWithName = true, const bool bApplyToSetsWithName = true);
 
@@ -1092,7 +1102,7 @@ public:
 	 * @param bApplyToSetsWithName If true, then controls in sets with the name will be modified
 	 */
 	UFUNCTION(BlueprintCallable, Category = PhysicsControl)
-	UE_API void SetControlsDisableCollision(
+	UE_API bool SetControlsDisableCollision(
 		const TArray<FName>& Names, const bool bDisableCollision,
 		const bool bApplyToControlsWithName = true, const bool bApplyToSetsWithName = true);
 
@@ -1102,7 +1112,7 @@ public:
 	 * @param bDisableCollision If set then the control will disable collision between the bodies it connects.
 	 */
 	UFUNCTION(BlueprintCallable, Category = PhysicsControl)
-	UE_API void SetControlsInSetDisableCollision(const FName Set, const bool bDisableCollision);
+	UE_API bool SetControlsInSetDisableCollision(const FName Set, const bool bDisableCollision);
 
 	/**
 	 * @param Name     The name of the control to access. 
@@ -1210,9 +1220,10 @@ public:
 	 * @param Name The name of the body modifier or set of body modifiers to destroy. 
 	 * @param bApplyToModifiersWithName If true, then body modifiers with the name will be modified
 	 * @param bApplyToSetsWithName If true, then body modifiers in sets with the name will be modified
+	 * @return Returns true if the body modifier or set can found/destroyed
 	 */
 	UFUNCTION(BlueprintCallable, Category = PhysicsControl, Meta = (ReturnDisplayName = "Success"))
-	UE_API void DestroyBodyModifier(
+	UE_API bool DestroyBodyModifier(
 		const FName Name, const bool bApplyToModifiersWithName = true, const bool bApplyToSetsWithName = true);
 
 	/**
@@ -1223,9 +1234,10 @@ public:
 	 *              then it can be split.
 	 * @param bApplyToModifiersWithName If true, then body modifiers with the name will be modified
 	 * @param bApplyToSetsWithName If true, then body modifiers in sets with the name will be modified
+	 * @return Returns true if all the body modifiers can be found/destroyed
 	 */
 	UFUNCTION(BlueprintCallable, Category = PhysicsControl)
-	UE_API void DestroyBodyModifiers(
+	UE_API bool DestroyBodyModifiers(
 		const TArray<FName>& Names, const bool bApplyToModifiersWithName = true, const bool bApplyToSetsWithName = true);
 
 	/**
@@ -1233,9 +1245,10 @@ public:
 	 *
 	 * @param Set The set of body modifiers to destroy. Standard sets will include "All" and things like
 	 *        "ArmLeft", depending on how body modifiers have been created.
+	 * @return Returns true if the set and all its modifiers can be found/destroyed
 	 */
 	UFUNCTION(BlueprintCallable, Category = PhysicsControl)
-	UE_API void DestroyBodyModifiersInSet(const FName Set);
+	UE_API bool DestroyBodyModifiersInSet(const FName Set);
 
 	/**
 	 * Modifies an existing Body Modifier Data in a body modifier or body modifiers in a set
@@ -1247,7 +1260,7 @@ public:
 	 * @param bApplyToSetsWithName If true, then body modifiers in sets with the name will be modified
 	 */
 	UFUNCTION(BlueprintCallable, Category = PhysicsControl, Meta = (ReturnDisplayName = "Success"))
-	UE_API void SetBodyModifierData(
+	UE_API bool SetBodyModifierData(
 		const FName Name, const FPhysicsControlModifierData ModifierData, 
 		const bool bApplyToModifiersWithName = true, const bool bApplyToSetsWithName = true);
 
@@ -1262,7 +1275,7 @@ public:
 	 * @param bApplyToSetsWithName If true, then body modifiers in sets with the name will be modified
 	 */
 	UFUNCTION(BlueprintCallable, Category = PhysicsControl)
-	UE_API void SetBodyModifierDatas(
+	UE_API bool SetBodyModifierDatas(
 		const TArray<FName>& Names, const FPhysicsControlModifierData ModifierData,
 		const bool bApplyToModifiersWithName = true, const bool bApplyToSetsWithName = true);
 
@@ -1274,7 +1287,7 @@ public:
 	 * @param ModifierData The new data
 	 */
 	UFUNCTION(BlueprintCallable, Category = PhysicsControl)
-	UE_API void SetBodyModifierDatasInSet(const FName Set, const FPhysicsControlModifierData ModifierData);
+	UE_API bool SetBodyModifierDatasInSet(const FName Set, const FPhysicsControlModifierData ModifierData);
 
 	/**
 	 * Modifies an existing Body Modifier Data in a body modifier or body modifiers in a set
@@ -1285,7 +1298,7 @@ public:
 	 * @param bApplyToSetsWithName If true, then body modifiers in sets with the name will be modified
 	 */
 	UFUNCTION(BlueprintCallable, Category = PhysicsControl, Meta = (ReturnDisplayName = "Success"))
-	UE_API void SetBodyModifierSparseData(
+	UE_API bool SetBodyModifierSparseData(
 		const FName Name, const FPhysicsControlModifierSparseData ModifierData,
 		const bool bApplyToModifiersWithName = true, const bool bApplyToSetsWithName = true);
 
@@ -1300,7 +1313,7 @@ public:
 	 * @param bApplyToSetsWithName If true, then body modifiers in sets with the name will be modified
 	 */
 	UFUNCTION(BlueprintCallable, Category = PhysicsControl)
-	UE_API void SetBodyModifierSparseDatas(
+	UE_API bool SetBodyModifierSparseDatas(
 		const TArray<FName>& Names, const FPhysicsControlModifierSparseData ModifierData,
 		const bool bApplyToModifiersWithName = true, const bool bApplyToSetsWithName = true);
 
@@ -1312,7 +1325,7 @@ public:
 	 * @param ModifierData The new data
 	 */
 	UFUNCTION(BlueprintCallable, Category = PhysicsControl)
-	UE_API void SetBodyModifierSparseDatasInSet(const FName Set, const FPhysicsControlModifierSparseData ModifierData);
+	UE_API bool SetBodyModifierSparseDatasInSet(const FName Set, const FPhysicsControlModifierSparseData ModifierData);
 
 	/**
 	 * Sets the kinematic target transform for a body modifier. 
@@ -1341,7 +1354,7 @@ public:
 	 * @param bApplyToSetsWithName If true, then body modifiers in sets with the name will be modified
 	 */
 	UFUNCTION(BlueprintCallable, Category = PhysicsControl, Meta = (ReturnDisplayName = "Success"))
-	UE_API void SetBodyModifierMovementType(
+	UE_API bool SetBodyModifierMovementType(
 		const FName Name, const EPhysicsMovementType MovementType = EPhysicsMovementType::Simulated,
 		const bool bApplyToModifiersWithName = true, const bool bApplyToSetsWithName = true);
 
@@ -1355,7 +1368,7 @@ public:
 	 * @param MovementType Whether to enable/disable simulation on the bodies
 	 */
 	UFUNCTION(BlueprintCallable, Category = PhysicsControl)
-	UE_API void SetBodyModifiersMovementType(
+	UE_API bool SetBodyModifiersMovementType(
 		const TArray<FName>& Names, const EPhysicsMovementType MovementType = EPhysicsMovementType::Simulated,
 		const bool bApplyToModifiersWithName = true, const bool bApplyToSetsWithName = true);
 
@@ -1367,7 +1380,7 @@ public:
 	 * @param MovementType Whether to enable/disable simulation on the bodies
 	 */
 	UFUNCTION(BlueprintCallable, Category = PhysicsControl)
-	UE_API void SetBodyModifiersInSetMovementType(
+	UE_API bool SetBodyModifiersInSetMovementType(
 		const FName                Set,
 		const EPhysicsMovementType MovementType = EPhysicsMovementType::Simulated);
 
@@ -1381,7 +1394,7 @@ public:
 	 * @param bApplyToSetsWithName If true, then body modifiers in sets with the name will be modified
 	 */
 	UFUNCTION(BlueprintCallable, Category = PhysicsControl, Meta = (ReturnDisplayName = "Success"))
-	UE_API void SetBodyModifierCollisionType(
+	UE_API bool SetBodyModifierCollisionType(
 		const FName Name, const ECollisionEnabled::Type CollisionType = ECollisionEnabled::QueryAndPhysics,
 		const bool bApplyToModifiersWithName = true, const bool bApplyToSetsWithName = true);
 
@@ -1396,7 +1409,7 @@ public:
 	 * @param bApplyToSetsWithName If true, then body modifiers in sets with the name will be modified
 	 */
 	UFUNCTION(BlueprintCallable, Category = PhysicsControl)
-	UE_API void SetBodyModifiersCollisionType(
+	UE_API bool SetBodyModifiersCollisionType(
 		const TArray<FName>& Names, const ECollisionEnabled::Type CollisionType = ECollisionEnabled::QueryAndPhysics,
 		const bool bApplyToModifiersWithName = true, const bool bApplyToSetsWithName = true);
 
@@ -1408,7 +1421,7 @@ public:
 	 * @param CollisionType Collision type to set on the bodies
 	 */
 	UFUNCTION(BlueprintCallable, Category = PhysicsControl)
-	UE_API void SetBodyModifiersInSetCollisionType(
+	UE_API bool SetBodyModifiersInSetCollisionType(
 		const FName                   Set,
 		const ECollisionEnabled::Type CollisionType = ECollisionEnabled::QueryAndPhysics);
 
@@ -1422,7 +1435,7 @@ public:
 	 * @param bApplyToSetsWithName If true, then body modifiers in sets with the name will be modified
 	 */
 	UFUNCTION(BlueprintCallable, Category = PhysicsControl, Meta = (ReturnDisplayName = "Success"))
-	UE_API void SetBodyModifierGravityMultiplier(
+	UE_API bool SetBodyModifierGravityMultiplier(
 		const FName Name, const float GravityMultiplier = 1.0f,
 		const bool bApplyToModifiersWithName = true, const bool bApplyToSetsWithName = true);
 
@@ -1437,7 +1450,7 @@ public:
 	 * @param bApplyToSetsWithName If true, then body modifiers in sets with the name will be modified
 	 */
 	UFUNCTION(BlueprintCallable, Category = PhysicsControl)
-	UE_API void SetBodyModifiersGravityMultiplier(
+	UE_API bool SetBodyModifiersGravityMultiplier(
 		const TArray<FName>& Names, const float GravityMultiplier = 1.0f,
 		const bool bApplyToModifiersWithName = true, const bool bApplyToSetsWithName = true);
 
@@ -1449,7 +1462,7 @@ public:
 	 * @param GravityMultiplier The amount of gravity to apply when simulating
 	 */
 	UFUNCTION(BlueprintCallable, Category = PhysicsControl)
-	UE_API void SetBodyModifiersInSetGravityMultiplier(
+	UE_API bool SetBodyModifiersInSetGravityMultiplier(
 		const FName Set,
 		const float GravityMultiplier = 1.0f);
 
@@ -1464,7 +1477,7 @@ public:
 	 * @param bApplyToSetsWithName If true, then body modifiers in sets with the name will be modified
 	 */
 	UFUNCTION(BlueprintCallable, Category = PhysicsControl, Meta = (ReturnDisplayName = "Success"))
-	UE_API void SetBodyModifierPhysicsBlendWeight(
+	UE_API bool SetBodyModifierPhysicsBlendWeight(
 		const FName Name, const float PhysicsBlendWeight = 1.0f,
 		const bool bApplyToModifiersWithName = true, const bool bApplyToSetsWithName = true);
 
@@ -1480,7 +1493,7 @@ public:
 	 * @param bApplyToSetsWithName If true, then body modifiers in sets with the name will be modified
 	 */
 	UFUNCTION(BlueprintCallable, Category = PhysicsControl)
-	UE_API void SetBodyModifiersPhysicsBlendWeight(
+	UE_API bool SetBodyModifiersPhysicsBlendWeight(
 		const TArray<FName>& Names, const float PhysicsBlendWeight = 1.0f,
 		const bool bApplyToModifiersWithName = true, const bool bApplyToSetsWithName = true);
 
@@ -1493,7 +1506,7 @@ public:
 	 *        animation and that coming from simulation.
 	 */
 	UFUNCTION(BlueprintCallable, Category = PhysicsControl)
-	UE_API void SetBodyModifiersInSetPhysicsBlendWeight(
+	UE_API bool SetBodyModifiersInSetPhysicsBlendWeight(
 		const FName Set,
 		const float PhysicsBlendWeight = 1.0f);
 
@@ -1507,7 +1520,7 @@ public:
 	 * @param bApplyToSetsWithName If true, then body modifiers in sets with the name will be modified
 	 */
 	UFUNCTION(BlueprintCallable, Category = PhysicsControl, Meta = (ReturnDisplayName = "Success"))
-	UE_API void SetBodyModifierKinematicTargetSpace(
+	UE_API bool SetBodyModifierKinematicTargetSpace(
 		const FName Name, const EPhysicsControlKinematicTargetSpace KinematicTargetSpace,
 		const bool bApplyToModifiersWithName = true, const bool bApplyToSetsWithName = true);
 
@@ -1522,7 +1535,7 @@ public:
 	 * @param bApplyToSetsWithName If true, then body modifiers in sets with the name will be modified
 	 */
 	UFUNCTION(BlueprintCallable, Category = PhysicsControl)
-	UE_API void SetBodyModifiersKinematicTargetSpace(
+	UE_API bool SetBodyModifiersKinematicTargetSpace(
 		const TArray<FName>& Names, const EPhysicsControlKinematicTargetSpace KinematicTargetSpace,
 		const bool bApplyToModifiersWithName = true, const bool bApplyToSetsWithName = true);
 
@@ -1534,7 +1547,7 @@ public:
 	 * @param KinematicTargetSpace The space in which the kinematic target should be applied.
 	 */
 	UFUNCTION(BlueprintCallable, Category = PhysicsControl)
-	UE_API void SetBodyModifiersInSetKinematicTargetSpace(
+	UE_API bool SetBodyModifiersInSetKinematicTargetSpace(
 		const FName                               Set,
 		const EPhysicsControlKinematicTargetSpace KinematicTargetSpace);
 
@@ -1551,7 +1564,7 @@ public:
 	 * @param bApplyToSetsWithName If true, then body modifiers in sets with the name will be modified
 	 */
 	UFUNCTION(BlueprintCallable, Category = PhysicsControl, Meta = (ReturnDisplayName = "Success"))
-	UE_API void SetBodyModifierUpdateKinematicFromSimulation(
+	UE_API bool SetBodyModifierUpdateKinematicFromSimulation(
 		const FName Name, const bool  bUpdateKinematicFromSimulation,
 		const bool bApplyToModifiersWithName = true, const bool bApplyToSetsWithName = true);
 
@@ -1569,7 +1582,7 @@ public:
 	 * @param bApplyToSetsWithName If true, then body modifiers in sets with the name will be modified
 	 */
 	UFUNCTION(BlueprintCallable, Category = PhysicsControl)
-	UE_API void SetBodyModifiersUpdateKinematicFromSimulation(
+	UE_API bool SetBodyModifiersUpdateKinematicFromSimulation(
 		const TArray<FName>& Names, const bool bUpdateKinematicFromSimulation,
 		const bool bApplyToModifiersWithName = true, const bool bApplyToSetsWithName = true);
 
@@ -1584,7 +1597,7 @@ public:
 	 *        dynamic ones.
 	 */
 	UFUNCTION(BlueprintCallable, Category = PhysicsControl)
-	UE_API void SetBodyModifiersInSetUpdateKinematicFromSimulation(
+	UE_API bool SetBodyModifiersInSetUpdateKinematicFromSimulation(
 		const FName Set,
 		const bool  bUpdateKinematicFromSimulation);
 
@@ -1860,6 +1873,18 @@ public:
 	UFUNCTION(BlueprintCallable, Category = PhysicsControl, Meta = (ReturnDisplayName = "Success"))
 	UE_API bool GetBodyModifierExists(const FName Name) const;
 
+	// Logs all controls and control sets to console and log file
+	UFUNCTION(BlueprintCallable, Category = PhysicsControl)
+	UE_API void LogControls() const;
+
+	// Logs all body modifiers and body modifier sets to console and log file
+	UFUNCTION(BlueprintCallable, Category = PhysicsControl)
+	UE_API void LogBodyModifiers() const;
+
+	// Logs all controls and body modifiers and their sets to console and log file
+	UFUNCTION(BlueprintCallable, Category = PhysicsControl)
+	UE_API void LogControlsAndBodyModifiers() const;
+
 public:
 	// Public property data
 
diff --git a/Engine/Plugins/Experimental/PhysicsControl/Source/PhysicsControl/Public/PhysicsControlOperatorNameGeneration.h b/Engine/Plugins/Experimental/PhysicsControl/Source/PhysicsControl/Public/PhysicsControlOperatorNameGeneration.h
--- a/Engine/Plugins/Experimental/PhysicsControl/Source/PhysicsControl/Public/PhysicsControlOperatorNameGeneration.h
+++ b/Engine/Plugins/Experimental/PhysicsControl/Source/PhysicsControl/Public/PhysicsControlOperatorNameGeneration.h
@@ -97,26 +97,40 @@ inline bool DoesNameExist(const FName Name, const TSet<FName>& ExistingNames)
 }
 
 //======================================================================================================================
+// Makes a unique version of OriginalName.
+// 
+// Note that the convention elsewhere in UE when creating an object, and then a second identical
+// one, is for the first object to have the original name, and then the next has a suffix _2, and
+// so on.
 template<typename CollectionType>
-FName GetUniqueName(const FString& NameBase, const CollectionType& ExistingNames, const int32 MaxNameIndex)
+FName GetUniqueName(const FString& OriginalName, const CollectionType& ExistingNames, const int32 MaxNameIndex)
 {
-	FString NameStr = NameBase;
+	if (!OriginalName.IsEmpty())
+	{
+		const FName Name(OriginalName);
+		if (!DoesNameExist(Name, ExistingNames))
+		{
+			return Name;
+		}
+	}
 
 	// If the number gets too large, almost certainly we're in some nasty situation where this is
 	// getting called in a loop. Better to quit and fail, rather than allow the constraint set to
 	// increase without bound. 
-	for (int32 Index = 0; Index < MaxNameIndex; ++Index)
+	for (int32 Index = 2; Index < MaxNameIndex; ++Index)
 	{
+		FString NameStr = OriginalName.IsEmpty()
+			? FString::Format(TEXT("Unnamed_{0}"), { Index })
+			: FString::Format(TEXT("{0}_{1}"), { OriginalName, Index });
 		const FName Name(NameStr);
 		if (!DoesNameExist(Name, ExistingNames))
 		{
 			return Name;
 		}
-		NameStr = FString::Format(TEXT("{0}_{1}"), { NameBase, Index });
 	}
 	UE_LOG(LogPhysicsControl, Warning,
-		TEXT("Unable to find a suitable Control name - the limit of (%d) has been exceeded"),
-		MaxNameIndex);
+		TEXT("Unable to find a suitable Control/Modifier name for %s - the limit of (%d) has been exceeded"),
+		*OriginalName, MaxNameIndex);
 
 	return NAME_None;
 }
@@ -125,41 +139,110 @@ FName GetUniqueName(const FString& NameBase, const CollectionType& ExistingNames
 template<typename CollectionType>
 FName GetUniqueBodyModifierName(const FName BodyName, const CollectionType& ExistingNames, const FString& NamePrefix)
 {
-	FString NameBase = NamePrefix;
+	FString Name = NamePrefix;
 	if (!BodyName.IsNone())
 	{
-		NameBase += BodyName.ToString();
+		Name += BodyName.ToString();
+	}
+	else
+	{
+		Name += TEXT("Body");
+	}
+
+	return GetUniqueName(Name, ExistingNames, MaxNumControlsOrModifiersPerName);
+}
+
+//======================================================================================================================
+template<typename CollectionType>
+FName GetUniqueBodyModifierName(
+	const UPrimitiveComponent* Component, const FName BoneName, const CollectionType& ExistingNames, const FString& NamePrefix)
+{
+	// If bones are involved, make the names based on the bone, ignoring the component
+	if (!BoneName.IsNone())
+	{
+		return GetUniqueBodyModifierName(BoneName, ExistingNames, NamePrefix);
+	}
+
+	FString Name = NamePrefix;
+	if (Component)
+	{
+		Name += Component->GetName();
 	}
 	else
 	{
-		NameBase = TEXT("Body");
+		Name += TEXT("Body");
 	}
 
-	return GetUniqueName(NameBase, ExistingNames, MaxNumControlsOrModifiersPerName);
+	return GetUniqueName(Name, ExistingNames, MaxNumControlsOrModifiersPerName);
 }
 
 //======================================================================================================================
 template<typename CollectionType>
 FName GetUniqueControlName(
-	const FName ParentBodyName, const FName ChildBodyName, const CollectionType& ExistingNames, const FString& NamePrefix)
+	const FName ParentName, const FName ChildName, const CollectionType& ExistingNames, const FString& NamePrefix)
 {
-	FString NameBase = NamePrefix;
-	if (!ParentBodyName.IsNone())
+	FString Name;
+	if (!ParentName.IsNone())
 	{
-		NameBase += ParentBodyName.ToString() + TEXT("_");
+		Name = ParentName.ToString();
 	}
-	if (!ChildBodyName.IsNone())
+
+	if (!ChildName.IsNone())
 	{
-		NameBase += ChildBodyName.ToString();
+		if (!Name.IsEmpty())
+		{
+			Name += TEXT("_");
+		}
+		Name += ChildName.ToString();
+	}
+
+	if (Name.IsEmpty())
+	{
+		Name = TEXT("Control");
 	}
 
-	return GetUniqueName(NameBase, ExistingNames, MaxNumControlsOrModifiersPerName);
+	Name = NamePrefix + Name;
+
+	return GetUniqueName(Name, ExistingNames, MaxNumControlsOrModifiersPerName);
 }
 
+//======================================================================================================================
+template<typename CollectionType>
+FName GetUniqueControlName(
+	const UPrimitiveComponent* ParentComponent, const FName ParentBoneName, 
+	const UPrimitiveComponent* ChildComponent, const FName ChildBoneName, 
+	const CollectionType& ExistingNames, const FString& NamePrefix)
+{
+	// If bones are involved, make the names based on the bone, ignoring the component
+	if (!ParentBoneName.IsNone() || !ChildBoneName.IsNone())
+	{
+		return GetUniqueControlName(ParentBoneName, ChildBoneName, ExistingNames, NamePrefix);
+	}
 
+	FString Name;
+	if (ParentComponent)
+	{
+		Name = ParentComponent->GetName();
+	}
+
+	if (ChildComponent)
+	{
+		if (!Name.IsEmpty())
+		{
+			Name += TEXT("_");
+		}
+		Name += ChildComponent->GetName();
+	}
 
+	if (Name.IsEmpty())
+	{
+		Name = TEXT("Control");
+	}
 
+	Name = NamePrefix + Name;
 
+	return GetUniqueName(Name, ExistingNames, MaxNumControlsOrModifiersPerName);
+}
 
-} // namespace PhysicsControlComponent
+} // namespace PhysicsControl
 } // namespace UE
