Fix for bad memory access when world collisions are enabled in RigidBodyWithControl or RBAN

diff --git a/Engine/Plugins/Experimental/PhysicsControl/Source/PhysicsControl/Private/AnimNode_RigidBodyWithControl.cpp b/Engine/Plugins/Experimental/PhysicsControl/Source/PhysicsControl/Private/AnimNode_RigidBodyWithControl.cpp
--- a/Engine/Plugins/Experimental/PhysicsControl/Source/PhysicsControl/Private/AnimNode_RigidBodyWithControl.cpp
+++ b/Engine/Plugins/Experimental/PhysicsControl/Source/PhysicsControl/Private/AnimNode_RigidBodyWithControl.cpp
@@ -1700,8 +1700,9 @@ void FAnimNode_RigidBodyWithControl::CollectWorldObjects()
 							ImmediatePhysics::MakeKinematicActorSetup(&OverlapComp->BodyInstance, TM));
 						ActorHandle->SetName(FName(OverlapComp->GetName()));
 						Simulation->AddToCollidingPairs(ActorHandle);
-						ComponentsInSim.Add(OverlapComp, FWorldObject(ActorHandle, ComponentsInSimTick));
-						// We need this: 
+						ComponentsInSim.Add(
+							OverlapComp, FWorldObject(ActorHandle, ComponentsInSimTick, Overlap.Component));
+						// We would need this: 
 						// OverlapComp->BodyInstance.bUpdateKinematicFromSimulation = true;
 						// in order to get velocities from kinematic objects, but we can't as there's no way of 
 						// reliably undoing it (multiple RBWC nodes might be interacting with the same world object).
@@ -1720,13 +1721,13 @@ void FAnimNode_RigidBodyWithControl::ExpireWorldObjects()
 	TArray<const UPrimitiveComponent*> PrunedEntries;
 	for (auto& WorldEntry : ComponentsInSim)
 	{
-		const UPrimitiveComponent* WorldComp = WorldEntry.Key;
 		FWorldObject& WorldObject = WorldEntry.Value;
+		const UPrimitiveComponent* WorldComp = WorldObject.Component.Get();
 
 		// Do we need to expire this object?
 		const int32 ExpireTickCount = RBANWithControl_WorldObjectExpiry;
 		if (((ComponentsInSimTick - WorldObject.LastSeenTick) > ExpireTickCount)	// Haven't seen this object for a while
-			|| !IsComponentDesiredInSim(this, WorldComp))
+			|| !IsComponentDesiredInSim(this, WorldComp) )
 		{
 			WorldObject.bExpired = true;
 		}
@@ -1784,7 +1785,7 @@ void FAnimNode_RigidBodyWithControl::UpdateWorldObjects(
 
 		for (const auto& WorldEntry : ComponentsInSim)
 		{ 
-			const UPrimitiveComponent* OverlapComp = WorldEntry.Key;
+			const UPrimitiveComponent* OverlapComp = WorldEntry.Value.Component.Get();
 			if (OverlapComp != nullptr)
 			{
 				ImmediatePhysics::FActorHandle* ActorHandle = WorldEntry.Value.ActorHandle;
@@ -1831,7 +1832,7 @@ void FAnimNode_RigidBodyWithControl::UpdateWorldObjects(
 	{
 		for (const auto& WorldEntry : ComponentsInSim)
 		{
-			const UPrimitiveComponent* OverlapComp = WorldEntry.Key;
+			const UPrimitiveComponent* OverlapComp = WorldEntry.Value.Component.Get();
 			if (OverlapComp != nullptr)
 			{
 				ImmediatePhysics::FActorHandle* ActorHandle = WorldEntry.Value.ActorHandle;
diff --git a/Engine/Plugins/Experimental/PhysicsControl/Source/PhysicsControl/Public/AnimNode_RigidBodyWithControl.h b/Engine/Plugins/Experimental/PhysicsControl/Source/PhysicsControl/Public/AnimNode_RigidBodyWithControl.h
--- a/Engine/Plugins/Experimental/PhysicsControl/Source/PhysicsControl/Public/AnimNode_RigidBodyWithControl.h
+++ b/Engine/Plugins/Experimental/PhysicsControl/Source/PhysicsControl/Public/AnimNode_RigidBodyWithControl.h
@@ -626,9 +626,13 @@ private:
 	struct FWorldObject
 	{
 		FWorldObject() : ActorHandle(nullptr), LastSeenTick(0), bExpired(false) {}
-		FWorldObject(ImmediatePhysics::FActorHandle* InActorHandle, int32 InLastSeenTick) 
-			: ActorHandle(InActorHandle), LastSeenTick(InLastSeenTick), bExpired(false) {}
+		FWorldObject(
+			ImmediatePhysics::FActorHandle*     InActorHandle, 
+			int32                               InLastSeenTick, 
+			TWeakObjectPtr<UPrimitiveComponent> InComponent)
+			: Component(InComponent), ActorHandle(InActorHandle), LastSeenTick(InLastSeenTick), bExpired(false){}
 
+		TWeakObjectPtr<UPrimitiveComponent> Component;
 		ImmediatePhysics::FActorHandle* ActorHandle;
 		int32 LastSeenTick;
 		bool bExpired;
@@ -685,7 +689,10 @@ private:
 	// List of cloth colliders in the sim that represent objects collected from other parts of this character.
 	TArray<FClothCollider> ClothColliders; 
 	
-	// List of actors in the sim that represent objects collected from the world.
+	// List of actors in the sim that represent objects collected from the world. Note that the key
+	// should never be dereferenced, as it is not a smart pointer, and might have been deleted.
+	// Instead it should just be used for comparisons, and the actual pointer to the component is
+	// stored in the WorldObject.
 	TMap<const UPrimitiveComponent*, FWorldObject> ComponentsInSim;
 	int32 ComponentsInSimTick;
 
diff --git a/Engine/Source/Runtime/AnimGraphRuntime/Private/BoneControllers/AnimNode_RigidBody.cpp b/Engine/Source/Runtime/AnimGraphRuntime/Private/BoneControllers/AnimNode_RigidBody.cpp
--- a/Engine/Source/Runtime/AnimGraphRuntime/Private/BoneControllers/AnimNode_RigidBody.cpp
+++ b/Engine/Source/Runtime/AnimGraphRuntime/Private/BoneControllers/AnimNode_RigidBody.cpp
@@ -1709,10 +1709,14 @@ void FAnimNode_RigidBody::CollectWorldObjects()
 					const bool bIsSelf = (UnsafeOwner == OverlapComp->GetOwner());
 					if (!bIsSelf)
 					{
-						// Create a kinematic actor. Not using Static as world-static objects may move in the simulation's frame of reference
-						ImmediatePhysics::FActorHandle* ActorHandle = PhysicsSimulation->CreateActor(ImmediatePhysics::MakeKinematicActorSetup(&OverlapComp->BodyInstance, OverlapComp->GetComponentTransform()));
+						// Create a kinematic actor. Not using Static as world-static objects may
+						// move in the simulation's frame of reference
+						ImmediatePhysics::FActorHandle* ActorHandle = PhysicsSimulation->CreateActor(
+								ImmediatePhysics::MakeKinematicActorSetup(
+									&OverlapComp->BodyInstance, OverlapComp->GetComponentTransform()));
 						PhysicsSimulation->AddToCollidingPairs(ActorHandle);
-						ComponentsInSim.Add(OverlapComp, FWorldObject(ActorHandle, ComponentsInSimTick));
+						ComponentsInSim.Add(
+							OverlapComp, FWorldObject(ActorHandle, ComponentsInSimTick, Overlap.Component));
 					}
 				}
 			}
@@ -1727,8 +1731,8 @@ void FAnimNode_RigidBody::ExpireWorldObjects()
 	TArray<const UPrimitiveComponent*> PrunedEntries;
 	for (auto& WorldEntry : ComponentsInSim)
 	{
-		const UPrimitiveComponent* WorldComp = WorldEntry.Key;
 		FWorldObject& WorldObject = WorldEntry.Value;
+		const UPrimitiveComponent* WorldComp = WorldObject.Component.Get();
 
 		// Do we need to expire this object?
 		const int32 ExpireTickCount = RBAN_WorldObjectExpiry;
@@ -1782,7 +1786,7 @@ void FAnimNode_RigidBody::UpdateWorldObjects(const FTransform& SpaceTransform)
 	{
 		for (TPair<const UPrimitiveComponent*, FWorldObject>& WorldEntry : ComponentsInSim)
 		{ 
-			const UPrimitiveComponent* OverlapComp = WorldEntry.Key;
+			const UPrimitiveComponent* OverlapComp = WorldEntry.Value.Component.Get();
 			if (OverlapComp != nullptr)
 			{
 				FWorldObject& WorldObject = WorldEntry.Value;
diff --git a/Engine/Source/Runtime/AnimGraphRuntime/Public/BoneControllers/AnimNode_RigidBody.h b/Engine/Source/Runtime/AnimGraphRuntime/Public/BoneControllers/AnimNode_RigidBody.h
--- a/Engine/Source/Runtime/AnimGraphRuntime/Public/BoneControllers/AnimNode_RigidBody.h
+++ b/Engine/Source/Runtime/AnimGraphRuntime/Public/BoneControllers/AnimNode_RigidBody.h
@@ -473,8 +473,15 @@ private:
 	struct FWorldObject
 	{
 		FWorldObject() : ActorHandle(nullptr), LastSeenTick(0), bExpired(false), bNew(true) {}
-		FWorldObject(ImmediatePhysics::FActorHandle* InActorHandle, int32 InLastSeenTick) : ActorHandle(InActorHandle), LastSeenTick(InLastSeenTick), bExpired(false), bNew(true) {}
-
+		FWorldObject(
+			ImmediatePhysics::FActorHandle*     InActorHandle, 
+			int32                               InLastSeenTick,
+			TWeakObjectPtr<UPrimitiveComponent> InComponent) 
+			: Component(InComponent), ActorHandle(InActorHandle), LastSeenTick(InLastSeenTick)
+			, bExpired(false), bNew(true) {
+		}
+
+		TWeakObjectPtr<UPrimitiveComponent> Component;
 		ImmediatePhysics::FActorHandle* ActorHandle;
 		int32 LastSeenTick;
 		uint8 bExpired : 1;
@@ -508,6 +515,10 @@ private:
 	// List of actors in the sim that represent objects collected from other parts of this character.
 	TArray<FClothCollider> ClothColliders; 
 	
+	// List of actors in the sim that represent objects collected from the world. Note that the key
+	// should never be dereferenced, as it is not a smart pointer, and might have been deleted.
+	// Instead it should just be used for comparisons, and the actual pointer to the component is
+	// stored in the WorldObject.
 	TMap<const UPrimitiveComponent*, FWorldObject> ComponentsInSim;
 	int32 ComponentsInSimTick;
 
