Guard against the RequiredBones container in RBAN and RBWC InitializeBoneReferences having an empty reference pose (to prevent a crash)

#Synced-CL 49327860

diff --git a/Engine/Plugins/Experimental/PhysicsControl/Source/PhysicsControl/Private/AnimNode_RigidBodyWithControl.cpp b/Engine/Plugins/Experimental/PhysicsControl/Source/PhysicsControl/Private/AnimNode_RigidBodyWithControl.cpp
--- a/Engine/Plugins/Experimental/PhysicsControl/Source/PhysicsControl/Private/AnimNode_RigidBodyWithControl.cpp
+++ b/Engine/Plugins/Experimental/PhysicsControl/Source/PhysicsControl/Private/AnimNode_RigidBodyWithControl.cpp
@@ -1860,11 +1860,18 @@ void FAnimNode_RigidBodyWithControl::InitializeBoneReferences(const FBoneContain
 	DECLARE_SCOPE_HIERARCHICAL_COUNTER_ANIMNODE(InitializeBoneReferences)
 	/** We only need to update simulated bones and children of simulated bones*/
 	const int32 NumBodies = Bodies.Num();
+	OutputBoneData.Empty(NumBodies);
+
+	if (!RequiredBones.IsValid())
+	{
+		UE_LOG(LogPhysicsControl, Warning, TEXT("FAnimNode_RigidBodyWithControl::InitializeBoneReferences: Required bones is invalid for Skeleton %s / SkeletalMesh %s."),
+			*GetNameSafe(RequiredBones.GetSkeletonAsset()), *GetNameSafe(RequiredBones.GetSkeletalMeshAsset()));
+		return;
+	}
+	const FReferenceSkeleton& RefSkeleton = RequiredBones.GetReferenceSkeleton();
 	const TArray<FBoneIndexType>& RequiredBoneIndices = RequiredBones.GetBoneIndicesArray();
 	const int32 NumRequiredBoneIndices = RequiredBoneIndices.Num();
-	const FReferenceSkeleton& RefSkeleton = RequiredBones.GetReferenceSkeleton();
 
-	OutputBoneData.Empty(NumBodies);
 	int32 OutputBoneDataMaxBodyIndex = -1;
 
 	int32 NumSimulatedBodies = 0;
diff --git a/Engine/Source/Runtime/AnimGraphRuntime/Private/BoneControllers/AnimNode_RigidBody.cpp b/Engine/Source/Runtime/AnimGraphRuntime/Private/BoneControllers/AnimNode_RigidBody.cpp
--- a/Engine/Source/Runtime/AnimGraphRuntime/Private/BoneControllers/AnimNode_RigidBody.cpp
+++ b/Engine/Source/Runtime/AnimGraphRuntime/Private/BoneControllers/AnimNode_RigidBody.cpp
@@ -1819,11 +1819,18 @@ void FAnimNode_RigidBody::InitializeBoneReferences(const FBoneContainer& Require
 	DECLARE_SCOPE_HIERARCHICAL_COUNTER_ANIMNODE(InitializeBoneReferences)
 	/** We only need to update simulated bones and children of simulated bones*/
 	const int32 NumBodies = Bodies.Num();
+	OutputBoneData.Empty(NumBodies);
+
+	if (!RequiredBones.IsValid())
+	{
+		UE_LOG(LogRBAN, Warning, TEXT("FAnimNode_RigidBody::InitializeBoneReferences: Required bones is invalid for Skeleton %s / SkeletalMesh %s."),
+			*GetNameSafe(RequiredBones.GetSkeletonAsset()), *GetNameSafe(RequiredBones.GetSkeletalMeshAsset()));
+		return;
+	}
+
+	const FReferenceSkeleton& RefSkeleton = RequiredBones.GetReferenceSkeleton();
 	const TArray<FBoneIndexType>& RequiredBoneIndices = RequiredBones.GetBoneIndicesArray();
 	const int32 NumRequiredBoneIndices = RequiredBoneIndices.Num();
-	const FReferenceSkeleton& RefSkeleton = RequiredBones.GetReferenceSkeleton();
-
-	OutputBoneData.Empty(NumBodies);
 
 	int32 NumSimulatedBodies = 0;
 	TArray<int32> SimulatedBodyIndices;
