Physics Control Component - make sure that limb-generated controls get automatically added to a set with the limb name

#Synced-CL 49649139

diff --git a/Engine/Plugins/Experimental/PhysicsControl/Source/PhysicsControl/Private/PhysicsControlComponent.cpp b/Engine/Plugins/Experimental/PhysicsControl/Source/PhysicsControl/Private/PhysicsControlComponent.cpp
--- a/Engine/Plugins/Experimental/PhysicsControl/Source/PhysicsControl/Private/PhysicsControlComponent.cpp
+++ b/Engine/Plugins/Experimental/PhysicsControl/Source/PhysicsControl/Private/PhysicsControlComponent.cpp
@@ -586,6 +586,8 @@ TMap<FName, FPhysicsControlNames> UPhysicsControlComponent::CreateControlsFromLi
 				ParentBoneName = WorldBoneName;
 			}
 
+			// When the control is created it will be added to SetName, which will be something like
+			// WorldSpace_Head, as well as "All".
 			const FName ControlName = CreateControl(
 				ParentComponent, ParentBoneName, BonesInLimb.SkeletalMeshComponent.Get(), ChildBoneName,
 				ControlData, FPhysicsControlTarget(), FName(SetName));
@@ -594,7 +596,7 @@ TMap<FName, FPhysicsControlNames> UPhysicsControlComponent::CreateControlsFromLi
 			{
 				LimbResult.Names.Add(ControlName);
 				AllControls.Names.Add(ControlName);
-				NameRecords.AddControl(ControlName, GetControlTypeName(ControlType));
+				NameRecords.AddControl(ControlName, { LimbName, GetControlTypeName(ControlType) }, false);
 			}
 			else
 			{
@@ -686,7 +688,8 @@ TMap<FName, FPhysicsControlNames> UPhysicsControlComponent::CreateControlsFromLi
 			{
 				LimbResult.Names.Add(ControlName);
 				AllControls.Names.Add(ControlName);
-				NameRecords.AddControl(ControlName, GetControlTypeName(EPhysicsControlType::ParentSpace));
+				NameRecords.AddControl(
+					ControlName, { LimbName, GetControlTypeName(EPhysicsControlType::ParentSpace) }, false);
 			}
 			else
 			{
diff --git a/Engine/Plugins/Experimental/PhysicsControl/Source/PhysicsControl/Private/PhysicsControlNameRecords.cpp b/Engine/Plugins/Experimental/PhysicsControl/Source/PhysicsControl/Private/PhysicsControlNameRecords.cpp
--- a/Engine/Plugins/Experimental/PhysicsControl/Source/PhysicsControl/Private/PhysicsControlNameRecords.cpp
+++ b/Engine/Plugins/Experimental/PhysicsControl/Source/PhysicsControl/Private/PhysicsControlNameRecords.cpp
@@ -2,7 +2,7 @@
 #include "PhysicsControlNameRecords.h"
 
 //======================================================================================================================
-void AddUniqueNamesToSet(const TArray<FName>& Names, const FName SetName, TMap<FName, TArray<FName>>& SetCollection)
+void AddUniqueNamesToSet(TConstArrayView<FName> Names, const FName SetName, TMap<FName, TArray<FName>>& SetCollection)
 {
 	if (!SetName.IsNone())
 	{
@@ -15,17 +15,20 @@ void AddUniqueNamesToSet(const TArray<FName>& Names, const FName SetName, TMap<F
 }
 
 //======================================================================================================================
-void FPhysicsControlNameRecords::AddControl(FName Name, FName SetName)
+void FPhysicsControlNameRecords::AddControl(FName Name, FName SetName, bool bAddToAll)
 {
 	if (!SetName.IsNone())
 	{
 		ControlSets.FindOrAdd(SetName).AddUnique(Name);
 	}
-	ControlSets.FindOrAdd(TEXT("All")).AddUnique(Name);
+	if (bAddToAll)
+	{
+		ControlSets.FindOrAdd(TEXT("All")).AddUnique(Name);
+	}
 }
 
 //======================================================================================================================
-void FPhysicsControlNameRecords::AddControl(FName Name, const TArray<FName>& SetNames)
+void FPhysicsControlNameRecords::AddControl(FName Name, TConstArrayView<FName> SetNames, bool bAddToAll)
 {
 	for (const FName SetName : SetNames)
 	{
@@ -34,15 +37,21 @@ void FPhysicsControlNameRecords::AddControl(FName Name, const TArray<FName>& Set
 			ControlSets.FindOrAdd(SetName).AddUnique(Name);
 		}
 	}
-	ControlSets.FindOrAdd(TEXT("All")).AddUnique(Name);
+	if (bAddToAll)
+	{
+		ControlSets.FindOrAdd(TEXT("All")).AddUnique(Name);
+	}
 }
 
 
 //======================================================================================================================
-void FPhysicsControlNameRecords::AddControls(const TArray<FName>& ControlNames, const FName SetName)
+void FPhysicsControlNameRecords::AddControls(TConstArrayView<FName> ControlNames, const FName SetName, bool bAddToAll)
 {
 	AddUniqueNamesToSet(ControlNames, SetName, ControlSets);
-	AddUniqueNamesToSet(ControlNames, "All", ControlSets);
+	if (bAddToAll)
+	{
+		AddUniqueNamesToSet(ControlNames, "All", ControlSets);
+	}
 }
 
 //======================================================================================================================
@@ -67,17 +76,20 @@ const TArray<FName>& FPhysicsControlNameRecords::GetControlNamesInSet(FName SetN
 }
 
 //======================================================================================================================
-void FPhysicsControlNameRecords::AddBodyModifier(FName Name, FName SetName)
+void FPhysicsControlNameRecords::AddBodyModifier(FName Name, FName SetName, bool bAddToAll)
 {
 	if (!SetName.IsNone())
 	{
 		BodyModifierSets.FindOrAdd(SetName).AddUnique(Name);
 	}
-	BodyModifierSets.FindOrAdd(TEXT("All")).AddUnique(Name);
+	if (bAddToAll)
+	{
+		BodyModifierSets.FindOrAdd(TEXT("All")).AddUnique(Name);
+	}
 }
 
 //======================================================================================================================
-void FPhysicsControlNameRecords::AddBodyModifier(FName Name, const TArray<FName>& SetNames)
+void FPhysicsControlNameRecords::AddBodyModifier(FName Name, TConstArrayView<FName> SetNames, bool bAddToAll)
 {
 	for (FName SetName : SetNames)
 	{
@@ -86,14 +98,21 @@ void FPhysicsControlNameRecords::AddBodyModifier(FName Name, const TArray<FName>
 			BodyModifierSets.FindOrAdd(SetName).AddUnique(Name);
 		}
 	}
-	BodyModifierSets.FindOrAdd(TEXT("All")).AddUnique(Name);
+	if (bAddToAll)
+	{
+		BodyModifierSets.FindOrAdd(TEXT("All")).AddUnique(Name);
+	}
 }
 
 //======================================================================================================================
-void FPhysicsControlNameRecords::AddBodyModifiers(const TArray<FName>& BodyModifierNames, const FName SetName)
+void FPhysicsControlNameRecords::AddBodyModifiers(
+	TConstArrayView<FName> BodyModifierNames, const FName SetName, bool bAddToAll)
 {
 	AddUniqueNamesToSet(BodyModifierNames, SetName, BodyModifierSets);
-	AddUniqueNamesToSet(BodyModifierNames, TEXT("All"), BodyModifierSets);
+	if (bAddToAll)
+	{
+		AddUniqueNamesToSet(BodyModifierNames, TEXT("All"), BodyModifierSets);
+	}
 }
 
 //======================================================================================================================
diff --git a/Engine/Plugins/Experimental/PhysicsControl/Source/PhysicsControl/Public/PhysicsControlNameRecords.h b/Engine/Plugins/Experimental/PhysicsControl/Source/PhysicsControl/Public/PhysicsControlNameRecords.h
--- a/Engine/Plugins/Experimental/PhysicsControl/Source/PhysicsControl/Public/PhysicsControlNameRecords.h
+++ b/Engine/Plugins/Experimental/PhysicsControl/Source/PhysicsControl/Public/PhysicsControlNameRecords.h
@@ -14,22 +14,22 @@
  */
 struct FPhysicsControlNameRecords
 {
-	/** Adds Name to SetName as well as adding it to the set "All" */
-	void AddControl(FName Name, FName SetName);
-	/** Adds Name to SetNames as well as adding it to the set "All" */
-	void AddControl(FName Name, const TArray<FName>& SetNames);
-	/** Adds a collection of Names to SetName as well as adding them to the set "All"  */
-	void AddControls(const TArray<FName>& ControlNames, const FName SetName);
+	/** Adds Name to SetName as well as optionally adding it to the set "All" */
+	void AddControl(FName Name, FName SetName, bool bAddToAll = true);
+	/** Adds Name to SetNames as well as optionally adding it to the set "All" */
+	void AddControl(FName Name, TConstArrayView<FName> SetNames, bool bAddToAll = true);
+	/** Adds a collection of Names to SetName as well as optionally adding them to the set "All"  */
+	void AddControls(TConstArrayView<FName> ControlNames, const FName SetName, bool bAddToAll = true);
 	/** Removes Name from all Sets */
 	void RemoveControl(FName Name);
 	const TArray<FName>& GetControlNamesInSet(FName SetName) const;
 
-	/** Adds Name to SetName as well as adding it to the set "All" */
-	void AddBodyModifier(FName Name, FName SetName);
-	/** Adds Name to SetNames as well as adding it to the set "All" */
-	void AddBodyModifier(FName Name, const TArray<FName>& SetNames);
-	/** Adds a collection of Names to SetName as well as adding them to the set "All"  */
-	void AddBodyModifiers(const TArray<FName>& BodyModifierNames, const FName SetName);
+	/** Adds Name to SetName as well as optionally adding it to the set "All" */
+	void AddBodyModifier(FName Name, FName SetName, bool bAddToAll = true);
+	/** Adds Name to SetNames as well as optionally adding it to the set "All" */
+	void AddBodyModifier(FName Name, TConstArrayView<FName> SetNames, bool bAddToAll = true);
+	/** Adds a collection of Names to SetName as well as optionally adding them to the set "All"  */
+	void AddBodyModifiers(TConstArrayView<FName> BodyModifierNames, const FName SetName, bool bAddToAll = true);
 	/** Removes Name from all Sets */
 	void RemoveBodyModifier(FName Name);
 	const TArray<FName>& GetBodyModifierNamesInSet(FName SetName) const;
