Support CCD in PhysicsControl. Currently exposed fully in Control Rig Physics as an option on the solver and individual bodies

#virtualized
#Synced-CL 51665965

#Synced-CL 51665965

#Synced-CL 51665965

diff --git a/Engine/Plugins/Experimental/ControlRigPhysics/Source/ControlRigPhysics/Private/RigPhysicsBodyExecution.cpp b/Engine/Plugins/Experimental/ControlRigPhysics/Source/ControlRigPhysics/Private/RigPhysicsBodyExecution.cpp
--- a/Engine/Plugins/Experimental/ControlRigPhysics/Source/ControlRigPhysics/Private/RigPhysicsBodyExecution.cpp
+++ b/Engine/Plugins/Experimental/ControlRigPhysics/Source/ControlRigPhysics/Private/RigPhysicsBodyExecution.cpp
@@ -241,6 +241,20 @@ FRigUnit_HierarchySetPhysicsBodyMovementType_Execute()
 	}
 }
 
+//======================================================================================================================
+FRigUnit_HierarchySetPhysicsBodyEnableCCD_Execute()
+{
+	DECLARE_SCOPE_HIERARCHICAL_COUNTER_RIGUNIT();
+	if (!ExecuteContext.Hierarchy)
+	{
+		return;
+	}
+	if (FRigPhysicsBodyComponent* Component = GetBody(ExecuteContext.Hierarchy, PhysicsBodyComponentKey))
+	{
+		Component->BodyData.bEnableCCD = bEnableCCD;
+	}
+}
+
 //======================================================================================================================
 FRigUnit_HierarchySetPhysicsBodyCollisionType_Execute()
 {
diff --git a/Engine/Plugins/Experimental/ControlRigPhysics/Source/ControlRigPhysics/Private/RigPhysicsData.cpp b/Engine/Plugins/Experimental/ControlRigPhysics/Source/ControlRigPhysics/Private/RigPhysicsData.cpp
--- a/Engine/Plugins/Experimental/ControlRigPhysics/Source/ControlRigPhysics/Private/RigPhysicsData.cpp
+++ b/Engine/Plugins/Experimental/ControlRigPhysics/Source/ControlRigPhysics/Private/RigPhysicsData.cpp
@@ -299,6 +299,12 @@ FArchive& operator <<(FArchive& Ar, FRigPhysicsSolverSettings& Data)
 	Ar << Data.bUseLinearJointSolver;
 	Ar << Data.bSolveJointPositionsLast;
 	Ar << Data.bUseManifolds;
+
+	if (Ar.CustomVer(FPhysicsControlObjectVersion::GUID) >= FPhysicsControlObjectVersion::ControlRigAllowCCD)
+	{
+		Ar << Data.bAllowCCD;
+	}
+
 	Ar << Data.PositionThresholdForReset;
 	if (Ar.CustomVer(FPhysicsControlObjectVersion::GUID) >= FPhysicsControlObjectVersion::ControlRigSpeedThresholdForReset)
 	{
diff --git a/Engine/Plugins/Experimental/ControlRigPhysics/Source/ControlRigPhysics/Private/RigPhysicsExecution.cpp b/Engine/Plugins/Experimental/ControlRigPhysics/Source/ControlRigPhysics/Private/RigPhysicsExecution.cpp
--- a/Engine/Plugins/Experimental/ControlRigPhysics/Source/ControlRigPhysics/Private/RigPhysicsExecution.cpp
+++ b/Engine/Plugins/Experimental/ControlRigPhysics/Source/ControlRigPhysics/Private/RigPhysicsExecution.cpp
@@ -541,6 +541,7 @@ FRigUnit_HierarchyInstantiateFromPhysicsAsset_Execute()
 			BodyData.CollisionType =
 				SkeletalBodySetup->CollisionReponse == EBodyCollisionResponse::BodyCollision_Enabled ?
 				ECollisionEnabled::PhysicsOnly : ECollisionEnabled::NoCollision;
+			BodyData.bEnableCCD = SkeletalBodySetup->DefaultInstance.bUseCCD;
 
 			if (bNeedBodyAsHelper)
 			{
@@ -849,6 +850,7 @@ FRigUnit_HierarchyImportCollisionFromPhysicsAsset_Execute()
 			BodyData.CollisionType =
 				SkeletalBodySetup->CollisionReponse == EBodyCollisionResponse::BodyCollision_Enabled ?
 				ECollisionEnabled::PhysicsOnly : ECollisionEnabled::NoCollision;
+			BodyData.bEnableCCD = SkeletalBodySetup->DefaultInstance.bUseCCD;
 
 			// Add the component
 			FRigComponentKey PhysicsBodyComponentKey = Controller->AddComponent(
@@ -900,4 +902,18 @@ FRigUnit_GetPhysicsSolverSpaceData_Execute()
 			}
 		}
 	}
-}
\ No newline at end of file
+}
+
+//======================================================================================================================
+FRigUnit_HierarchySetPhysicsSolverAllowCCD_Execute()
+{
+	if (!ExecuteContext.Hierarchy)
+	{
+		return;
+	}
+	if (FRigPhysicsSolverComponent* SolverComponent = Cast<FRigPhysicsSolverComponent>(
+		ExecuteContext.Hierarchy->FindComponent(PhysicsSolverComponentKey)))
+	{
+		SolverComponent->SolverSettings.bAllowCCD = bAllowCCD;
+	}
+}
diff --git a/Engine/Plugins/Experimental/ControlRigPhysics/Source/ControlRigPhysics/Private/RigPhysicsSimulation.cpp b/Engine/Plugins/Experimental/ControlRigPhysics/Source/ControlRigPhysics/Private/RigPhysicsSimulation.cpp
--- a/Engine/Plugins/Experimental/ControlRigPhysics/Source/ControlRigPhysics/Private/RigPhysicsSimulation.cpp
+++ b/Engine/Plugins/Experimental/ControlRigPhysics/Source/ControlRigPhysics/Private/RigPhysicsSimulation.cpp
@@ -1889,6 +1889,7 @@ void FRigPhysicsSimulation::StepSimulation(
 	Chaos::FCollisionDetectorSettings CollisionDetectorSettings = Simulation->GetCollisionDetectorSettings();
 	CollisionDetectorSettings.BoundsVelocityInflation = SolverSettings.BoundsVelocityMultiplier;
 	CollisionDetectorSettings.MaxVelocityBoundsExpansion = SolverSettings.MaxVelocityBoundsExpansion;
+	CollisionDetectorSettings.bAllowCCD = SolverSettings.bAllowCCD;
 	Simulation->SetCollisionDetectorSettings(CollisionDetectorSettings);
 
 	// This gets reset to 100 after every simulation step!
diff --git a/Engine/Plugins/Experimental/ControlRigPhysics/Source/ControlRigPhysics/Public/RigPhysicsBodyExecution.h b/Engine/Plugins/Experimental/ControlRigPhysics/Source/ControlRigPhysics/Public/RigPhysicsBodyExecution.h
--- a/Engine/Plugins/Experimental/ControlRigPhysics/Source/ControlRigPhysics/Public/RigPhysicsBodyExecution.h
+++ b/Engine/Plugins/Experimental/ControlRigPhysics/Source/ControlRigPhysics/Public/RigPhysicsBodyExecution.h
@@ -538,6 +538,35 @@ struct FRigUnit_HierarchySetPhysicsBodyDamping : public FRigUnit_PhysicsBaseMuta
 	float AngularDamping = 0;
 };
 
+/**
+ * Enable/disable CCD. Note that this can be expensive, so disable it when not needed!
+ * Requires CCD to be allowed in the solver.
+ */
+USTRUCT(meta = (DisplayName = "Set Physics Body Enable CCD", Varying))
+struct FRigUnit_HierarchySetPhysicsBodyEnableCCD : public FRigUnit_PhysicsBaseMutable
+{
+	GENERATED_BODY()
+
+	FRigUnit_HierarchySetPhysicsBodyEnableCCD()
+	{
+		PhysicsBodyComponentKey.ElementKey.Type = ERigElementType::Bone;
+		PhysicsBodyComponentKey.Name = FRigPhysicsBodyComponent::GetDefaultName();
+	}
+
+	RIGVM_METHOD()
+	UE_API virtual void Execute() override;
+
+	// The Physics Body to be updated
+	UPROPERTY(meta = (Input))
+	FRigComponentKey PhysicsBodyComponentKey;
+
+	// Enable / disable CCD. Note that this can be expensive, so disable it when not needed!
+	// Requires CCD to be allowed in the solver.
+	UPROPERTY(meta = (Input))
+	bool bEnableCCD = false;
+};
+
+
 /**
  * This will add a force (or impulse, torque, velocity change etc) to the body. The force record 
  * will be stored in the physics body component, and then applied next time the solver is stepped, 
diff --git a/Engine/Plugins/Experimental/ControlRigPhysics/Source/ControlRigPhysics/Public/RigPhysicsData.h b/Engine/Plugins/Experimental/ControlRigPhysics/Source/ControlRigPhysics/Public/RigPhysicsData.h
--- a/Engine/Plugins/Experimental/ControlRigPhysics/Source/ControlRigPhysics/Public/RigPhysicsData.h
+++ b/Engine/Plugins/Experimental/ControlRigPhysics/Source/ControlRigPhysics/Public/RigPhysicsData.h
@@ -536,6 +536,12 @@ struct FRigPhysicsSolverSettings
 	UPROPERTY(EditAnywhere, Category = CollisionSettings)
 	bool bUseManifolds = true;
 
+	/**
+	 * Allows the solver to use the EnableCCD setting on the bodies
+	 */
+	UPROPERTY(EditAnywhere, Category = CollisionSettings)
+	bool bAllowCCD = false;
+
 	// If any object in the simulation exceeds this distance from the simulation origin, we will
 	// reset the pose and velocity of the entire simulation. This is to detect problems/explosions.
 	// Disabled when the value is zero.
diff --git a/Engine/Plugins/Experimental/ControlRigPhysics/Source/ControlRigPhysics/Public/RigPhysicsExecution.h b/Engine/Plugins/Experimental/ControlRigPhysics/Source/ControlRigPhysics/Public/RigPhysicsExecution.h
--- a/Engine/Plugins/Experimental/ControlRigPhysics/Source/ControlRigPhysics/Public/RigPhysicsExecution.h
+++ b/Engine/Plugins/Experimental/ControlRigPhysics/Source/ControlRigPhysics/Public/RigPhysicsExecution.h
@@ -527,4 +527,30 @@ struct FRigUnit_GetPhysicsSolverSpaceData : public FRigUnit_PhysicsBase
 	FVector Gravity = FVector::ZeroVector;
 };
 
+/**
+ * Enables/disables support for CCD in the solver
+ */
+USTRUCT(meta = (DisplayName = "Set Physics Solver Allow CCD", Varying))
+struct FRigUnit_HierarchySetPhysicsSolverAllowCCD : public FRigUnit_PhysicsBaseMutable
+{
+	GENERATED_BODY()
+
+	FRigUnit_HierarchySetPhysicsSolverAllowCCD()
+	{
+		PhysicsSolverComponentKey.ElementKey.Type = ERigElementType::Bone;
+		PhysicsSolverComponentKey.Name = FRigPhysicsSolverComponent::GetDefaultName();
+	}
+
+	RIGVM_METHOD()
+	UE_API virtual void Execute() override;
+
+	// The Physics Solver to be updated
+	UPROPERTY(meta = (Input))
+	FRigComponentKey PhysicsSolverComponentKey;
+
+	// Whether or not the solver allows CCD to be used. CCD also needs to be enabled on individual bodies
+	UPROPERTY(meta = (Input))
+	bool bAllowCCD = false;
+};
+
 #undef UE_API
diff --git a/Engine/Plugins/Experimental/PhysicsControl/Source/PhysicsControl/Private/PhysicsControlData.cpp b/Engine/Plugins/Experimental/PhysicsControl/Source/PhysicsControl/Private/PhysicsControlData.cpp
--- a/Engine/Plugins/Experimental/PhysicsControl/Source/PhysicsControl/Private/PhysicsControlData.cpp
+++ b/Engine/Plugins/Experimental/PhysicsControl/Source/PhysicsControl/Private/PhysicsControlData.cpp
@@ -149,6 +149,7 @@ FPhysicsControlModifierData Interpolate(
 	INTERPOLATE_PARAM(PhysicsBlendWeight);
 	SELECT_PARAM(KinematicTargetSpace);
 	SELECT_PARAM(bUpdateKinematicFromSimulation);
+	SELECT_PARAM(bEnableCCD);
 	return Output;
 }
 
@@ -163,6 +164,7 @@ FPhysicsControlModifierSparseData Interpolate(
 	INTERPOLATE_PARAM(PhysicsBlendWeight);
 	SELECT_PARAM(KinematicTargetSpace);
 	SELECT_PARAM(bUpdateKinematicFromSimulation);
+	SELECT_PARAM(bEnableCCD);
 
 	SET_ENABLED_PARAM(MovementType);
 	SET_ENABLED_PARAM(CollisionType);
@@ -170,6 +172,7 @@ FPhysicsControlModifierSparseData Interpolate(
 	SET_ENABLED_PARAM(PhysicsBlendWeight);
 	SET_ENABLED_PARAM(KinematicTargetSpace);
 	SET_ENABLED_PARAM(bUpdateKinematicFromSimulation);
+	SET_ENABLED_PARAM(bEnableCCD);
 	return Output;
 }
 
@@ -216,6 +219,7 @@ void FPhysicsControlModifierData::UpdateFromSparseData(const FPhysicsControlModi
 	SET_SPARSE_DATA(PhysicsBlendWeight);
 	SET_SPARSE_DATA(KinematicTargetSpace);
 	SET_SPARSE_DATA(bUpdateKinematicFromSimulation);
+	SET_SPARSE_DATA(bEnableCCD);
 }
 
 //======================================================================================================================
@@ -261,6 +265,10 @@ FArchive& operator <<(FArchive& Ar, FPhysicsControlModifierData& Data)
 		Ar << Data.KinematicTargetSpace;
 	}
 	Ar << Data.bUpdateKinematicFromSimulation;
+	if (Ar.CustomVer(FPhysicsControlObjectVersion::GUID) >= FPhysicsControlObjectVersion::EnableCCD)
+	{
+		Ar << Data.bEnableCCD;
+	}
 	return Ar;
 }
 
diff --git a/Engine/Plugins/Experimental/PhysicsControl/Source/PhysicsControl/Private/RigidBodyWithControl.cpp b/Engine/Plugins/Experimental/PhysicsControl/Source/PhysicsControl/Private/RigidBodyWithControl.cpp
--- a/Engine/Plugins/Experimental/PhysicsControl/Source/PhysicsControl/Private/RigidBodyWithControl.cpp
+++ b/Engine/Plugins/Experimental/PhysicsControl/Source/PhysicsControl/Private/RigidBodyWithControl.cpp
@@ -148,6 +148,8 @@ void UpdateBodyFromModifierData(
 			ModifierData.MovementType != EPhysicsMovementType::Simulated);
 	}
 
+	ActorHandle.SetCCDEnabled(ModifierData.bEnableCCD);
+
 	// Note that the actual kinematic targets will be set separately, since they need to be set
 	// for all kinematics whether or not they were under a modifier.
 
diff --git a/Engine/Plugins/Experimental/PhysicsControl/Source/PhysicsControl/Public/PhysicsControlData.h b/Engine/Plugins/Experimental/PhysicsControl/Source/PhysicsControl/Public/PhysicsControlData.h
--- a/Engine/Plugins/Experimental/PhysicsControl/Source/PhysicsControl/Public/PhysicsControlData.h
+++ b/Engine/Plugins/Experimental/PhysicsControl/Source/PhysicsControl/Public/PhysicsControlData.h
@@ -155,13 +155,15 @@ struct FPhysicsControlModifierSparseData
 		const float                               InGravityMultiplier = 1.0f,
 		const float                               InPhysicsBlendWeight = 1.0f,
 		const EPhysicsControlKinematicTargetSpace InKinematicTargetSpace = EPhysicsControlKinematicTargetSpace::OffsetInBoneSpace,
-		const bool                                InUpdateKinematicFromSimulation = true)
+		const bool                                InUpdateKinematicFromSimulation = true,
+		const bool                                InEnableCCD = false)
 		: MovementType(InMovementType)
 		, CollisionType(InCollisionType)
 		, GravityMultiplier(InGravityMultiplier)
 		, PhysicsBlendWeight(InPhysicsBlendWeight)
 		, KinematicTargetSpace(InKinematicTargetSpace)
 		, bUpdateKinematicFromSimulation(InUpdateKinematicFromSimulation)
+		, bEnableCCD(InEnableCCD)
 		// The flags
 		, bEnableMovementType(true)
 		, bEnableCollisionType(true)
@@ -169,6 +171,7 @@ struct FPhysicsControlModifierSparseData
 		, bEnablePhysicsBlendWeight(true)
 		, bEnableKinematicTargetSpace(true)
 		, bEnablebUpdateKinematicFromSimulation(true)
+		, bEnablebEnableCCD(true)
 	{
 	}
 
@@ -213,6 +216,13 @@ struct FPhysicsControlModifierSparseData
 	UPROPERTY(EditAnywhere, BlueprintReadWrite, Category = PhysicsControl, meta = (editcondition = "bEnablebUpdateKinematicFromSimulation"))
 	uint8 bUpdateKinematicFromSimulation : 1;
 
+	/**
+	 * Enable/disable CCD. Note that this can be expensive, so disable it when not needed! 
+	 * Requires CCD to be allowed in the solver.
+	 */
+	UPROPERTY(EditAnywhere, BlueprintReadWrite, Category = PhysicsControl, meta = (editcondition = "bEnablebEnableCCD"))
+	uint8 bEnableCCD : 1;
+
 	UPROPERTY(EditAnywhere, Category = PhysicsControl, meta = (InlineEditConditionToggle))
 	uint8 bEnableMovementType : 1;
 
@@ -230,6 +240,9 @@ struct FPhysicsControlModifierSparseData
 
 	UPROPERTY(EditAnywhere, Category = PhysicsControl, meta = (InlineEditConditionToggle))
 	uint8 bEnablebUpdateKinematicFromSimulation : 1;
+
+	UPROPERTY(EditAnywhere, Category = PhysicsControl, meta = (InlineEditConditionToggle))
+	uint8 bEnablebEnableCCD : 1;
 };
 
 // Specifies a space for queries, forces etc
@@ -260,13 +273,15 @@ struct FPhysicsControlModifierData
 		const float                               InGravityMultiplier = 1.0f,
 		const float                               InPhysicsBlendWeight = 1.0f,
 		const EPhysicsControlKinematicTargetSpace InKinematicTargetSpace = EPhysicsControlKinematicTargetSpace::OffsetInBoneSpace,
-		const bool                                InUpdateKinematicFromSimulation = true)
+		const bool                                InUpdateKinematicFromSimulation = true,
+		const bool                                InEnableCCD = false)
 		: MovementType(InMovementType)
 		, CollisionType(InCollisionType)
 		, KinematicTargetSpace(InKinematicTargetSpace)
 		, GravityMultiplier(InGravityMultiplier)
 		, PhysicsBlendWeight(InPhysicsBlendWeight)
 		, bUpdateKinematicFromSimulation(InUpdateKinematicFromSimulation)
+		, bEnableCCD(InEnableCCD)
 	{
 	}
 
@@ -314,6 +329,13 @@ struct FPhysicsControlModifierData
 	 */
 	UPROPERTY(EditAnywhere, BlueprintReadWrite, Category = PhysicsControl)
 	bool bUpdateKinematicFromSimulation;
+
+	/**
+	 * Enable/disable CCD. Note that this can be expensive, so disable it when not needed!
+	 * Requires CCD to be allowed in the solver.
+	 */
+	UPROPERTY(EditAnywhere, BlueprintReadWrite, Category = PhysicsControl)
+	bool bEnableCCD;
 };
 
 /**
diff --git a/Engine/Plugins/Experimental/PhysicsControl/Source/PhysicsControl/Public/PhysicsControlObjectVersion.h b/Engine/Plugins/Experimental/PhysicsControl/Source/PhysicsControl/Public/PhysicsControlObjectVersion.h
--- a/Engine/Plugins/Experimental/PhysicsControl/Source/PhysicsControl/Public/PhysicsControlObjectVersion.h
+++ b/Engine/Plugins/Experimental/PhysicsControl/Source/PhysicsControl/Public/PhysicsControlObjectVersion.h
@@ -40,6 +40,8 @@ struct FPhysicsControlObjectVersion
 		RemoveSkeletalAnimationVelocityMultiplier,
 		RemoveUseSkeletalAnimation,
 		ControlRigShapeData,
+		EnableCCD,
+		ControlRigAllowCCD,
 
 		// -----<new versions can be added above this line>-----
 		VersionPlusOne,
diff --git a/Engine/Source/Runtime/Engine/Private/PhysicsEngine/ImmediatePhysics/ImmediatePhysicsChaos/ImmediatePhysicsActorHandle_Chaos.cpp b/Engine/Source/Runtime/Engine/Private/PhysicsEngine/ImmediatePhysics/ImmediatePhysicsChaos/ImmediatePhysicsActorHandle_Chaos.cpp
--- a/Engine/Source/Runtime/Engine/Private/PhysicsEngine/ImmediatePhysics/ImmediatePhysicsChaos/ImmediatePhysicsActorHandle_Chaos.cpp
+++ b/Engine/Source/Runtime/Engine/Private/PhysicsEngine/ImmediatePhysics/ImmediatePhysicsChaos/ImmediatePhysicsActorHandle_Chaos.cpp
@@ -745,6 +745,25 @@ namespace ImmediatePhysics_Chaos
 		return Handle()->CenterOfMass();
 	}
 
+	void FActorHandle::SetCCDEnabled(bool bEnable)
+	{
+		using namespace Chaos;
+		if (FPBDRigidParticleHandle* Dynamic = ParticleHandle->CastToRigidParticle())
+		{
+			Dynamic->SetCCDEnabled(bEnable);
+		}
+	}
+
+	bool FActorHandle::GetCCDEnabled() const
+	{
+		using namespace Chaos;
+		if (FPBDRigidParticleHandle* Dynamic = ParticleHandle->CastToRigidParticle())
+		{
+			return Dynamic->CCDEnabled();
+		}
+		return false;
+	}
+
 	int32 FActorHandle::GetLevel() const
 	{
 		return Level;
diff --git a/Engine/Source/Runtime/Engine/Public/Physics/ImmediatePhysics/ImmediatePhysicsChaos/ImmediatePhysicsActorHandle_Chaos.h b/Engine/Source/Runtime/Engine/Public/Physics/ImmediatePhysics/ImmediatePhysicsChaos/ImmediatePhysicsActorHandle_Chaos.h
--- a/Engine/Source/Runtime/Engine/Public/Physics/ImmediatePhysics/ImmediatePhysicsChaos/ImmediatePhysicsActorHandle_Chaos.h
+++ b/Engine/Source/Runtime/Engine/Public/Physics/ImmediatePhysics/ImmediatePhysicsChaos/ImmediatePhysicsActorHandle_Chaos.h
@@ -208,6 +208,12 @@ namespace ImmediatePhysics_Chaos
 		/** Get the actor-space centre of mass offset (location only) */
 		ENGINE_API FVector GetLocalCoMLocation() const;
 
+		/** Enable/disable CCD */
+		ENGINE_API void SetCCDEnabled(bool bEnable);
+
+		/** returns whether CCD is enabled for the particle */
+		ENGINE_API bool GetCCDEnabled() const;
+
 		ENGINE_API Chaos::FGeometryParticleHandle* GetParticle();
 		ENGINE_API const Chaos::FGeometryParticleHandle* GetParticle() const;
 
